#ue4
- Removed conflicting edit/visibility flags on some gameplay classes
- Added an error to UHT when conflciting edit/visibility or BP exposure specifiers are used together

diff --git a/Engine/Source/Programs/UnrealHeaderTool/Private/HeaderParser.cpp b/Engine/Source/Programs/UnrealHeaderTool/Private/HeaderParser.cpp
--- a/Engine/Source/Programs/UnrealHeaderTool/Private/HeaderParser.cpp
+++ b/Engine/Source/Programs/UnrealHeaderTool/Private/HeaderParser.cpp
@@ -2537,6 +2537,8 @@ bool FHeaderParser::GetVarType
 	}
 
 	// Process the list of specifiers
+	bool bSeenEditSpecifier = false;
+	bool bSeenBlueprintEditSpecifier = false;
 	for (TArray<FPropertySpecifier>::TIterator SpecifierIt(SpecifiersFound); SpecifierIt; ++SpecifierIt)
 	{
 		const FString& Specifier = SpecifierIt->Key;
@@ -2545,47 +2547,91 @@ bool FHeaderParser::GetVarType
 		{
 			if (Specifier == TEXT("EditAnywhere"))
 			{
+				if (bSeenEditSpecifier)
+				{
+					FError::Throwf(TEXT("Found more than one edit/visibility specifier (%s), only one is allowed"), *Specifier);
+				}
 				Flags |= CPF_Edit;
+				bSeenEditSpecifier = true;
 			}
 			else if (Specifier == TEXT("EditInstanceOnly"))
 			{
-				Flags |= CPF_Edit|CPF_DisableEditOnTemplate;
+				if (bSeenEditSpecifier)
+				{
+					FError::Throwf(TEXT("Found more than one edit/visibility specifier (%s), only one is allowed"), *Specifier);
+				}
+				Flags |= CPF_Edit | CPF_DisableEditOnTemplate;
+				bSeenEditSpecifier = true;
 			}
 			else if (Specifier == TEXT("EditDefaultsOnly")) 
 			{
-				Flags |= CPF_Edit|CPF_DisableEditOnInstance;
+				if (bSeenEditSpecifier)
+				{
+					FError::Throwf(TEXT("Found more than one edit/visibility specifier (%s), only one is allowed"), *Specifier);
+				}
+				Flags |= CPF_Edit | CPF_DisableEditOnInstance;
+				bSeenEditSpecifier = true;
 			}
 			else if (Specifier == TEXT("VisibleAnywhere")) 
 			{
-				Flags |= CPF_Edit|CPF_EditConst;
+				if (bSeenEditSpecifier)
+				{
+					FError::Throwf(TEXT("Found more than one edit/visibility specifier (%s), only one is allowed"), *Specifier);
+				}
+				Flags |= CPF_Edit | CPF_EditConst;
+				bSeenEditSpecifier = true;
 			}
 			else if (Specifier == TEXT("VisibleInstanceOnly"))
 			{
-				Flags |= CPF_Edit|CPF_EditConst|CPF_DisableEditOnTemplate;
+				if (bSeenEditSpecifier)
+				{
+					FError::Throwf(TEXT("Found more than one edit/visibility specifier (%s), only one is allowed"), *Specifier);
+				}
+				Flags |= CPF_Edit | CPF_EditConst | CPF_DisableEditOnTemplate;
+				bSeenEditSpecifier = true;
 			}
 			else if (Specifier == TEXT("VisibleDefaultsOnly"))
 			{
-				Flags |= CPF_Edit|CPF_EditConst|CPF_DisableEditOnInstance;
+				if (bSeenEditSpecifier)
+				{
+					FError::Throwf(TEXT("Found more than one edit/visibility specifier (%s), only one is allowed"), *Specifier);
+				}
+				Flags |= CPF_Edit | CPF_EditConst | CPF_DisableEditOnInstance;
+				bSeenEditSpecifier = true;
 			}
 			else if (Specifier == TEXT("BlueprintReadWrite"))
 			{
+				if (bSeenBlueprintEditSpecifier)
+				{
+					FError::Throwf(TEXT("Found more than one Blueprint read/write specifier (%s), only one is allowed"), *Specifier);
+				}
+
 				const FString* PrivateAccessMD = MetaDataFromNewStyle.Find(TEXT("AllowPrivateAccess"));  // FBlueprintMetadata::MD_AllowPrivateAccess
 				const bool bAllowPrivateAccess = PrivateAccessMD ? (*PrivateAccessMD == TEXT("true")) : false;
 				if ((CurrentAccessSpecifier == ACCESS_Private) && !bAllowPrivateAccess)
 				{
 					FError::Throwf(TEXT("BlueprintReadWrite should not be used on private members"));
 				}
+
 				Flags |= CPF_BlueprintVisible;
+				bSeenBlueprintEditSpecifier = true;
 			}
 			else if (Specifier == TEXT("BlueprintReadOnly"))
 			{
+				if (bSeenBlueprintEditSpecifier)
+				{
+					FError::Throwf(TEXT("Found more than one Blueprint read/write specifier (%s), only one is allowed"), *Specifier);
+				}
+
 				const FString* PrivateAccessMD = MetaDataFromNewStyle.Find(TEXT("AllowPrivateAccess"));  // FBlueprintMetadata::MD_AllowPrivateAccess
 				const bool bAllowPrivateAccess = PrivateAccessMD ? (*PrivateAccessMD == TEXT("true")) : false;
 				if ((CurrentAccessSpecifier == ACCESS_Private) && !bAllowPrivateAccess)
 				{
 					FError::Throwf(TEXT("BlueprintReadOnly should not be used on private members"));
 				}
+
 				Flags |= CPF_BlueprintVisible|CPF_BlueprintReadOnly;
+				bSeenBlueprintEditSpecifier = true;
 			}
 			else if (Specifier == TEXT("Config"))
 			{
