#ttp 331356 - 'Recompile' editor button should tell user to quit and compile if hot reload is not possible.
#proj UE4
#branch UE4
#summary Changed UBT and UHT exit codes to inform about compilation result. The editor is reading exit code of UBT process and changes output message informing if this is a failure due to header changes.

diff --git a/Engine/Source/Editor/MainFrame/Private/MainFrameModule.cpp b/Engine/Source/Editor/MainFrame/Private/MainFrameModule.cpp
--- a/Engine/Source/Editor/MainFrame/Private/MainFrameModule.cpp
+++ b/Engine/Source/Editor/MainFrame/Private/MainFrameModule.cpp
@@ -659,7 +659,7 @@ void FMainFrameModule::OnCancelCodeCompilationClicked()
 	FModuleManager::Get().RequestStopCompilation();
 }
 
-void FMainFrameModule::HandleLevelEditorModuleCompileFinished( const FString& LogDump, bool bCompileSucceeded, bool bShowLog )
+void FMainFrameModule::HandleLevelEditorModuleCompileFinished(const FString& LogDump, ECompilationResult::Type CompilationResult, bool bShowLog)
 {
 	// Track stats
 	{
@@ -671,7 +671,7 @@ void FMainFrameModule::HandleLevelEditorModuleCompileFinished( const FString& Lo
 			FEngineAnalytics::GetProvider().RecordEvent( 
 				TEXT("Editor.Modules.Recompile"),
 				TEXT("Duration"), FString::Printf(TEXT("%.3f"), ModuleCompileDuration),
-				TEXT("Result"), bCompileSucceeded ? TEXT("Succeeded") : TEXT("Failed"));
+				TEXT("Result"), CompilationResult == ECompilationResult::Succeeded ? TEXT("Succeeded") : TEXT("Failed"));
 		}
 	}
 
@@ -679,7 +679,7 @@ void FMainFrameModule::HandleLevelEditorModuleCompileFinished( const FString& Lo
 
 	if (NotificationItem.IsValid())
 	{
-		if (bCompileSucceeded)
+		if (CompilationResult == ECompilationResult::Succeeded)
 		{
 			GEditor->PlayPreviewSound(CompileSuccessSound);
 			NotificationItem->SetText(NSLOCTEXT("MainFrame", "RecompileComplete", "Compile Complete!"));
@@ -697,7 +697,15 @@ void FMainFrameModule::HandleLevelEditorModuleCompileFinished( const FString& Lo
 			};
 
 			GEditor->PlayPreviewSound(CompileFailSound);
-			NotificationItem->SetText(NSLOCTEXT("MainFrame", "RecompileFailed", "Compile Failed!"));
+			if (CompilationResult == ECompilationResult::FailedDueToHeaderChange)
+			{
+				NotificationItem->SetText(NSLOCTEXT("MainFrame", "RecompileFailedDueToHeaderChange", "Compile failed due to the header changes. Close the editor and recompile project in IDE to apply changes."));
+			}
+			else
+			{
+				NotificationItem->SetText(NSLOCTEXT("MainFrame", "RecompileFailed", "Compile Failed!"));
+			}
+			
 			NotificationItem->SetCompletionState(SNotificationItem::CS_Fail);
 			NotificationItem->SetHyperlink(FSimpleDelegate::CreateStatic(&Local::ShowCompileLog));
 		}
diff --git a/Engine/Source/Editor/MainFrame/Private/MainFrameModule.h b/Engine/Source/Editor/MainFrame/Private/MainFrameModule.h
--- a/Engine/Source/Editor/MainFrame/Private/MainFrameModule.h
+++ b/Engine/Source/Editor/MainFrame/Private/MainFrameModule.h
@@ -145,7 +145,7 @@ private:
 	void OnCancelCodeCompilationClicked();
 
 	// Handles the level editor module finishing to recompile.
-	void HandleLevelEditorModuleCompileFinished( const FString& LogDump, bool bCompileSucceeded, bool bShowLog );
+	void HandleLevelEditorModuleCompileFinished( const FString& LogDump, ECompilationResult::Type CompilationResult, bool bShowLog );
 
 	// Handles the VSHandler having finished launching Visual Studio.
 	void HandleVSAccessorLaunched( const bool WasSuccessful );
diff --git a/Engine/Source/Editor/UnrealEd/Private/Kismet2/CompilerResultsLog.cpp b/Engine/Source/Editor/UnrealEd/Private/Kismet2/CompilerResultsLog.cpp
--- a/Engine/Source/Editor/UnrealEd/Private/Kismet2/CompilerResultsLog.cpp
+++ b/Engine/Source/Editor/UnrealEd/Private/Kismet2/CompilerResultsLog.cpp
@@ -309,7 +309,7 @@ void FCompilerResultsLog::OnGotoError(const TSharedRef<IMessageToken>& Token)
 	}
 }
 
-void FCompilerResultsLog::GetGlobalModuleCompilerDump(const FString& LogDump, bool bCompileSucceeded, bool bShowLog)
+void FCompilerResultsLog::GetGlobalModuleCompilerDump(const FString& LogDump, ECompilationResult::Type CompilationResult, bool bShowLog)
 {
 	FMessageLog MessageLog(Name);
 
diff --git a/Engine/Source/Editor/UnrealEd/Public/Kismet2/CompilerResultsLog.h b/Engine/Source/Editor/UnrealEd/Public/Kismet2/CompilerResultsLog.h
--- a/Engine/Source/Editor/UnrealEd/Public/Kismet2/CompilerResultsLog.h
+++ b/Engine/Source/Editor/UnrealEd/Public/Kismet2/CompilerResultsLog.h
@@ -96,7 +96,7 @@ private:
 	static void OnGotoError(const class TSharedRef<IMessageToken>& Token);
 
 	/** Callback function for binding the global compiler dump to open the static compiler log */
-	static void GetGlobalModuleCompilerDump(const FString& LogDump, bool bCompileSucceeded, bool bShowLog);
+	static void GetGlobalModuleCompilerDump(const FString& LogDump, ECompilationResult::Type CompilationResult, bool bShowLog);
 
 	/** The log's name, for easy re-use */
 	static const FName Name;
diff --git a/Engine/Source/Programs/UnrealBuildTool/Configuration/UEBuildTarget.cs b/Engine/Source/Programs/UnrealBuildTool/Configuration/UEBuildTarget.cs
--- a/Engine/Source/Programs/UnrealBuildTool/Configuration/UEBuildTarget.cs
+++ b/Engine/Source/Programs/UnrealBuildTool/Configuration/UEBuildTarget.cs
@@ -1094,8 +1094,8 @@ namespace UnrealBuildTool
 			}
 		}
 
-		/** Builds the target, returning a list of output files. */
-		public List<FileItem> Build()
+		/** Builds the target, appending list of output files and returns building result. */
+		public ECompilationResult Build(List<FileItem> OutputItems)
 		{
 			// Set up the global compile and link environment in GlobalCompileEnvironment and GlobalLinkEnvironment.
 			SetupGlobalEnvironment();
@@ -1352,7 +1352,7 @@ namespace UnrealBuildTool
                 GenerateManifest(AppBinaries, GlobalLinkEnvironment.Config.TargetPlatform, SpecialRocketLibFilesThatAreBuildProducts);
                 if (!BuildConfiguration.bXGEExport)
                 {
-                    return new List<FileItem>();
+                    return ECompilationResult.Succeeded;
                 }
 			}
 
@@ -1364,8 +1364,6 @@ namespace UnrealBuildTool
 				GlobalCompileEnvironment.SharedPCHHeaderFiles = SharedPCHHeaderFiles;
 			}
 
-			// Build the target's binaries.
-			var OutputItems = new List<FileItem>();
 			foreach (var Binary in AppBinaries)
 			{
 				OutputItems.AddRange(Binary.Build(GlobalCompileEnvironment, GlobalLinkEnvironment));
@@ -1418,9 +1416,12 @@ namespace UnrealBuildTool
 					// Execute the header tool
 					var TargetName = !String.IsNullOrEmpty( GameName ) ? GameName : AppName;
 					string ModuleInfoFileName = Path.GetFullPath( Path.Combine( ProjectIntermediateDirectory, "UnrealHeaderTool.manifest" ) );
-					if (!ExternalExecution.ExecuteHeaderToolIfNecessary(this, UObjectModules, ModuleInfoFileName))
+
+					ECompilationResult UHTResult = ECompilationResult.OtherCompilationError;
+					if (!ExternalExecution.ExecuteHeaderToolIfNecessary(this, UObjectModules, ModuleInfoFileName, ref UHTResult))
 					{
-						throw new BuildException( "UnrealHeaderTool failed for target '" + TargetName + "' (platform: " + Platform.ToString() + ", module info: " + ModuleInfoFileName + ")." );
+						Log.TraceInformation("UnrealHeaderTool failed for target '" + TargetName + "' (platform: " + Platform.ToString() + ", module info: " + ModuleInfoFileName + ").");
+						return UHTResult;
 					}
 				}
 
@@ -1443,10 +1444,11 @@ namespace UnrealBuildTool
 					}
 					Writer.WriteEndElement();
 				}
-				return new List<FileItem>();
+
+				OutputItems.Clear();
 			}
 
-			return OutputItems;
+			return ECompilationResult.Succeeded;
 		}
 
 		/// <summary>
diff --git a/Engine/Source/Programs/UnrealBuildTool/System/ExternalExecution.cs b/Engine/Source/Programs/UnrealBuildTool/System/ExternalExecution.cs
--- a/Engine/Source/Programs/UnrealBuildTool/System/ExternalExecution.cs
+++ b/Engine/Source/Programs/UnrealBuildTool/System/ExternalExecution.cs
@@ -11,6 +11,17 @@ using System.Reflection;
 
 namespace UnrealBuildTool
 {
+	// This enum has to be compatible with the one defined in the
+	// UE4\Engine\Source\Runtime\Core\Public\Modules\ModuleManager.h
+	// to keep communication between UHT, UBT and Editor compiling
+	// processes valid.
+	public enum ECompilationResult
+	{
+		Succeeded = 0,
+		FailedDueToHeaderChange = 1,
+		OtherCompilationError = 2
+	}
+
 	/** Information about a module that needs to be passed to UnrealHeaderTool for code generation */
 	public struct UHTModuleInfo
 	{
@@ -352,7 +363,7 @@ namespace UnrealBuildTool
 		 * Builds and runs the header tool and touches the header directories.
 		 * Performs any early outs if headers need no changes, given the UObject modules, tool path, game name, and configuration
 		 */
-		public static bool ExecuteHeaderToolIfNecessary( UEBuildTarget Target, List<UHTModuleInfo> UObjectModules, string ModuleInfoFileName )
+		public static bool ExecuteHeaderToolIfNecessary( UEBuildTarget Target, List<UHTModuleInfo> UObjectModules, string ModuleInfoFileName, ref ECompilationResult UHTResult )
 		{
 			// We never want to try to execute the header tool when we're already trying to build it!
 			var bIsBuildingUHT = Target.GetTargetName().Equals( "UnrealHeaderTool", StringComparison.InvariantCultureIgnoreCase );
@@ -421,12 +432,12 @@ namespace UnrealBuildTool
 
 				Stopwatch s = new Stopwatch();
 				s.Start();
-				int ErrorCode = RunExternalExecutable(ExternalExecution.GetHeaderToolPath(), CmdLine);
+				UHTResult = (ECompilationResult) RunExternalExecutable(ExternalExecution.GetHeaderToolPath(), CmdLine);
 				s.Stop();
-					
-				if (ErrorCode != 0)
+
+				if (UHTResult != ECompilationResult.Succeeded)
 				{
-					Log.TraceInformation( "Error: Failed to generate code for {0} - error code: {1}", ActualTargetName, ErrorCode );
+					Log.TraceInformation("Error: Failed to generate code for {0} - error code: {1}", ActualTargetName, (int) UHTResult);
 					return false;
 				}
 
diff --git a/Engine/Source/Programs/UnrealBuildTool/System/ProjectFileGenerator.cs b/Engine/Source/Programs/UnrealBuildTool/System/ProjectFileGenerator.cs
--- a/Engine/Source/Programs/UnrealBuildTool/System/ProjectFileGenerator.cs
+++ b/Engine/Source/Programs/UnrealBuildTool/System/ProjectFileGenerator.cs
@@ -1209,7 +1209,7 @@ namespace UnrealBuildTool
 							ArgumentsCopy[ CurEntryIndex + 1 ] = Arguments[ CurEntryIndex ];
 						}
 
-						bSuccess = UnrealBuildTool.RunUBT( ArgumentsCopy );
+						bSuccess = UnrealBuildTool.RunUBT( ArgumentsCopy ) == ECompilationResult.Succeeded;
 						if( !bSuccess )
 						{
 							break;
diff --git a/Engine/Source/Programs/UnrealBuildTool/System/UnrealBuildTool.cs b/Engine/Source/Programs/UnrealBuildTool/System/UnrealBuildTool.cs
--- a/Engine/Source/Programs/UnrealBuildTool/System/UnrealBuildTool.cs
+++ b/Engine/Source/Programs/UnrealBuildTool/System/UnrealBuildTool.cs
@@ -573,7 +573,7 @@ namespace UnrealBuildTool
 				UProjectInfo.FillProjectInfo();
 			}
 
-			bool bSuccess = true;
+			ECompilationResult Result = ECompilationResult.Succeeded;
 
 			// Reset early so we can access BuildConfiguration even before RunUBT() is called
 			BuildConfiguration.Reset();
@@ -848,19 +848,24 @@ namespace UnrealBuildTool
 
                     if (bGenerateVCProjectFiles || bGenerateXcodeProjectFiles || bGenerateMakefiles)
 					{
-						bSuccess = true;
+						bool bGenerationSuccess = true;
 						if (bGenerateVCProjectFiles)
 						{
-							bSuccess &= GenerateProjectFiles(new VCProjectFileGenerator(), Arguments);
+							bGenerationSuccess &= GenerateProjectFiles(new VCProjectFileGenerator(), Arguments);
 						}
 						if (bGenerateXcodeProjectFiles)
 						{
-							bSuccess &= GenerateProjectFiles(new XcodeProjectFileGenerator(), Arguments);
+							bGenerationSuccess &= GenerateProjectFiles(new XcodeProjectFileGenerator(), Arguments);
 						}
                         if (bGenerateMakefiles)
  						{
- 							bSuccess &= GenerateProjectFiles(new MakefileGenerator(), Arguments);
+							bGenerationSuccess &= GenerateProjectFiles(new MakefileGenerator(), Arguments);
  						}
+
+						if(!bGenerationSuccess)
+						{
+							Result = ECompilationResult.OtherCompilationError;
+						}
 					}
 					else if (bRunCopyrightVerification)
 					{
@@ -876,12 +881,12 @@ namespace UnrealBuildTool
 						}
 
 						// Build our project
-						if (bSuccess)
+						if (Result == ECompilationResult.Succeeded)
 						{
 							if (UEBuildConfiguration.bPrepForDeployment == false)
 							{
 								// If we are only prepping for deployment, assume the build already occurred.
-								bSuccess = RunUBT(Arguments);
+								Result = RunUBT(Arguments);
 							}
 							else
 							{
@@ -896,7 +901,7 @@ namespace UnrealBuildTool
 
 							// If we build w/ bXGEExport true, we didn't REALLY build at this point, 
 							// so don't bother with doing the PrepTargetForDeployment call. 
-							if ((bSuccess == true) && (BuildConfiguration.bDeployAfterCompile == true) && (BuildConfiguration.bXGEExport == false) &&
+							if ((Result == ECompilationResult.Succeeded) && (BuildConfiguration.bDeployAfterCompile == true) && (BuildConfiguration.bXGEExport == false) &&
 								(UEBuildConfiguration.bGenerateManifest == false) && (UEBuildConfiguration.bCleanProject == false))
 							{
 								UEBuildDeploy DeployHandler = UEBuildDeploy.GetBuildDeploy(CheckPlatform);
@@ -947,7 +952,7 @@ namespace UnrealBuildTool
 				catch (Exception Exception)
 				{
 					Log.TraceError("UnrealBuildTool Exception: " + Exception);
-					bSuccess = false;
+					Result = ECompilationResult.OtherCompilationError;
 				}
 
 				if (bUseMutex)
@@ -964,7 +969,7 @@ namespace UnrealBuildTool
 			// Print some performance info
 			Log.TraceVerbose("Execution time: {0}", (DateTime.UtcNow - StartTime - MutexWaitTime).TotalSeconds);
 
-			return bSuccess ? 0 : 1;
+			return (int) Result;
 		}
 
 
@@ -988,7 +993,7 @@ namespace UnrealBuildTool
 
 
 		// @todo: Ideally get rid of RunUBT() and all of the Clear/Reset stuff!
-		public static bool RunUBT(string[] Arguments)
+		public static ECompilationResult RunUBT(string[] Arguments)
 		{
 			bool bSuccess = true;
 
@@ -1021,6 +1026,7 @@ namespace UnrealBuildTool
 
 			var Targets = new List<UEBuildTarget>();
 			string ExecutorName = "Unknown";
+			ECompilationResult BuildResult = ECompilationResult.Succeeded;
 
 			try
 			{
@@ -1053,8 +1059,15 @@ namespace UnrealBuildTool
 							DependencyCache.GetDependencyCachePathForTarget(Target)
 							);
 					}
-				
-					var TargetOutputItems = Target.Build();
+
+					var TargetOutputItems = new List<FileItem>();
+					BuildResult = Target.Build(TargetOutputItems);
+
+					if(BuildResult != ECompilationResult.Succeeded)
+					{
+						break;
+					}
+
 					OutputItemsForAllTargets.AddRange( TargetOutputItems );
 
 					if ( (BuildConfiguration.bXGEExport && UEBuildConfiguration.bGenerateManifest) || (!ProjectFileGenerator.bGenerateProjectFiles && !UEBuildConfiguration.bGenerateManifest && !UEBuildConfiguration.bCleanProject))
@@ -1108,7 +1121,11 @@ namespace UnrealBuildTool
 						);
 				}
 
-				if ((BuildConfiguration.bXGEExport && UEBuildConfiguration.bGenerateManifest) || (!GeneratingActionGraph && !ProjectFileGenerator.bGenerateProjectFiles && !UEBuildConfiguration.bGenerateManifest && !UEBuildConfiguration.bCleanProject))
+				if (BuildResult == ECompilationResult.Succeeded &&
+					(
+						(BuildConfiguration.bXGEExport && UEBuildConfiguration.bGenerateManifest) ||
+						(!GeneratingActionGraph && !ProjectFileGenerator.bGenerateProjectFiles && !UEBuildConfiguration.bGenerateManifest && !UEBuildConfiguration.bCleanProject)
+					))
 				{
 					// Plan the actions to execute for the build.
 					List<Action> ActionsToExecute = GetActionsToExecute(OutputItemsForAllTargets, Targets);
@@ -1140,12 +1157,12 @@ namespace UnrealBuildTool
 			{
 				// Output the message only, without the call stack
 				Log.TraceInformation(Exception.Message);
-				bSuccess = false;
+				BuildResult = ECompilationResult.OtherCompilationError;
 			}
 			catch (Exception Exception)
 			{
 				Log.TraceInformation("ERROR: {0}", Exception);
-				bSuccess = false;
+				BuildResult = ECompilationResult.OtherCompilationError;
 			}
 
 			// Figure out how long we took to execute.
@@ -1194,7 +1211,7 @@ namespace UnrealBuildTool
 					);
 			}
 
-			return bSuccess;
+			return BuildResult;
 		}
 
 		private static void ParseBuildConfigurationFlags(string[] Arguments)
diff --git a/Engine/Source/Programs/UnrealHeaderTool/Private/CodeGenerator.cpp b/Engine/Source/Programs/UnrealHeaderTool/Private/CodeGenerator.cpp
--- a/Engine/Source/Programs/UnrealHeaderTool/Private/CodeGenerator.cpp
+++ b/Engine/Source/Programs/UnrealHeaderTool/Private/CodeGenerator.cpp
@@ -4240,6 +4240,12 @@ void FNativeClassHeaderGenerator::DeleteUnusedGeneratedHeaders()
 	}
 }
 
+/**
+ * Dirty hack global variable to allow different result codes passed through
+ * exceptions. Needs to be fixed in future versions of UHT.
+ */
+ECompilationResult::Type GCompilationResult = ECompilationResult::OtherCompilationError;
+
 bool FNativeClassHeaderGenerator::SaveHeaderIfChanged(const TCHAR* HeaderPath, const TCHAR* InNewHeaderContents)
 {
 	if ( !bAllowSaveExportedHeaders )
@@ -4334,6 +4340,8 @@ bool FNativeClassHeaderGenerator::SaveHeaderIfChanged(const TCHAR* HeaderPath, c
 		{
 			FString ConflictPath = FString(HeaderPath) + TEXT(".conflict");
 			FFileHelper::SaveStringToFile(NewHeaderContents, *ConflictPath);
+
+			GCompilationResult = ECompilationResult::FailedDueToHeaderChange;
 			FError::Throwf(TEXT("ERROR: '%s': Changes to generated code are not allowed - conflicts written to '%s'"), HeaderPath, *ConflictPath);
 		}
 
@@ -4588,9 +4596,10 @@ void FNativeClassHeaderGenerator::ExportGeneratedCPP()
 }
 
 
-int32 UnrealHeaderTool_Main(const FString& ModuleInfoFilename)
+ECompilationResult::Type UnrealHeaderTool_Main(const FString& ModuleInfoFilename)
 {
 	check(GIsUCCMakeStandaloneHeaderGenerator);
+	ECompilationResult::Type Result = ECompilationResult::Succeeded;
 	
 	if ( !FParse::Param( FCommandLine::Get(), TEXT("IgnoreWarnings")) )
 	{
@@ -4607,7 +4616,7 @@ int32 UnrealHeaderTool_Main(const FString& ModuleInfoFilename)
 	catch (const TCHAR* Ex)
 	{
 		UE_LOG(LogCompile, Error, TEXT("Failed to load manifest file '%s': %s"), *ModuleInfoFilename, Ex);
-		return 3;
+		return GCompilationResult;
 	}
 
 	NameLookupCPP = new FNameLookupCPP();
@@ -4616,7 +4625,6 @@ int32 UnrealHeaderTool_Main(const FString& ModuleInfoFilename)
 	GScriptHelper = ScriptHelper;
 
 	// Load classes for editing.
-	bool Success = true;
 	int32 NumFailures = 0;
 
 	// Three passes.  1) Public 'Classes' headers (legacy)  2) Public headers   3) Private headers
@@ -4648,8 +4656,10 @@ int32 UnrealHeaderTool_Main(const FString& ModuleInfoFilename)
 
 		for (const auto& Module : GManifest.Modules)
 		{
-			if (!Success)
+			if (Result != ECompilationResult::Succeeded)
+			{
 				break;
+			}
 
 			// We'll make an ordered list of all UObject headers we care about.
 			TArray<FString> UObjectHeaders;
@@ -4744,6 +4754,7 @@ int32 UnrealHeaderTool_Main(const FString& ModuleInfoFilename)
 						}
 
 						FString FormattedErrorMessage = FString::Printf(TEXT("%sError: %s\r\n"), *Prefix, ErrorMsg);
+						Result = GCompilationResult;
 
 						UE_LOG(LogCompile, Log, TEXT("%s"), *FormattedErrorMessage);
 						GWarn->Log(ELogVerbosity::Error, FormattedErrorMessage);
@@ -4753,14 +4764,17 @@ int32 UnrealHeaderTool_Main(const FString& ModuleInfoFilename)
 				#endif
 				}
 
-				Success = NumFailures == 0;
+				if(Result == ECompilationResult::Succeeded && NumFailures != 0)
+				{
+					Result = ECompilationResult::OtherCompilationError;
+				}
 			}
 		}
 	}
 
 
 	// Save generated headers
-	if ( Success )
+	if ( Result == ECompilationResult::Succeeded )
 	{
 		// Verify that all script declared superclasses exist.
 		for( TObjectIterator<UClass> ItC; ItC; ++ItC )
@@ -4771,19 +4785,19 @@ int32 UnrealHeaderTool_Main(const FString& ModuleInfoFilename)
 			if (ScriptSuperClass && !ScriptSuperClass->HasAnyClassFlags(CLASS_Intrinsic) && GClassStrippedHeaderTextMap.Contains(ScriptClass) && !GClassStrippedHeaderTextMap.Contains(ScriptSuperClass))
 			{
 				UE_LOG(LogCompile, Error, TEXT("Superclass %s of class %s not found"), *ScriptSuperClass->GetName(), *ScriptClass->GetName());
-				Success = false;
+				Result = ECompilationResult::OtherCompilationError;
 				++NumFailures;
 			}
 		}
 
-		if (Success)
+		if (Result == ECompilationResult::Succeeded)
 		{
 			for (const auto& Module : GManifest.Modules)
 			{
 				if (UPackage* Package = Cast<UPackage>( StaticFindObjectFast( UPackage::StaticClass(), NULL, FName(*Module.LongPackageName), false, false ) ))
 				{
-					Success = FHeaderParser::ParseAllHeadersInside(GWarn, Package, Module.SaveExportedHeaders, GManifest.UseRelativePaths);
-					if (!Success)
+					Result = FHeaderParser::ParseAllHeadersInside(GWarn, Package, Module.SaveExportedHeaders, GManifest.UseRelativePaths);
+					if (Result != ECompilationResult::Succeeded)
 					{
 						++NumFailures;
 						break;
@@ -4825,8 +4839,13 @@ int32 UnrealHeaderTool_Main(const FString& ModuleInfoFilename)
 	ScriptHelper = NULL;
 
 	GIsRequestingExit = true;
+
+	if(Result == ECompilationResult::Succeeded && NumFailures > 0)
+	{
+		return ECompilationResult::OtherCompilationError;
+	}
 	
-	return NumFailures > 0 ? 4 : 0;
+	return Result;
 }
 
 UClass* GenerateCodeForHeader
diff --git a/Engine/Source/Programs/UnrealHeaderTool/Private/HeaderParser.cpp b/Engine/Source/Programs/UnrealHeaderTool/Private/HeaderParser.cpp
--- a/Engine/Source/Programs/UnrealHeaderTool/Private/HeaderParser.cpp
+++ b/Engine/Source/Programs/UnrealHeaderTool/Private/HeaderParser.cpp
@@ -22,6 +22,12 @@ enum {MAX_ARRAY_SIZE=2048};
 
 static const FName NAME_ToolTip(TEXT("ToolTip"));
 
+/**
+ * Dirty hack global variable to allow different result codes passed through
+ * exceptions. Needs to be fixed in future versions of UHT.
+ */
+extern ECompilationResult::Type GCompilationResult;
+
 /*-----------------------------------------------------------------------------
 	Utility functions.
 -----------------------------------------------------------------------------*/
@@ -2612,7 +2618,6 @@ bool FHeaderParser::GetVarType
 				{
 					FError::Throwf(TEXT("BlueprintReadWrite should not be used on private members"));
 				}
-
 				Flags |= CPF_BlueprintVisible;
 				bSeenBlueprintEditSpecifier = true;
 			}
@@ -2629,7 +2634,6 @@ bool FHeaderParser::GetVarType
 				{
 					FError::Throwf(TEXT("BlueprintReadOnly should not be used on private members"));
 				}
-
 				Flags |= CPF_BlueprintVisible|CPF_BlueprintReadOnly;
 				bSeenBlueprintEditSpecifier = true;
 			}
@@ -6177,14 +6181,14 @@ void FHeaderParser::FinalizeScriptExposedFunctions(UClass* Class)
 
 //
 // Parses the header associated with the specified class.
-// Returns 1 if parsing was a success, 0 if any errors occurred.
+// Returns result enumeration.
 //
-bool FHeaderParser::ParseHeaderForOneClass(FClassTree& AllClasses, UClass* InClass)
+ECompilationResult::Type FHeaderParser::ParseHeaderForOneClass(FClassTree& AllClasses, UClass* InClass)
 {
 	// Early-out if this class has previously failed some aspect of parsing
 	if (FailedClassesAnnotation.Get(InClass))
 	{
-		return false;
+		return ECompilationResult::OtherCompilationError;
 	}
 
 	// Reset the parser to begin a new class
@@ -6197,7 +6201,7 @@ bool FHeaderParser::ParseHeaderForOneClass(FClassTree& AllClasses, UClass* InCla
 	bHaveSeenUClass                             = false;
 	bClassHasGeneratedBody                      = false;
 
-	bool bSuccess = false;
+	ECompilationResult::Type Result = ECompilationResult::OtherCompilationError;
 
 	ClassData = GScriptHelper->AddClassData(Class);
 
@@ -6370,7 +6374,7 @@ bool FHeaderParser::ParseHeaderForOneClass(FClassTree& AllClasses, UClass* InCla
 		}
 
 		// First-pass success.
-		bSuccess = true;
+		Result = ECompilationResult::Succeeded;
 	}
 #if !PLATFORM_EXCEPTIONS_DISABLED
 	catch( TCHAR* ErrorMsg )
@@ -6386,7 +6390,7 @@ bool FHeaderParser::ParseHeaderForOneClass(FClassTree& AllClasses, UClass* InCla
 		}
 
 		FailedClassesAnnotation.Set(Class);
-		bSuccess = false;
+		Result = GCompilationResult;
 	}
 #endif
 
@@ -6394,7 +6398,7 @@ bool FHeaderParser::ParseHeaderForOneClass(FClassTree& AllClasses, UClass* InCla
 	Class->Bind();
 
 	// Finalize functions
-	if (bSuccess)
+	if (Result == ECompilationResult::Succeeded)
 	{
 		FinalizeScriptExposedFunctions( Class );
 	}
@@ -6405,7 +6409,7 @@ bool FHeaderParser::ParseHeaderForOneClass(FClassTree& AllClasses, UClass* InCla
 		FError::Throwf(TEXT("Expected an include at the top of the header: '#include \"%s\"'"), *ExpectedHeaderName);
 	}
 
-	return bSuccess; //@TODO: UCREMOVAL: This function is always returning true even on a compiler error; should this continue?
+	return Result; //@TODO: UCREMOVAL: This function is always returning succeeded even on a compiler error; should this continue?
 }
 
 /*-----------------------------------------------------------------------------
@@ -6413,16 +6417,16 @@ bool FHeaderParser::ParseHeaderForOneClass(FClassTree& AllClasses, UClass* InCla
 -----------------------------------------------------------------------------*/
 
 // Parse Class's annotated headers and optionally its child classes.  Marks the class as CLASS_Parsed.
-bool FHeaderParser::ParseHeaders(FClassTree& AllClasses, FHeaderParser& HeaderParser, UClass* Class, bool bParseSubclasses)
+ECompilationResult::Type FHeaderParser::ParseHeaders(FClassTree& AllClasses, FHeaderParser& HeaderParser, UClass* Class, bool bParseSubclasses)
 {
 	check(Class != NULL);
 
 	if (!GClassStrippedHeaderTextMap.Contains(Class))
 	{
-		return true;
+		return ECompilationResult::Succeeded;
 	}
 
-	bool bResult = true;
+	ECompilationResult::Type Result = ECompilationResult::Succeeded;
 
 	// Handle all dependencies of the class.
 	auto& DependentOn = *GClassDependentOnMap.FindOrAdd(Class);
@@ -6487,14 +6491,18 @@ bool FHeaderParser::ParseHeaders(FClassTree& AllClasses, FHeaderParser& HeaderPa
 		{
 			UClass* NextClass = ClassesToParse.Pop();
 
-			if (ParseHeaders(AllClasses, HeaderParser, NextClass, true))
+			ECompilationResult::Type ParseResult = ParseHeaders(AllClasses, HeaderParser, NextClass, true);
+
+			if (ParseResult == ECompilationResult::Succeeded)
 			{
 				break;
 			}
 
-			if (!ParseHeaders(AllClasses, HeaderParser, NextClass, false))
+			ParseResult = ParseHeaders(AllClasses, HeaderParser, NextClass, false);
+
+			if (ParseResult != ECompilationResult::Succeeded)
 			{
-				bResult = false;
+				Result = ParseResult;
 				break;
 			}
 		}
@@ -6504,10 +6512,11 @@ bool FHeaderParser::ParseHeaders(FClassTree& AllClasses, FHeaderParser& HeaderPa
 	if (!(Class->ClassFlags & CLASS_Parsed))
 	{
 		UClass* CurrentSuperClass = Class->GetSuperClass();
-		if (!HeaderParser.ParseHeaderForOneClass(AllClasses, Class))
+		ECompilationResult::Type OneClassResult = HeaderParser.ParseHeaderForOneClass(AllClasses, Class);
+		if (OneClassResult != ECompilationResult::Succeeded)
 		{
 			// if we couldn't parse this class, we won't be able to parse its children
-			return false;
+			return OneClassResult;
 		}
 
 		if (CurrentSuperClass != Class->GetSuperClass())
@@ -6532,15 +6541,22 @@ bool FHeaderParser::ParseHeaders(FClassTree& AllClasses, FHeaderParser& HeaderPa
 
 			//note: you must always pass in the root tree node here, since we may add new classes to the tree
 			// if an manual dependency (through dependson()) is encountered
-			if (!ParseHeaders(AllClasses, HeaderParser, SubClass, bParseSubclasses))
+			ECompilationResult::Type ParseResult = ParseHeaders(AllClasses, HeaderParser, SubClass, bParseSubclasses);
+
+			if (ParseResult == ECompilationResult::FailedDueToHeaderChange)
+			{
+				return ParseResult;
+			}
+
+			if (ParseResult == ECompilationResult::OtherCompilationError)
 			{
-				bResult = false;
+				Result = ECompilationResult::OtherCompilationError;
 			}
 		}
 	}
 
 	// Success.
-	return bResult;
+	return Result;
 }
 
 bool FHeaderParser::DependentClassNameFromHeader(const TCHAR* HeaderFilename, FString& OutClassName)
@@ -6693,7 +6709,7 @@ FString FHeaderParser::RequireExactlyOneSpecifierValue(const FPropertySpecifier&
 }
 
 // Parse all headers for classes that are inside LimitOuter.
-bool FHeaderParser::ParseAllHeadersInside(FFeedbackContext* Warn, UPackage* LimitOuter, bool bAllowSaveExportedHeaders, bool bUseRelativePaths)
+ECompilationResult::Type FHeaderParser::ParseAllHeadersInside(FFeedbackContext* Warn, UPackage* LimitOuter, bool bAllowSaveExportedHeaders, bool bUseRelativePaths)
 {
 	// Disable loading of objects outside of this package (or more exactly, objects which aren't UFields, CDO, or templates)
 	TGuardValue<bool> AutoRestoreVerifyObjectRefsFlag(GVerifyObjectReferencesOnly, true);
@@ -6724,7 +6740,7 @@ bool FHeaderParser::ParseAllHeadersInside(FFeedbackContext* Warn, UPackage* Limi
 	HeaderParser.Filename = IFileManager::Get().ConvertToAbsolutePathForExternalAppForRead(*GClassSourceFileMap[BaseClass]);
 
 	// Hierarchically parse all classes.
-	bool bSuccess = true;
+	ECompilationResult::Type Result = ECompilationResult::Succeeded;
 #if !PLATFORM_EXCEPTIONS_DISABLED
 	try
 #endif
@@ -6732,10 +6748,10 @@ bool FHeaderParser::ParseAllHeadersInside(FFeedbackContext* Warn, UPackage* Limi
 		// Parse the headers
 		const bool bParseSubclasses = true;
 
-		bSuccess = FHeaderParser::ParseHeaders(AllClasses, HeaderParser, BaseClass, bParseSubclasses);
+		Result = FHeaderParser::ParseHeaders(AllClasses, HeaderParser, BaseClass, bParseSubclasses);
 
 		// Export the autogenerated code wrappers
-		if (bSuccess)
+		if (Result == ECompilationResult::Succeeded)
 		{
 			// At this point all headers have been parsed and the header parser will
 			// no longer have up to date info about what's being done so unregister it 
@@ -6759,14 +6775,13 @@ bool FHeaderParser::ParseAllHeadersInside(FFeedbackContext* Warn, UPackage* Limi
 	catch( TCHAR* ErrorMsg )
 	{
 		Warn->Log(ELogVerbosity::Error, ErrorMsg);
-
-		bSuccess = false;
+		Result = GCompilationResult;
 	}
 #endif
 	// Unregister the header parser from the feedback context
 	Warn->SetContext(NULL);
 
-	return bSuccess;
+	return Result;
 }
 
 /** 
diff --git a/Engine/Source/Programs/UnrealHeaderTool/Private/HeaderParser.h b/Engine/Source/Programs/UnrealHeaderTool/Private/HeaderParser.h
--- a/Engine/Source/Programs/UnrealHeaderTool/Private/HeaderParser.h
+++ b/Engine/Source/Programs/UnrealHeaderTool/Private/HeaderParser.h
@@ -77,7 +77,7 @@ class FHeaderParser : public FBaseParser, public FContextSupplier
 {
 public:
 	// Parse all headers for classes that are inside LimitOuter.
-	static bool ParseAllHeadersInside(FFeedbackContext* Warn, UPackage* LimitOuter, bool bAllowSaveExportedHeaders, bool bUseRelativePaths);
+	static ECompilationResult::Type ParseAllHeadersInside(FFeedbackContext* Warn, UPackage* LimitOuter, bool bAllowSaveExportedHeaders, bool bUseRelativePaths);
 
 	// Performs a preliminary parse of the text in the specified buffer, pulling out:
 	//   Class name and parent class name
@@ -313,9 +313,9 @@ protected:
 	 * @param	Class				the class to parse
 	 * @param	bParseSubclasses	true if we should parse all child classes of Class
 	 *
-	 * @return	true if the class was successfully compiled, false otherwise
+	 * @return	Result enumeration.
 	 */
-	static bool ParseHeaders(FClassTree& AllClasses, FHeaderParser& HeaderParser, UClass* Class, bool bParseSubclasses);
+	static ECompilationResult::Type ParseHeaders(FClassTree& AllClasses, FHeaderParser& HeaderParser, UClass* Class, bool bParseSubclasses);
 
 	//@TODO: Remove this method
 	static void ParseClassName(const TCHAR* Temp, FString& ClassName);
@@ -339,7 +339,7 @@ protected:
 	// End of FContextSupplier interface.
 
 	// High-level compiling functions.
-	bool ParseHeaderForOneClass(FClassTree& AllClasses, UClass* InClass);
+	ECompilationResult::Type ParseHeaderForOneClass(FClassTree& AllClasses, UClass* InClass);
 	void CompileDirective(UClass* Class);
 	void FinalizeScriptExposedFunctions(UClass* Class);
 	UEnum* CompileEnum(UClass* Owner);
diff --git a/Engine/Source/Programs/UnrealHeaderTool/Private/UnrealHeaderToolMain.cpp b/Engine/Source/Programs/UnrealHeaderTool/Private/UnrealHeaderToolMain.cpp
--- a/Engine/Source/Programs/UnrealHeaderTool/Private/UnrealHeaderToolMain.cpp
+++ b/Engine/Source/Programs/UnrealHeaderTool/Private/UnrealHeaderToolMain.cpp
@@ -73,17 +73,12 @@ int32 main( int32 ArgC, char* ArgV[] )
 		ModuleInfoFilename = FPaths::ConvertRelativePathToFull(FPlatformProcess::BaseDir(), TEXT("../../Source/Programs/UnrealHeaderTool/Resources/UHTDebugging.manifest"));
 	}
 
-	extern int32 UnrealHeaderTool_Main(const FString& ModuleInfoFilename);
-	int32 ErrorLevel = UnrealHeaderTool_Main(ModuleInfoFilename);
-
-	if (GWarn->Errors.Num() > 0 || GWarn->Warnings.Num() > 0)
-	{
-		ErrorLevel = 2;
-	}
+	extern ECompilationResult::Type UnrealHeaderTool_Main(const FString& ModuleInfoFilename);
+	ECompilationResult::Type Result = UnrealHeaderTool_Main(ModuleInfoFilename);
 
 	FEngineLoop::AppPreExit();
 	FEngineLoop::AppExit();
 
-	return ErrorLevel;
+	return Result;
 }
 
diff --git a/Engine/Source/Runtime/Core/Private/Modules/ModuleManager.cpp b/Engine/Source/Runtime/Core/Private/Modules/ModuleManager.cpp
--- a/Engine/Source/Runtime/Core/Private/Modules/ModuleManager.cpp
+++ b/Engine/Source/Runtime/Core/Private/Modules/ModuleManager.cpp
@@ -1563,7 +1563,7 @@ bool FModuleManager::StartCompilingModuleDLLs(const FString& GameName, const TAr
 		// No longer compiling modules
 		ModulesBeingCompiled.Empty();
 
-		ModuleCompilerFinishedEvent.Broadcast(FString(), false, false);
+		ModuleCompilerFinishedEvent.Broadcast(FString(), ECompilationResult::OtherCompilationError, false);
 
 		// Fire task completion delegate 
 		
@@ -1784,7 +1784,7 @@ void FModuleManager::CheckForFinishedModuleDLLCompile( const bool bWaitForComple
 {
 #if PLATFORM_DESKTOP && !IS_MONOLITHIC
 	bCompileStillInProgress = false;
-	bCompileSucceeded = false;
+	ECompilationResult::Type CompilationResult = ECompilationResult::OtherCompilationError;
 
 	// Is there a compilation in progress?
 	if( IsCurrentlyCompiling() )
@@ -1820,7 +1820,7 @@ void FModuleManager::CheckForFinishedModuleDLLCompile( const bool bWaitForComple
 		}
 
 		// Check to see if the compile has finished yet
-		int32 ReturnCode = 1;
+		int32 ReturnCode = -1;
 		while (bCompileStillInProgress)
 		{
 			if( FPlatformProcess::GetProcReturnCode( ModuleCompileProcessHandle, &ReturnCode ) )
@@ -1860,13 +1860,13 @@ void FModuleManager::CheckForFinishedModuleDLLCompile( const bool bWaitForComple
 			// Compilation finished, now we need to grab all of the text from the output pipe
 			ModuleCompileReadPipeText += FPlatformProcess::ReadPipe(ModuleCompileReadPipe);
 
-			// Was the compile successful?
-			bCompileSucceeded = ( ReturnCode == 0 );
+			// The ReturnCode is -1 only if compilation was cancelled.
+			CompilationResult = ReturnCode != -1 ? (ECompilationResult::Type)ReturnCode : ECompilationResult::OtherCompilationError;
 
 			// If compilation succeeded for all modules, go back to the modules and update their module file names
 			// in case we recompiled the modules to a new unique file name.  This is needed so that when the module
 			// is reloaded after the recompile, we load the new DLL file name, not the old one
-			if( bCompileSucceeded )
+			if(CompilationResult == ECompilationResult::Succeeded)
 			{
 				for( int32 CurModuleIndex = 0; CurModuleIndex < ModulesThatWereBeingRecompiled.Num(); ++CurModuleIndex )
 				{
@@ -1910,10 +1910,12 @@ void FModuleManager::CheckForFinishedModuleDLLCompile( const bool bWaitForComple
 			// No longer compiling modules
 			ModulesBeingCompiled.Empty();
 
+			bCompileSucceeded = CompilationResult == ECompilationResult::Succeeded;
+
 			if ( bFireEvents )
 			{
 				const bool bShowLogOnSuccess = false;
-				ModuleCompilerFinishedEvent.Broadcast(FinalOutput, bCompileSucceeded, !bCompileSucceeded || bShowLogOnSuccess);
+				ModuleCompilerFinishedEvent.Broadcast(FinalOutput, CompilationResult, !bCompileSucceeded || bShowLogOnSuccess);
 
 				// Fire task completion delegate 
 				RecompileModulesCallback.ExecuteIfBound( true, bCompileSucceeded );
diff --git a/Engine/Source/Runtime/Core/Public/Modules/ModuleManager.h b/Engine/Source/Runtime/Core/Public/Modules/ModuleManager.h
--- a/Engine/Source/Runtime/Core/Public/Modules/ModuleManager.h
+++ b/Engine/Source/Runtime/Core/Public/Modules/ModuleManager.h
@@ -24,6 +24,20 @@ namespace ELoadModuleFailureReason
 	};
 }
 
+// This enum has to be compatible with the one defined in the
+// UE4\Engine\Source\Programs\UnrealBuildTool\System\ExternalExecution.cs
+// to keep communication between UHT, UBT and Editor compiling
+// processes valid.
+namespace ECompilationResult
+{
+	enum Type
+	{
+		Succeeded = 0,
+		FailedDueToHeaderChange = 1,
+		OtherCompilationError = 2
+	};
+}
+
 namespace EModuleCompileMethod
 {
 	enum Type
@@ -96,9 +110,9 @@ public:
 	FModuleCompilerStartedEvent& OnModuleCompilerStarted() { return ModuleCompilerStartedEvent; }
 
 	/** Delegate for binding functions to be called when the module compiler finishes,
-		passing in the compiler output and a boolean which indicates success or failure
+		passing in the compiler output and a enumeration with compilation result
 		as well as a second boolean which determines if the log should be shown */
-	DECLARE_EVENT_ThreeParams( FModuleManager, FModuleCompilerFinishedEvent, const FString&, bool, bool );
+	DECLARE_EVENT_ThreeParams(FModuleManager, FModuleCompilerFinishedEvent, const FString&, ECompilationResult::Type, bool);
 	FModuleCompilerFinishedEvent& OnModuleCompilerFinished() { return ModuleCompilerFinishedEvent; }
 
 	/** Called after a module recompile finishes.  First argument specifies whether the compilation has finished, 
