#ttp 331467: Make UHT not depend on CoreUObject

#add FClass and FClasses added to encapsulate (and ultimately replace) UClass and FClassTree.
#add StringUtils contains a bunch of string-related functions which were in FHeaderParser for no real reason.
#change Replaced FClassTree with UHT and fixed up all the code which was affected.

diff --git a/Engine/Source/Programs/UnrealHeaderTool/Private/Class.cpp b/Engine/Source/Programs/UnrealHeaderTool/Private/Class.cpp
new file mode 100644
--- /dev/null
+++ b/Engine/Source/Programs/UnrealHeaderTool/Private/Class.cpp
@@ -0,0 +1,79 @@
+// Copyright 1998-2014 Epic Games, Inc. All Rights Reserved.
+
+#include "UnrealHeaderTool.h"
+#include "Class.h"
+#include "ClassMaps.h"
+
+bool FClass::Inherits(const FClass* SuspectBase) const
+{
+	return IsChildOf(SuspectBase);
+}
+
+FString FClass::GetNameWithPrefix(EEnforceInterfacePrefix::Type EnforceInterfacePrefix) const
+{
+	const TCHAR* Prefix = 0;
+
+	if (HasAnyClassFlags(CLASS_Interface))
+	{
+		// Grab the expected prefix for interfaces (U on the first one, I on the second one)
+		switch (EnforceInterfacePrefix)
+		{
+			case EEnforceInterfacePrefix::None:
+				// For old-style files: "I" for interfaces, unless it's the actual "Interface" class, which gets "U"
+				if (GetFName() == NAME_Interface)
+				{
+					Prefix = TEXT("U");
+				}
+				else
+				{
+					Prefix = TEXT("I");
+				}
+				break;
+
+			case EEnforceInterfacePrefix::I:
+				Prefix = TEXT("I");
+				break;
+
+			case EEnforceInterfacePrefix::U:
+				Prefix = TEXT("U");
+				break;
+
+			default:
+				check(false);
+		}
+	}
+	else
+	{
+		// Get the expected class name with prefix
+		Prefix = GetPrefixCPP();
+	}
+
+	return FString::Printf(TEXT("%s%s"), Prefix, *GetName());
+}
+
+FClass* FClass::GetSuperClass() const
+{
+	return static_cast<FClass*>(static_cast<const UClass*>(this)->GetSuperClass());
+}
+
+TArray<FName> FClass::GetDependentNames() const
+{
+	TArray<FName> Result;
+
+	for (auto Name : *GClassDependentOnMap.FindOrAdd(const_cast<FClass*>(this)))
+	{
+		Result.Add(Name);
+	}
+
+	return Result;
+}
+
+TArray<FClass*> FClass::GetInterfaceTypes() const
+{
+	TArray<FClass*> Result;
+	for (auto& i : Interfaces)
+	{
+		Result.Add((FClass*)i.Class);
+	}
+	return Result;
+}
diff --git a/Engine/Source/Programs/UnrealHeaderTool/Private/Class.h b/Engine/Source/Programs/UnrealHeaderTool/Private/Class.h
new file mode 100644
--- /dev/null
+++ b/Engine/Source/Programs/UnrealHeaderTool/Private/Class.h
@@ -0,0 +1,49 @@
+// Copyright 1998-2014 Epic Games, Inc. All Rights Reserved.
+
+/*=============================================================================
+	Class.h: Represents a parsed class.
+=============================================================================*/
+
+#pragma once
+
+struct EEnforceInterfacePrefix
+{
+	enum Type
+	{
+		None,
+		I,
+		U
+	};
+};
+
+class FClass : public UClass
+{
+public:
+	FClass();
+
+	/**
+	 * Tests if this class inherits another.
+	 *
+	 * @param SuspectBase The class to test if it's a base of this.
+	 * @return true if the SuspectBase is a base of this, false otherwise.
+	 */
+	bool Inherits(const FClass* SuspectBase) const;
+
+	/** 
+	 * Returns the name of the given class with a valid prefix.
+	 *
+	 * @param InClass Class used to create a valid class name with prefix
+	 */
+	FString GetNameWithPrefix(EEnforceInterfacePrefix::Type EnforceInterfacePrefix = EEnforceInterfacePrefix::None) const;
+
+	/**
+	 * Returns the super class of this class, or NULL if there is no superclass.
+	 *
+	 * @return The super class of this class.
+	 */
+	FClass* GetSuperClass() const;
+
+	TArray<FName> GetDependentNames() const;
+
+	TArray<FClass*> GetInterfaceTypes() const;
+};
diff --git a/Engine/Source/Programs/UnrealHeaderTool/Private/Classes.cpp b/Engine/Source/Programs/UnrealHeaderTool/Private/Classes.cpp
new file mode 100644
--- /dev/null
+++ b/Engine/Source/Programs/UnrealHeaderTool/Private/Classes.cpp
@@ -0,0 +1,205 @@
+// Copyright 1998-2014 Epic Games, Inc. All Rights Reserved.
+
+#include "UnrealHeaderTool.h"
+#include "Classes.h"
+#include "Class.h"
+#include "StringUtils.h"
+
+namespace
+{
+	/** 
+	 * Returns True if the given class name includes a valid Unreal prefix and matches based on the given class.
+	 *
+	 * @param InNameToCheck - Name w/ potential prefix to check
+	 * @param OriginalClass - Class to check against
+	 */
+	bool ClassNameHasValidPrefix(const FString InNameToCheck, const FClass* OriginalClass)
+	{
+		bool bIsLabledDeprecated;
+		GetClassPrefix( InNameToCheck, bIsLabledDeprecated );
+
+		// If the class is labeled deprecated, don't try to resolve it during header generation, valid results can't be guaranteed.
+		if (bIsLabledDeprecated)
+		{
+			return true;
+		}
+
+		const FString OriginalClassName = OriginalClass->GetNameWithPrefix();
+
+		bool bNamesMatch = (InNameToCheck == OriginalClassName);
+
+		if (!bNamesMatch)
+		{
+			//@TODO: UCREMOVAL: I/U interface hack - Ignoring prefixing for this call
+			if (OriginalClass->HasAnyClassFlags(CLASS_Interface))
+			{
+				bNamesMatch = InNameToCheck.Mid(1) == OriginalClassName.Mid(1);
+			}
+		}
+
+		return bNamesMatch;
+	}
+}
+
+FClasses::FClasses(UPackage* InPackage)
+	: UObjectClass((FClass*)UObject::StaticClass())
+	, ClassTree   (UObjectClass)
+{
+	for (UClass* Class : TObjectRange<UClass>())
+	{
+		if (Class->IsIn(InPackage))
+		{
+			ClassTree.AddClass(Class);
+		}
+	}
+}
+
+FClass* FClasses::GetRootClass() const
+{
+	return UObjectClass;
+}
+
+bool FClasses::IsDependentOn(const FClass* Suspect, const FClass* Source) const
+{
+	// A class is not dependent on itself.
+	if (Suspect == Source)
+		return false;
+
+	// Children are all implicitly dependent on their parent, that is, children require their parent
+	// to be compiled first therefore if the source is a parent of the suspect, the suspect is
+	// dependent on the source.
+	if (Suspect->Inherits(Source))
+		return true;
+
+	// Now consider all dependents of the suspect. If any of them are dependent on the source, the
+	// suspect is too.
+	for (auto It : Suspect->GetDependentNames())
+	{
+		const FClass* Dependency = FindClass(*It.ToString());
+
+		// Error.
+		if (!Dependency)
+			continue;
+
+		// the parser disallows declaring the parent class as a dependency, so the only way this could occur is
+		// if the parent for a native class has been changed (which causes the new parent to be inserted as a dependency),
+		// if this is the case, skip it or we'll go into a loop
+		if (Suspect->GetSuperClass() == Dependency)
+			continue;
+
+		if (Dependency == Source || IsDependentOn(Dependency, Source))
+			return true;
+	}
+
+	// Consider all children of the suspect if any of them are dependent on the source, the suspect is
+	// too because it itself is dependent on its children.
+
+	if(!Suspect->HasAnyClassFlags(CLASS_Interface))
+	{
+		if (!ClassTree.FindNode(const_cast<FClass*>(Suspect)))
+			FError::Throwf(TEXT("Unparsed class '%s' found while validating DependsOn entries for '%s'"), *Suspect->GetName(), *Source->GetName());
+	}
+
+	return false;
+}
+
+FClass* FClasses::FindClass(const TCHAR* ClassName) const
+{
+	check(ClassName);
+
+	UObject* ClassPackage = ANY_PACKAGE;
+
+	if (UClass* Result = FindObject<UClass>(ClassPackage, ClassName))
+		return (FClass*)Result;
+
+	if (UObjectRedirector* RenamedClassRedirector = FindObject<UObjectRedirector>(ClassPackage, ClassName))
+		return (FClass*)CastChecked<UClass>(RenamedClassRedirector->DestinationObject);
+
+	return NULL;
+}
+
+TArray<FClass*> FClasses::GetDerivedClasses(FClass* Parent) const
+{
+	const FClassTree* ClassLeaf = ClassTree.FindNode(Parent);
+	TArray<const FClassTree*> ChildLeaves;
+	ClassLeaf->GetChildClasses(ChildLeaves);
+
+	TArray<FClass*> Result;
+	for (auto Node : ChildLeaves)
+	{
+		Result.Add((FClass*)Node->GetClass());
+	}
+
+	return Result;
+}
+
+FClass* FClasses::FindScriptClass(const FString& InClassName) const
+{
+	FString ErrorMsg;
+	return FindScriptClass(InClassName, ErrorMsg);
+}
+
+FClass* FClasses::FindScriptClassOrThrow(const FString& InClassName) const
+{
+	FString ErrorMsg;
+	if (FClass* Result = FindScriptClass(InClassName, ErrorMsg))
+		return Result;
+
+	FError::Throwf(*ErrorMsg);
+
+	// Unreachable, but compiler will warn otherwise because FError::Throwf isn't declared noreturn
+	return 0;
+}
+
+FClass* FClasses::FindScriptClass(const FString& InClassName, FString& OutErrorMsg) const
+{
+	// Strip the class name of its prefix and then do a search for the class
+	FString ClassNameStripped = GetClassNameWithPrefixRemoved(InClassName);
+	if (FClass* FoundClass = FindClass(*ClassNameStripped))
+	{
+		// If the class was found with the stripped class name, verify that the correct prefix was used and throw an error otherwise
+		if( !ClassNameHasValidPrefix( InClassName, FoundClass ) )
+		{
+			OutErrorMsg = FString::Printf(TEXT("Class '%s' has an incorrect prefix, expecting '%s'"), *InClassName, *FoundClass->GetNameWithPrefix());
+			return NULL;
+		}
+
+		return (FClass*)FoundClass;
+	}
+
+	// Couldn't find the class with a class name stripped of prefix (or a prefix was not found)
+	// See if the prefix was forgotten by trying to find the class with the given identifier
+	if (FClass* FoundClass = FindClass(*InClassName))
+	{
+		// If the class was found with the given identifier, the user forgot to use the correct Unreal prefix	
+		OutErrorMsg = FString::Printf(TEXT("Class '%s' is missing a prefix, expecting '%s'"), *InClassName, *FoundClass->GetNameWithPrefix());
+	}
+	else
+	{
+		// If the class was still not found, it wasn't a valid identifier
+		OutErrorMsg = FString::Printf(TEXT("Class '%s' not found."), *InClassName );
+	}
+
+	return NULL;
+}
+
+TArray<FClass*> FClasses::GetClassesInPackage(UPackage* InPackage) const
+{
+	TArray<FClass*> Result;
+	Result.Add(UObjectClass);
+
+	// This cast is evil, but it'll work until we get TArray covariance. ;-)
+	ClassTree.GetChildClasses((TArray<UClass*>&)Result, [=](const UClass* Class) { return InPackage == ANY_PACKAGE || Class->GetOuter() == InPackage; }, true);
+
+	return Result;
+}
+
+void FClasses::ChangeParentClass(FClass* Class)
+{
+	ClassTree.ChangeParentClass(Class);
+}
+
+void FClasses::Validate()
+{
+	ClassTree.Validate();
+}
diff --git a/Engine/Source/Programs/UnrealHeaderTool/Private/Classes.h b/Engine/Source/Programs/UnrealHeaderTool/Private/Classes.h
new file mode 100644
--- /dev/null
+++ b/Engine/Source/Programs/UnrealHeaderTool/Private/Classes.h
@@ -0,0 +1,104 @@
+// Copyright 1998-2014 Epic Games, Inc. All Rights Reserved.
+
+/*=============================================================================
+	Classes.h: Container structure for all the Classes.
+=============================================================================*/
+
+#pragma once
+
+#include "ClassTree.h"
+
+#define WIP_UHT_REFACTOR 1
+
+class FClass;
+
+class FClasses
+{
+public:
+	explicit FClasses(UPackage* InPackage);
+
+	/**
+	 * Returns the root class (i.e. UObject)
+	 *
+	 * @return The root class.
+	 */
+	FClass* GetRootClass() const;
+
+	/**
+	 * Determines whether the class hierarchy rooted at Suspect is
+	 * dependent on the hierarchy rooted at Source.
+	 *
+	 * @param   Suspect  Root of hierarchy for suspect class
+	 * @param   Source   Root of hierarchy for source class
+	 * @return           true if the hierarchy rooted at Suspect is dependent on the one rooted at Source, false otherwise
+	 */
+	bool IsDependentOn(const FClass* Suspect, const FClass* Source) const;
+
+	FClass* FindClass(const TCHAR* ClassName) const;
+
+	TArray<FClass*> GetDerivedClasses(FClass* Parent) const;
+
+	/** 
+	 * Attempts to find a script class based on the given name. Will attempt to strip
+	 * the prefix of the given name while searching. Throws script errors when appropriate.
+	 *
+	 * @param   InClassName  Name w/ Unreal prefix to use when searching for a class
+	 * @return               The found class, or NULL if the class was not found.
+	 */
+	FClass* FindScriptClass(const FString& InClassName) const;
+
+	/** 
+	 * Attempts to find a script class based on the given name. Will attempt to strip
+	 * the prefix of the given name while searching. Throws an exception with the script error
+	 * if the class could not be found.
+	 *
+	 * @param   InClassName  Name w/ Unreal prefix to use when searching for a class
+	 * @return               The found class.
+	 */
+	FClass* FindScriptClassOrThrow(const FString& InClassName) const;
+
+	/** 
+	 * Attempts to find a script class based on the given name. Will attempt to strip
+	 * the prefix of the given name while searching. Throws script errors when appropriate.
+	 *
+	 * @param   InClassName  Name w/ Unreal prefix to use when searching for a class
+	 * @param   OutErrorMsg  Error message (if any) giving the caller flexibility in how they present an error
+	 * @return               The found class, or NULL if the class was not found.
+	 */
+	FClass* FindScriptClass(const FString& InClassName, FString& OutErrorMsg) const;
+
+	/**
+	 * Returns an array of classes for the given package.
+	 *
+	 * @param   InPackage  The package to return the classes from.
+	 * @return             The classes in the specified package.
+	 */
+	TArray<FClass*> GetClassesInPackage(UPackage* InPackage = ANY_PACKAGE) const;
+
+// Anything in here should eventually be removed when this class encapsulates its own data structure, rather than being 'poked' by the outside
+#if WIP_UHT_REFACTOR
+
+	/**
+	 * Move a class node in the hierarchy tree after a class has changed its SuperClass
+	 * 
+	 * @param	SearchClass			the class that has changed parents
+	 * @param	InNewParentClass	if non-null force reparenting to this instead of the SuperClass
+	 *
+	 * @return	true if SearchClass was successfully moved to the new location
+	 */
+	void ChangeParentClass(FClass* Class);
+
+	/**
+	 * Validates the state of the tree (shouldn't be needed once this class has well-defined invariants).
+	 */
+	void Validate();
+
+#endif
+
+private:
+	FClass*     UObjectClass;
+	FClassTree  ClassTree;
+
+	friend auto begin(const FClasses& Classes) -> decltype(begin(TObjectRange<FClass>())) { return begin(TObjectRange<FClass>()); }
+	friend auto end  (const FClasses& Classes) -> decltype(end  (TObjectRange<FClass>())) { return end  (TObjectRange<FClass>()); }
+};
diff --git a/Engine/Source/Programs/UnrealHeaderTool/Private/CodeGenerator.cpp b/Engine/Source/Programs/UnrealHeaderTool/Private/CodeGenerator.cpp
--- a/Engine/Source/Programs/UnrealHeaderTool/Private/CodeGenerator.cpp
+++ b/Engine/Source/Programs/UnrealHeaderTool/Private/CodeGenerator.cpp
@@ -26,28 +26,6 @@ static UClass* GenerateCodeForHeader(UObject* InParent, const TCHAR* Name, EObje
 
 FCompilerMetadataManager* GScriptHelper = NULL;
 
-/////////////////////////////////////////////////////
-// FClassPackageComparator
-
-class FClassPackageComparator
-{
-	UPackage* Package;
-
-public:
-
-	FClassPackageComparator( UPackage* InPackage )
-	: Package(InPackage)
-	{
-	}
-
-	bool IsValidClass( const class FClassTree* Node ) const
-	{
-		UClass* Cls = Node->GetClass();
-		return Cls == UObject::StaticClass() || Cls->GetOuter() == Package;
-	}
-};
-
-
 /** C++ name lookup helper */
 FNameLookupCPP* NameLookupCPP;
 
@@ -2124,7 +2102,7 @@ void FNativeClassHeaderGenerator::ExportClassHeaderRecursive( UClass* Class, TAr
 
 		// Handle #include directives as 'dependson' (allowing them to fail).
 		const bool bClassNameFromHeader = FHeaderParser::DependentClassNameFromHeader(*DependencyClassNameStripped, DependencyClassNameStripped);
-		DependencyClassNameStripped = FHeaderParser::GetClassNameWithPrefixRemoved( DependencyClassNameStripped );
+		DependencyClassNameStripped = GetClassNameWithPrefixRemoved( DependencyClassNameStripped );
 
 		int32 ClassIndex = FindByIndirectName( Classes, *DependencyClassNameStripped );
 		UClass* FoundClass = FindObject<UClass>(ANY_PACKAGE, *DependencyClassNameStripped);
@@ -2133,7 +2111,7 @@ void FNativeClassHeaderGenerator::ExportClassHeaderRecursive( UClass* Class, TAr
 		if (FoundClass == NULL)
 		{
 			// Try again with temporary class name. This may be struct only header.
-			DependencyClassNameStripped = FHeaderParser::GetClassNameWithoutPrefix(DependencyClassName.ToString());
+			DependencyClassNameStripped = GetClassNameWithoutPrefix(DependencyClassName.ToString());
 			DependencyClassNameStripped = FHeaderParser::GenerateTemporaryClassName(*DependencyClassNameStripped);
 			ClassIndex = FindByIndirectName( Classes, *DependencyClassNameStripped );
 			FoundClass = FindObject<UClass>(ANY_PACKAGE, *DependencyClassNameStripped);
@@ -3816,7 +3794,7 @@ bool FNativeClassHeaderGenerator::FindInterDependencyRecursive( TMap<const FStri
  * @param	DependencyChain		[out] Receives dependency chain, if found.
  * @return	true if a dependency chain was found and filled in.
  */
-bool FNativeClassHeaderGenerator::FindDependencyChain( const UClass* Class, const FString& Header1, const FString& Header2, TArray<const UClass*>& DependencyChain )
+bool FNativeClassHeaderGenerator::FindDependencyChain( FClass* Class, const FString& Header1, const FString& Header2, TArray<FClass*>& DependencyChain )
 {
 	DependencyChain.Empty();
 	return FindDependencyChainRecursive( Class, Header1, Header2, false, DependencyChain );
@@ -3832,7 +3810,7 @@ bool FNativeClassHeaderGenerator::FindDependencyChain( const UClass* Class, cons
  * @param	DependencyChain		[out] Receives dependency chain, if found. Must be empty before the call.
  * @return	true if a dependency chain was found and filled in.
  */
-bool FNativeClassHeaderGenerator::FindDependencyChainRecursive( const UClass* Class, const FString& Header1, const FString& Header2, bool bChainStarted, TArray<const UClass*>& DependencyChain )
+bool FNativeClassHeaderGenerator::FindDependencyChainRecursive( FClass* Class, const FString& Header1, const FString& Header2, bool bChainStarted, TArray<FClass*>& DependencyChain )
 {
 	bool bIsExportClass = Class->HasAnyClassFlags(CLASS_Native) && !Class->HasAnyClassFlags(CLASS_NoExport|CLASS_Intrinsic|CLASS_Temporary);
 	if ( bIsExportClass )
@@ -3851,15 +3829,14 @@ bool FNativeClassHeaderGenerator::FindDependencyChainRecursive( const UClass* Cl
 		}
 	}
 
-	UClass* SuperClass = Class->GetSuperClass();
+	FClass* SuperClass = Class->GetSuperClass();
 	if ( SuperClass && SuperClass->GetOuter() == Class->GetOuter() && FindDependencyChainRecursive( SuperClass, Header1, Header2, bChainStarted, DependencyChain ) )
 	{
 		return true;
 	}
 
-	for (TArray<FImplementedInterface>::TConstIterator It(Class->Interfaces); It; ++It)
+	for (FClass* InterfaceClass : Class->GetInterfaceTypes())
 	{
-		UClass* InterfaceClass = It->Class;
 		if ( InterfaceClass->GetOuter() == Class->GetOuter() && FindDependencyChainRecursive( InterfaceClass, Header1, Header2, bChainStarted, DependencyChain ) )
 		{
 			return true;
@@ -3874,11 +3851,10 @@ bool FNativeClassHeaderGenerator::FindDependencyChainRecursive( const UClass* Cl
 }
 
 // Constructor.
-FNativeClassHeaderGenerator::FNativeClassHeaderGenerator( UPackage* InPackage, FClassTree& inClassTree, bool InAllowSaveExportedHeaders, bool bInUseRelativePaths )
+FNativeClassHeaderGenerator::FNativeClassHeaderGenerator( UPackage* InPackage, FClasses& AllClasses, bool InAllowSaveExportedHeaders, bool bInUseRelativePaths )
 	: CurrentClass                          (NULL)
 	, API                                   (FPackageName::GetShortName(InPackage).ToUpper())
 	, Package                               (InPackage)
-	, ClassTree                             (inClassTree)
 	, bIsExportingForOffsetDeterminationOnly(false)
 	, bAllowSaveExportedHeaders             (InAllowSaveExportedHeaders)
 	, bFailIfGeneratedCodeChanges           (false)
@@ -3893,8 +3869,7 @@ FNativeClassHeaderGenerator::FNativeClassHeaderGenerator( UPackage* InPackage, F
 
 	// Tag native classes in this package for export.
 
-	Classes.Add(UObject::StaticClass());
-	ClassTree.GetChildClasses(Classes, FDefaultComparator(), true);
+	Classes = AllClasses.GetClassesInPackage();
 
 	TMap<FName,UClass*> ClassNameMap;
 	for ( int32 ClassIndex = 0; ClassIndex < Classes.Num(); ClassIndex++ )
@@ -3927,14 +3902,9 @@ FNativeClassHeaderGenerator::FNativeClassHeaderGenerator( UPackage* InPackage, F
 	bool bCircularDependencyDetected = false;
  	{
 	 	TMap<const FString*, HeaderDependents> HeaderDependencyMap;
-		TArray<UClass*> ClassesInPackage;
-
-		ClassesInPackage.Add(UObject::StaticClass());
-		ClassTree.GetChildClasses(ClassesInPackage, FClassPackageComparator(Package), true);
-		for (auto ClassIt = ClassesInPackage.CreateConstIterator(); ClassIt; ++ClassIt)
+		TArray<FClass*> ClassesInPackage = AllClasses.GetClassesInPackage(Package);
+		for (FClass* Cls : ClassesInPackage)
 		{
-			UClass* Cls = *ClassIt;
-
 			auto ClsHeaderFilename = GClassHeaderFilenameMap.FindRef(Cls);
 
 			bool bIsExportClass = Cls->HasAnyClassFlags(CLASS_Native) && !Cls->HasAnyClassFlags(CLASS_NoExport|CLASS_Intrinsic|CLASS_Temporary);
@@ -3992,21 +3962,18 @@ FNativeClassHeaderGenerator::FNativeClassHeaderGenerator( UPackage* InPackage, F
 				if (!FindInterDependency(HeaderDependencyMap, Header, ClassHeaderFilename1, ClassHeaderFilename2))
 					continue;
 
-				TArray<const UClass*> DependencyChain1, DependencyChain2;
-				for (auto ClassIt = ClassesInPackage.CreateConstIterator(); ClassIt; ++ClassIt)
+				TArray<FClass*> DependencyChain1;
+				for (FClass* Class : ClassesInPackage)
 				{
-					UClass* Class = *ClassIt;
-
 					FindDependencyChain( Class, *ClassHeaderFilename1, *ClassHeaderFilename2, DependencyChain1 );
 
 					if (DependencyChain1.Num())
 						break;
 				}
 
-				for (auto ClassIt = ClassesInPackage.CreateConstIterator(); ClassIt; ++ClassIt)
+				TArray<FClass*> DependencyChain2;
+				for (FClass* Class : ClassesInPackage)
 				{
-					UClass* Class = *ClassIt;
-
 					FindDependencyChain( Class, *ClassHeaderFilename2, *ClassHeaderFilename1, DependencyChain2 );
 
 					if (DependencyChain2.Num())
@@ -4017,10 +3984,8 @@ FNativeClassHeaderGenerator::FNativeClassHeaderGenerator( UPackage* InPackage, F
 					continue;
 
 				FString DependencyChainString1;
-				for (auto DependencyIt = DependencyChain1.CreateConstIterator(); DependencyIt; ++DependencyIt)
+				for (FClass* Dependency : DependencyChain1)
 				{
-					auto Dependency = *DependencyIt;
-
 					if (!DependencyChainString1.IsEmpty())
 					{
 						DependencyChainString1 += TEXT(" -> ");
@@ -4030,10 +3995,8 @@ FNativeClassHeaderGenerator::FNativeClassHeaderGenerator( UPackage* InPackage, F
 				}
 
 				FString DependencyChainString2;
-				for (auto DependencyIt = DependencyChain2.CreateConstIterator(); DependencyIt; ++DependencyIt)
+				for (FClass* Dependency : DependencyChain2)
 				{
-					auto Dependency = *DependencyIt;
-
 					if (!DependencyChainString1.IsEmpty())
 					{
 						DependencyChainString2 += TEXT(" -> ");
@@ -4881,13 +4844,13 @@ UClass* GenerateCodeForHeader
 		bNonClassHeader = true;
 	}
 
-	FString ClassNameStripped = FHeaderParser::GetClassNameWithPrefixRemoved( *ClassName );
+	FString ClassNameStripped = GetClassNameWithPrefixRemoved( *ClassName );
 
 	// Ensure the base class has any valid prefix and exists as a valid class. Checking for the 'correct' prefix will occur during compilation
 	FString BaseClassNameStripped;
 	if (!BaseClassName.IsEmpty())
 	{
-		BaseClassNameStripped = FHeaderParser::GetClassNameWithPrefixRemoved(BaseClassName);
+		BaseClassNameStripped = GetClassNameWithPrefixRemoved(BaseClassName);
 		if( !FHeaderParser::ClassNameHasValidPrefix(BaseClassName, BaseClassNameStripped) )
 			FError::Throwf(TEXT("No prefix or invalid identifier for base class %s.\nClass names must match Unreal prefix specifications (e.g., \"UObject\" or \"AActor\")"), *BaseClassName );
 
diff --git a/Engine/Source/Programs/UnrealHeaderTool/Private/HeaderParser.cpp b/Engine/Source/Programs/UnrealHeaderTool/Private/HeaderParser.cpp
--- a/Engine/Source/Programs/UnrealHeaderTool/Private/HeaderParser.cpp
+++ b/Engine/Source/Programs/UnrealHeaderTool/Private/HeaderParser.cpp
@@ -6,6 +6,7 @@
 #include "HeaderParser.h"
 #include "NativeClassExporter.h"
 #include "ClassMaps.h"
+#include "Classes.h"
 
 #include "UObjectAnnotation.h"
 
@@ -48,156 +49,6 @@ namespace
 		return false;
 	}
 
-	struct EEnforceInterfacePrefix
-	{
-		enum Type
-		{
-			None,
-			I,
-			U
-		};
-	};
-
-	/** 
-	 * Returns the name of the given class with a valid prefix.
-	 *
-	 * @param InClass - Class used to create a valid class name with prefix
-	 */
-	FString GetClassNameWithPrefix(const UClass* InClass, EEnforceInterfacePrefix::Type EnforceInterfacePrefix = EEnforceInterfacePrefix::None)
-	{
-		const TCHAR* Prefix = 0;
-
-		if (InClass->HasAnyClassFlags(CLASS_Interface))
-		{
-			// Grab the expected prefix for interfaces (U on the first one, I on the second one)
-			switch (EnforceInterfacePrefix)
-			{
-				case EEnforceInterfacePrefix::None:
-					// For old-style files: "I" for interfaces, unless it's the actual "Interface" class, which gets "U"
-					if (InClass->GetFName() == NAME_Interface)
-					{
-						Prefix = TEXT("U");
-					}
-					else
-					{
-						Prefix = TEXT("I");
-					}
-					break;
-
-				case EEnforceInterfacePrefix::I:
-					Prefix = TEXT("I");
-					break;
-
-				case EEnforceInterfacePrefix::U:
-					Prefix = TEXT("U");
-					break;
-
-				default:
-					check(false);
-			}
-		}
-		else
-		{
-			// Get the expected class name with prefix
-			Prefix = InClass->GetPrefixCPP();
-		}
-
-		return FString::Printf(TEXT("%s%s"), Prefix, *InClass->GetName());
-	}
-
-	// Utility function to find a named class
-	UClass* FindClass(const TCHAR* ClassName)
-	{
-		check(ClassName);
-
-		UObject* ClassPackage = ANY_PACKAGE;
-
-		if (UClass* Result = FindObject<UClass>(ClassPackage, ClassName))
-		{
-			return Result;
-		}
-
-		if (UObjectRedirector* RenamedClassRedirector = FindObject<UObjectRedirector>(ClassPackage, ClassName))
-		{
-			return CastChecked<UClass>(RenamedClassRedirector->DestinationObject);
-		}
-
-		return NULL;
-	}
-
-	/** 
-	 * Attempts to find a script class based on the given name. Will attempt to strip
-	 * the prefix of the given name while searching. Throws script errors when appropriate.
-	 *
-	 * @param InClassName - Name w/ Unreal prefix to use when searching for a class
-	 * @param OutErrorMsg - Error message (if any) giving the caller flexibility in how they present an error
-	 * @return              The found class, or NULL if the class was not found.
-	 */
-	UClass* FindScriptClass( const FString& InClassName, FString& OutErrorMsg )
-	{
-		// Strip the class name of its prefix and then do a search for the class
-		FString ClassNameStripped = FHeaderParser::GetClassNameWithPrefixRemoved(InClassName);
-		if (UClass* FoundClass = FindClass(*ClassNameStripped))
-		{
-			// If the class was found with the stripped class name, verify that the correct prefix was used and throw an error otherwise
-			if( !FHeaderParser::ClassNameHasValidPrefix( InClassName, FoundClass ) )
-			{
-				OutErrorMsg = FString::Printf(TEXT("Class '%s' has an incorrect prefix, expecting '%s'"), *InClassName, *GetClassNameWithPrefix(FoundClass));
-				return NULL;
-			}
-
-			return FoundClass;
-		}
-
-		// Couldn't find the class with a class name stripped of prefix (or a prefix was not found)
-		// See if the prefix was forgotten by trying to find the class with the given identifier
-		if (UClass* FoundClass = FindClass(*InClassName))
-		{
-			// If the class was found with the given identifier, the user forgot to use the correct Unreal prefix	
-			OutErrorMsg = FString::Printf(TEXT("Class '%s' is missing a prefix, expecting '%s'"), *InClassName, *GetClassNameWithPrefix(FoundClass));
-		}
-		else
-		{
-			// If the class was still not found, it wasn't a valid identifier
-			OutErrorMsg = FString::Printf(TEXT("Class '%s' not found."), *InClassName );
-		}
-
-		return NULL;
-	}
-
-	/** 
-	 * Attempts to find a script class based on the given name. Will attempt to strip
-	 * the prefix of the given name while searching. Throws script errors when appropriate.
-	 *
-	 * @param InClassName - Name w/ Unreal prefix to use when searching for a class
-	 * @return              The found class, or NULL if the class was not found.
-	 */
-	UClass* FindScriptClass(const FString& InClassName)
-	{
-		FString ErrorMsg;
-		return FindScriptClass(InClassName, ErrorMsg);
-	}
-
-	/** 
-	 * Attempts to find a script class based on the given name. Will attempt to strip
-	 * the prefix of the given name while searching. Throws an exception with the script error
-	 * if the class could not be found.
-	 *
-	 * @param InClassName - Name w/ Unreal prefix to use when searching for a class
-	 * @return              The found class.
-	 */
-	UClass* FindScriptClassOrThrow(const FString& InClassName)
-	{
-		FString ErrorMsg;
-		UClass* Result = FindScriptClass(InClassName, ErrorMsg);
-		if (!Result)
-		{
-			FError::Throwf(*ErrorMsg);
-		}
-
-		return Result;
-	}
-
 	void AddHeaderGroup(UClass* Class, const FString& HeaderGroup, const TCHAR* ScopeForError)
 	{
 		if (FString* HeaderFilename = GClassHeaderFilenameMap.Find(Class))
@@ -520,76 +371,6 @@ FScriptLocation::FScriptLocation()
 /////////////////////////////////////////////////////
 // FHeaderParser
 
-/**
- * Static helper function to determine whether the class hierarchy rooted at Suspect is
- * dependent on the hierarchy rooted at Source.
- *
- * @param	Suspect		Root of hierarchy for suspect class
- * @param	Source		Root of hierarchy for source class
- * @param	AllClasses	Array of parsed classes
- * @return	true if the hierarchy rooted at Suspect is dependent on the one rooted at Source, false otherwise
- */
-bool FHeaderParser::IsDependentOn( UClass* Suspect, UClass* Source, const FClassTree& AllClasses )
-{
-	check( Source );
-	check( Suspect );
-
-	// A class is not dependent on itself.
-	if( Suspect == Source )
-	{
-			return false;
-	}
-
-	// Children are all implicitly dependent on their parent, that is, children require their parent
-	// to be compiled first therefore if the source is a parent of the suspect, the suspect is
-	// dependent on the source.
-	if( Suspect->IsChildOf( Source ) )
-	{
-			return true;
-	}
-
-	// Now consider all dependents of the suspect. If any of them are dependent on the source, the
-	// suspect is too.
-	for( auto It = GClassDependentOnMap.FindOrAdd(Suspect)->CreateConstIterator(); It; ++It )
-	{
-		UClass*	Dependency = FindClass(*It->ToString());
-
-		if( !Dependency )
-		{
-			// Error.
-			continue;
-		}
-
-		// the parser disallows declaring the parent class as a dependency, so the only way this could occur is
-		// if the parent for a native class has been changed (which causes the new parent to be inserted as a dependency),
-		// if this is the case, skip it or we'll go into a loop
-		if ( Suspect->GetSuperClass() == Dependency )
-		{
-			continue;
-		}
-
-		if( Dependency == Source || IsDependentOn( Dependency, Source, AllClasses ) )
-		{
-			return true;
-		}
-	}
-
-	// Consider all children of the suspect if any of them are dependent on the source, the suspect is
-	// too because it itself is dependent on its children.
-
-	if(!Suspect->HasAnyClassFlags(CLASS_Interface))
-	{
-		const FClassTree* SuspectLeaf = AllClasses.FindNode(Suspect);
-		if ( SuspectLeaf == NULL )
-		{
-			FError::Throwf(TEXT("Unparsed class '%s' found while validating DependsOn entries for '%s'"), *Suspect->GetName(), *Source->GetName());
-			return false;
-		}
-	}
-
-	return false;
-}
-
 FString FHeaderParser::GetContext()
 {
 	// Return something useful if we didn't even get as far as class parsing
@@ -611,7 +392,7 @@ FString FHeaderParser::GetContext()
 //
 // Get a qualified class.
 //
-UClass* FHeaderParser::GetQualifiedClass( const TCHAR* Thing )
+FClass* FHeaderParser::GetQualifiedClass(const FClasses& AllClasses, const TCHAR* Thing)
 {
 	TCHAR ClassName[256]=TEXT("");
 
@@ -626,7 +407,7 @@ UClass* FHeaderParser::GetQualifiedClass( const TCHAR* Thing )
 		FError::Throwf(TEXT("%s: Missing class name"), Thing );
 	}
 
-	return FindScriptClassOrThrow(ClassName);
+	return AllClasses.FindScriptClassOrThrow(ClassName);
 }
 
 /*-----------------------------------------------------------------------------
@@ -1260,7 +1041,7 @@ EAccessSpecifier FHeaderParser::ParseAccessProtectionSpecifier(FToken& Token)
 /**
  * Compile a struct definition.
  */
-UScriptStruct* FHeaderParser::CompileStructDeclaration(UClass* Scope)
+UScriptStruct* FHeaderParser::CompileStructDeclaration(FClasses& AllClasses, FClass* Scope)
 {
 	check(Scope);
 
@@ -1404,11 +1185,11 @@ UScriptStruct* FHeaderParser::CompileStructDeclaration(UClass* Scope)
 				{
 					ParentStructNameInScript = FString(ParentName.Identifier);
 					FString ParentNameStripped = GetClassNameWithPrefixRemoved(ParentScope.Identifier);
-					StructClass = FindClass(*ParentNameStripped);
+					StructClass = AllClasses.FindClass(*ParentNameStripped);
 					if( !StructClass )
 					{
 						// If we find the literal class name, the user didn't use a prefix
-						StructClass = FindClass(ParentScope.Identifier);
+						StructClass = AllClasses.FindClass(ParentScope.Identifier);
 						if( StructClass )
 						{
 							FError::Throwf(TEXT("'struct': Parent struct class '%s' is missing a prefix, expecting '%s'"), ParentScope.Identifier, *FString::Printf(TEXT("%s%s"),StructClass->GetPrefixCPP(),ParentScope.Identifier) );
@@ -1568,7 +1349,7 @@ UScriptStruct* FHeaderParser::CompileStructDeclaration(UClass* Scope)
 		}
 		else if (Token.Matches(TEXT("UPROPERTY"), ESearchCase::CaseSensitive))
 		{
-			CompileVariableDeclaration(Struct, EPropertyDeclarationStyle::UPROPERTY);
+			CompileVariableDeclaration(AllClasses, Struct, EPropertyDeclarationStyle::UPROPERTY);
 		}
 		else if (Token.Matches(TEXT("UFUNCTION"), ESearchCase::CaseSensitive))
 		{
@@ -1947,7 +1728,7 @@ void FHeaderParser::PushNest( ENestType NestType, FName ThisName, UStruct* InNod
  * @param	NestType	nesting type of the current node
  * @param	Descr		text to use in error message if any errors are encountered
  */
-void FHeaderParser::PopNest( ENestType NestType, const TCHAR* Descr )
+void FHeaderParser::PopNest( FClasses& AllClasses, ENestType NestType, const TCHAR* Descr )
 {
 	// Validate the nesting state.
 	if (NestLevel <= 0)
@@ -1973,7 +1754,7 @@ void FHeaderParser::PopNest( ENestType NestType, const TCHAR* Descr )
 
 			// now validate all delegate variables declared in the class
 			TMap<FName, UFunction*> DelegateCache;
-			FixupDelegateProperties(TopFunction, TopFunction->GetOwnerClass(), DelegateCache);
+			FixupDelegateProperties(AllClasses, TopFunction, TopFunction->GetOwnerClass(), DelegateCache);
 		}
 	}
 	else if (NestType == NEST_Class)
@@ -1987,7 +1768,7 @@ void FHeaderParser::PopNest( ENestType NestType, const TCHAR* Descr )
 		{
 			// now validate all delegate variables declared in the class
 			TMap<FName,UFunction*> DelegateCache;
-			FixupDelegateProperties(TopClass, TopClass, DelegateCache);
+			FixupDelegateProperties(AllClasses, TopClass, TopClass, DelegateCache);
 		}
 
 		// Validate all the rep notify events here, to make sure they're implemented
@@ -2066,7 +1847,7 @@ void FHeaderParser::PopNest( ENestType NestType, const TCHAR* Descr )
 		if (ClassData->ContainsDelegates())
 		{
 			TMap<FName,UFunction*> DelegateCache;
-			FixupDelegateProperties(TopNode, ExactCast<UClass>(TopNode), DelegateCache);
+			FixupDelegateProperties(AllClasses, TopNode, ExactCast<UClass>(TopNode), DelegateCache);
 		}
 	}
 	else
@@ -2099,7 +1880,7 @@ void FHeaderParser::PopNest( ENestType NestType, const TCHAR* Descr )
  * @param	OwnerClass			the class currently being compiled.
  * @param	DelegateCache		cached map of delegates that have already been found; used for faster lookup.
  */
-void FHeaderParser::FixupDelegateProperties( UStruct* ValidationScope, UClass* OwnerClass, TMap<FName, UFunction*>& DelegateCache )
+void FHeaderParser::FixupDelegateProperties( FClasses& AllClasses, UStruct* ValidationScope, UClass* OwnerClass, TMap<FName, UFunction*>& DelegateCache )
 {
 	check(ValidationScope);
 	check(OwnerClass);
@@ -2168,7 +1949,7 @@ void FHeaderParser::FixupDelegateProperties( UStruct* ValidationScope, UClass* O
 						}
 
 						// make sure that the class that contains the delegate can be referenced here
-						UClass* DelegateOwnerClass = FindScriptClassOrThrow(DelegateClassName);
+						UClass* DelegateOwnerClass = AllClasses.FindScriptClassOrThrow(DelegateClassName);
 						CheckInScope(DelegateOwnerClass);
 						SourceDelegateFunction = Cast<UFunction>(FindField(DelegateOwnerClass, *DelegateName, false, UFunction::StaticClass(), NULL));	
 					}
@@ -2234,7 +2015,7 @@ void FHeaderParser::FixupDelegateProperties( UStruct* ValidationScope, UClass* O
 			UStruct* InternalStruct = Cast<UStruct>(Field);
 			if ( InternalStruct != NULL )
 			{
-				FixupDelegateProperties(InternalStruct, OwnerClass, DelegateCache);
+				FixupDelegateProperties(AllClasses, InternalStruct, OwnerClass, DelegateCache);
 			}
 		}
 	}
@@ -2467,6 +2248,7 @@ void FHeaderParser::CompileDirective(UClass* Class)
 
 bool FHeaderParser::GetVarType
 (
+	FClasses&                       AllClasses,
 	UStruct*                        Scope,
 	FPropertyBase&                  VarProperty,
 	EObjectFlags&                   ObjectFlags,
@@ -2566,7 +2348,7 @@ bool FHeaderParser::GetVarType
 				{
 					FError::Throwf(TEXT("Found more than one edit/visibility specifier (%s), only one is allowed"), *Specifier);
 				}
-				Flags |= CPF_Edit | CPF_DisableEditOnTemplate;
+				Flags |= CPF_Edit|CPF_DisableEditOnTemplate;
 				bSeenEditSpecifier = true;
 			}
 			else if (Specifier == TEXT("EditDefaultsOnly")) 
@@ -2575,7 +2357,7 @@ bool FHeaderParser::GetVarType
 				{
 					FError::Throwf(TEXT("Found more than one edit/visibility specifier (%s), only one is allowed"), *Specifier);
 				}
-				Flags |= CPF_Edit | CPF_DisableEditOnInstance;
+				Flags |= CPF_Edit|CPF_DisableEditOnInstance;
 				bSeenEditSpecifier = true;
 			}
 			else if (Specifier == TEXT("VisibleAnywhere")) 
@@ -2584,7 +2366,7 @@ bool FHeaderParser::GetVarType
 				{
 					FError::Throwf(TEXT("Found more than one edit/visibility specifier (%s), only one is allowed"), *Specifier);
 				}
-				Flags |= CPF_Edit | CPF_EditConst;
+				Flags |= CPF_Edit|CPF_EditConst;
 				bSeenEditSpecifier = true;
 			}
 			else if (Specifier == TEXT("VisibleInstanceOnly"))
@@ -2593,7 +2375,7 @@ bool FHeaderParser::GetVarType
 				{
 					FError::Throwf(TEXT("Found more than one edit/visibility specifier (%s), only one is allowed"), *Specifier);
 				}
-				Flags |= CPF_Edit | CPF_EditConst | CPF_DisableEditOnTemplate;
+				Flags |= CPF_Edit|CPF_EditConst|CPF_DisableEditOnTemplate;
 				bSeenEditSpecifier = true;
 			}
 			else if (Specifier == TEXT("VisibleDefaultsOnly"))
@@ -2602,7 +2384,7 @@ bool FHeaderParser::GetVarType
 				{
 					FError::Throwf(TEXT("Found more than one edit/visibility specifier (%s), only one is allowed"), *Specifier);
 				}
-				Flags |= CPF_Edit | CPF_EditConst | CPF_DisableEditOnInstance;
+				Flags |= CPF_Edit|CPF_EditConst|CPF_DisableEditOnInstance;
 				bSeenEditSpecifier = true;
 			}
 			else if (Specifier == TEXT("BlueprintReadWrite"))
@@ -2968,7 +2750,7 @@ bool FHeaderParser::GetVarType
 		uint64 OriginalVarTypeFlags = VarType.PropertyFlags;
 		VarType.PropertyFlags |= Flags;
 
-		GetVarType( Scope, VarProperty, InnerFlags, Disallow, TEXT("'tarray'"), &VarType, EPropertyDeclarationStyle::None, VariableCategory );
+		GetVarType( AllClasses, Scope, VarProperty, InnerFlags, Disallow, TEXT("'tarray'"), &VarType, EPropertyDeclarationStyle::None, VariableCategory );
 		if (VarProperty.ArrayType != EArrayType::None)
 		{
 			FError::Throwf(TEXT("Arrays within arrays not supported.") );
@@ -3181,7 +2963,7 @@ bool FHeaderParser::GetVarType
 				FToken InnerClass;
 				if (GetIdentifier(InnerClass))
 				{
-					TempClass = FindScriptClass(InnerClass.Identifier);
+					TempClass = AllClasses.FindScriptClass(InnerClass.Identifier);
 
 					if (bIsAutoweakPtrTemplate)
 					{
@@ -3222,7 +3004,7 @@ bool FHeaderParser::GetVarType
 			}
 			else
 			{
-				TempClass = FindScriptClass(VarType.Identifier);
+				TempClass = AllClasses.FindScriptClass(VarType.Identifier);
 			}
 
 			if (TempClass != NULL)
@@ -3250,7 +3032,7 @@ bool FHeaderParser::GetVarType
 							FError::Throwf(TEXT("'class': Missing class limitor"));
 						}
 
-						VarProperty.MetaClass = FindScriptClassOrThrow(Limitor.Identifier);
+						VarProperty.MetaClass = AllClasses.FindScriptClassOrThrow(Limitor.Identifier);
 
 						RequireSymbol( TEXT(">"), TEXT("'class limitor'") );
 					}
@@ -3304,16 +3086,14 @@ bool FHeaderParser::GetVarType
 		// Resolve delegates declared in another class  //@TODO: UCREMOVAL: This seems extreme
 		if (!bHandledType)
 		{
-			for (TObjectIterator<UClass> ClassIt; ClassIt; ++ClassIt)
+			for (FClass* OtherClass : AllClasses)
 			{
-				UClass* OtherClass = *ClassIt;
-
 				if (UFunction* DelegateFunc = Cast<UFunction>(FindField(OtherClass, *(IdentifierStripped + HEADER_GENERATED_DELEGATE_SIGNATURE_SUFFIX))))
 				{
 					bHandledType = true;
 
 					VarProperty = FPropertyBase( DelegateFunc->HasAnyFunctionFlags(FUNC_MulticastDelegate) ? CPT_MulticastDelegate : CPT_Delegate);
-					VarProperty.DelegateName = *FString::Printf(TEXT("%s.%s"), *GetClassNameWithPrefix(OtherClass), *IdentifierStripped);
+					VarProperty.DelegateName = *FString::Printf(TEXT("%s.%s"), *OtherClass->GetNameWithPrefix(), *IdentifierStripped);
 					
 					if (!(Disallow & CPF_InstancedReference))
 					{
@@ -3896,7 +3676,7 @@ UProperty* FHeaderParser::GetVarNameAndDim
 //
 // Compile a declaration in Token. Returns 1 if compiled, 0 if not.
 //
-bool FHeaderParser::CompileDeclaration( FToken& Token )
+bool FHeaderParser::CompileDeclaration( FClasses& AllClasses, FToken& Token )
 {
 	EAccessSpecifier AccessSpecifier = ParseAccessProtectionSpecifier(Token);
 	if (AccessSpecifier)
@@ -3920,7 +3700,7 @@ bool FHeaderParser::CompileDeclaration( FToken& Token )
 		// Start parsing the second class
 		bEncounteredNewStyleClass_UnmatchedBrackets = true;
 		bHaveSeenSecondInterfaceClass               = true;
-		ParseSecondInterfaceClass();
+		ParseSecondInterfaceClass(AllClasses);
 	}
 	else if (Token.Matches(TEXT("GENERATED_UCLASS_BODY")))
 	{
@@ -3964,7 +3744,7 @@ bool FHeaderParser::CompileDeclaration( FToken& Token )
 		}
 		bHaveSeenUClass = true;
 		bEncounteredNewStyleClass_UnmatchedBrackets = true;		
-		CompileClassDeclaration();
+		CompileClassDeclaration(AllClasses);
 	}
 	else if (Token.Matches(TEXT("UINTERFACE")) && (TopNest->Allow & ALLOW_Class))
 	{
@@ -3975,26 +3755,26 @@ bool FHeaderParser::CompileDeclaration( FToken& Token )
 		bHaveSeenUClass = true;
 		bEncounteredNewStyleClass_UnmatchedBrackets = true;
 		bHaveSeenFirstInterfaceClass = true;
-		CompileInterfaceDeclaration();
+		CompileInterfaceDeclaration(AllClasses);
 	}
 	else if (Token.Matches(TEXT("UFUNCTION"), ESearchCase::CaseSensitive))
 	{
-		CompileFunctionDeclaration();
+		CompileFunctionDeclaration(AllClasses);
 	}
 	else if (Token.Matches(TEXT("UDELEGATE")))
 	{
-		CompileDelegateDeclaration(Token.Identifier, EDelegateSpecifierAction::Parse);
+		CompileDelegateDeclaration(AllClasses, Token.Identifier, EDelegateSpecifierAction::Parse);
 	}
 	else if (IsValidDelegateDeclaration(Token)) // Legacy delegate parsing - it didn't need a UDELEGATE
 	{
-		CompileDelegateDeclaration(Token.Identifier);
+		CompileDelegateDeclaration(AllClasses, Token.Identifier);
 	}
 	else if (Token.Matches(TEXT("UPROPERTY"), ESearchCase::CaseSensitive))
 	{
 		CheckAllow( TEXT("'Member variable declaration'"), ALLOW_VarDecl );
 		check( TopNest->NestType == NEST_Class );
 
-		CompileVariableDeclaration(TopNode, EPropertyDeclarationStyle::UPROPERTY);
+		CompileVariableDeclaration(AllClasses, TopNode, EPropertyDeclarationStyle::UPROPERTY);
 	}
 	else if (Token.Matches(TEXT("UENUM")))
 	{
@@ -4004,7 +3784,7 @@ bool FHeaderParser::CompileDeclaration( FToken& Token )
 	else if (Token.Matches(TEXT("USTRUCT")))
 	{
 		// Struct definition.
-		CompileStructDeclaration(Class);
+		CompileStructDeclaration(AllClasses, Class);
 	}
 	else if (Token.Matches(TEXT("#")))
 	{
@@ -4028,11 +3808,11 @@ bool FHeaderParser::CompileDeclaration( FToken& Token )
 				// Remember all interface classes are now parsed so that we can parse other (non-U)classes.
 				bFinishedParsingInterfaceClasses = true;
 			}
-			PopNest( NEST_Interface, TEXT("'Interface'") );
+			PopNest( AllClasses, NEST_Interface, TEXT("'Interface'") );
 		}
 		else
 		{
-			PopNest( NEST_Class, TEXT("'Class'") );
+			PopNest( AllClasses, NEST_Class, TEXT("'Class'") );
 
 			// Ensure classes have a GENERATED_UCLASS_BODY declaration
 			if (bHaveSeenUClass && !bClassHasGeneratedBody)
@@ -4169,7 +3949,7 @@ bool FHeaderParser::SafeMatchSymbol( const TCHAR* Match )
 	return false;
 }
 
-void FHeaderParser::ParseClassNameDeclaration(FString& DeclaredClassName, FString& RequiredAPIMacroIfPresent)
+void FHeaderParser::ParseClassNameDeclaration(FClasses& AllClasses, FString& DeclaredClassName, FString& RequiredAPIMacroIfPresent)
 {
 	ParseNameWithPotentialAPIMacroPrefix(/*out*/ DeclaredClassName, /*out*/ RequiredAPIMacroIfPresent, TEXT("class"));
 
@@ -4185,7 +3965,7 @@ void FHeaderParser::ParseClassNameDeclaration(FString& DeclaredClassName, FStrin
 	if (bSpecifiesParentClass)
 	{
 		// Set the base class.
-		UClass* TempClass = GetQualifiedClass( TEXT("'extends'") );
+		UClass* TempClass = GetQualifiedClass(AllClasses, TEXT("'extends'"));
 		// a class cannot 'extends' an interface, use 'implements'
 		if (TempClass->ClassFlags & CLASS_Interface)
 		{
@@ -4242,7 +4022,7 @@ void FHeaderParser::ParseClassNameDeclaration(FString& DeclaredClassName, FStrin
 				}
 			}
 
-			HandleOneInheritedClass(*InterfaceName);
+			HandleOneInheritedClass(AllClasses, *InterfaceName);
 		}
 	}
 	else if( Class->GetSuperClass() )
@@ -4251,10 +4031,10 @@ void FHeaderParser::ParseClassNameDeclaration(FString& DeclaredClassName, FStrin
 	}
 }
 
-void FHeaderParser::HandleOneInheritedClass(FString InterfaceName)
+void FHeaderParser::HandleOneInheritedClass(FClasses& AllClasses, FString InterfaceName)
 {
 	// Check for UInterface derived interface inheritance
-	if (UClass* Interface = FindScriptClass(InterfaceName))
+	if (UClass* Interface = AllClasses.FindScriptClass(InterfaceName))
 	{
 		// Try to find the interface
 		if ( !Interface->HasAnyClassFlags(CLASS_Interface) )
@@ -4293,7 +4073,7 @@ void FHeaderParser::HandleOneInheritedClass(FString InterfaceName)
 /**
  * Compiles a class declaration.
  */
-void FHeaderParser::CompileClassDeclaration()
+void FHeaderParser::CompileClassDeclaration(FClasses& AllClasses)
 {
 	// Start of a class block.
 	CheckAllow( TEXT("'class'"), ALLOW_Class );
@@ -4325,7 +4105,7 @@ void FHeaderParser::CompileClassDeclaration()
 
 	FString DeclaredClassName;
 	FString RequiredAPIMacroIfPresent;
-	ParseClassNameDeclaration(/*out*/ DeclaredClassName, /*out*/ RequiredAPIMacroIfPresent);
+	ParseClassNameDeclaration(AllClasses, /*out*/ DeclaredClassName, /*out*/ RequiredAPIMacroIfPresent);
 
 	// Record that this class is RequiredAPI if the CORE_API style macro was present
 	if (!RequiredAPIMacroIfPresent.IsEmpty())
@@ -4369,7 +4149,7 @@ void FHeaderParser::CompileClassDeclaration()
 		{
 			FString WithinNameStr = RequireExactlyOneSpecifierValue(PropSpecifier);
 
-			UClass* RequiredWithinClass = FindClass(*WithinNameStr);
+			UClass* RequiredWithinClass = AllClasses.FindClass(*WithinNameStr);
 
 			if (!RequiredWithinClass)
 			{
@@ -4637,9 +4417,8 @@ void FHeaderParser::CompileClassDeclaration()
 		///Class->ClassWithin = ExpectedWithin;
 	}
 
-
 	// All classes must start with a valid Unreal prefix
-	const FString ExpectedClassName = GetClassNameWithPrefix(Class);
+	const FString ExpectedClassName = Class->GetNameWithPrefix();
 	if( DeclaredClassName != ExpectedClassName )
 	{
 		FError::Throwf(TEXT("Class name '%s' is invalid, should be identified as '%s'"), *DeclaredClassName, *ExpectedClassName );
@@ -4733,7 +4512,7 @@ void FHeaderParser::CompileClassDeclaration()
 	}
 }
 
-void FHeaderParser::ParseInterfaceNameDeclaration(FString& DeclaredInterfaceName, FString& RequiredAPIMacroIfPresent)
+void FHeaderParser::ParseInterfaceNameDeclaration(FClasses& AllClasses, FString& DeclaredInterfaceName, FString& RequiredAPIMacroIfPresent)
 {
 	ParseNameWithPotentialAPIMacroPrefix(/*out*/ DeclaredInterfaceName, /*out*/ RequiredAPIMacroIfPresent, TEXT("interface"));
 
@@ -4746,7 +4525,7 @@ void FHeaderParser::ParseInterfaceNameDeclaration(FString& DeclaredInterfaceName
 
 	// verify if our super class is an interface class
 	// the super class should have been marked as CLASS_Interface at the importing stage, if it were an interface
-	UClass* TempClass = GetQualifiedClass( TEXT("'extends'") );
+	UClass* TempClass = GetQualifiedClass(AllClasses, TEXT("'extends'"));
 	if( !(TempClass->ClassFlags & CLASS_Interface) )
 	{
 		// UInterface is special and actually extends from UObject, which isn't an interface
@@ -4765,7 +4544,7 @@ void FHeaderParser::ParseInterfaceNameDeclaration(FString& DeclaredInterfaceName
 	}
 }
 
-void FHeaderParser::ParseSecondInterfaceClass()
+void FHeaderParser::ParseSecondInterfaceClass(FClasses& AllClasses)
 {
 	FString ErrorMsg(TEXT("C++ interface mix-in class declaration"));
 
@@ -4774,10 +4553,10 @@ void FHeaderParser::ParseSecondInterfaceClass()
 	// Get a class name
 	FString DeclaredInterfaceName;
 	FString RequiredAPIMacroIfPresent;
-	ParseInterfaceNameDeclaration(/*out*/ DeclaredInterfaceName, /*out*/ RequiredAPIMacroIfPresent);
+	ParseInterfaceNameDeclaration(AllClasses, /*out*/ DeclaredInterfaceName, /*out*/ RequiredAPIMacroIfPresent);
 
 	// All classes must start with a valid Unreal prefix
-	const FString ExpectedInterfaceName = GetClassNameWithPrefix(Class, EEnforceInterfacePrefix::I);
+	const FString ExpectedInterfaceName = Class->GetNameWithPrefix(EEnforceInterfacePrefix::I);
 	if (DeclaredInterfaceName != ExpectedInterfaceName)
 	{
 		FError::Throwf(TEXT("Interface name '%s' is invalid, the second interface class should be identified as '%s'"), *DeclaredInterfaceName, *ExpectedInterfaceName);
@@ -4793,7 +4572,7 @@ void FHeaderParser::ParseSecondInterfaceClass()
 /**
  *  compiles Java or C# style interface declaration
  */
-void FHeaderParser::CompileInterfaceDeclaration()
+void FHeaderParser::CompileInterfaceDeclaration(FClasses& AllClasses)
 {
 	// Start of an interface block. Since Interfaces and Classes are always at the same nesting level,
 	// whereever a class declaration is allowed, an interface declaration is also allowed.
@@ -4811,7 +4590,7 @@ void FHeaderParser::CompileInterfaceDeclaration()
 
 	// New style files have the interface name / extends afterwards
 	RequireIdentifier(TEXT("class"), TEXT("Interface declaration"));
-	ParseInterfaceNameDeclaration(/*out*/ DeclaredInterfaceName, /*out*/ RequiredAPIMacroIfPresent);
+	ParseInterfaceNameDeclaration(AllClasses, /*out*/ DeclaredInterfaceName, /*out*/ RequiredAPIMacroIfPresent);
 
 
 	// Record that this interface is RequiredAPI if the CORE_API style macro was present
@@ -4866,7 +4645,7 @@ void FHeaderParser::CompileInterfaceDeclaration()
 	}
 
 	// All classes must start with a valid Unreal prefix
-	const FString ExpectedInterfaceName = GetClassNameWithPrefix(Class, EEnforceInterfacePrefix::U);
+	const FString ExpectedInterfaceName = Class->GetNameWithPrefix(EEnforceInterfacePrefix::U);
 	if (DeclaredInterfaceName != ExpectedInterfaceName)
 	{
 		FError::Throwf(TEXT("Interface name '%s' is invalid, the first class should be identified as '%s'"), *DeclaredInterfaceName, *ExpectedInterfaceName );
@@ -4900,7 +4679,7 @@ bool FHeaderParser::IsValidDelegateDeclaration(const FToken& Token) const
 }
 
 // Parse the parameter list of a function or delegate declaration
-void FHeaderParser::ParseParameterList(UFunction* Function, bool bExpectCommaBeforeName, TMap<FName, FString>* MetaData)
+void FHeaderParser::ParseParameterList(FClasses& AllClasses, UFunction* Function, bool bExpectCommaBeforeName, TMap<FName, FString>* MetaData)
 {
 	// Get parameter list.
 	if ( MatchSymbol(TEXT(")")) )
@@ -4912,7 +4691,7 @@ void FHeaderParser::ParseParameterList(UFunction* Function, bool bExpectCommaBef
 		// Get parameter type.
 		FToken Property(CPT_None);
 		EObjectFlags ObjectFlags;
-		GetVarType( TopNode, Property, ObjectFlags, ~(CPF_ParmFlags|CPF_AutoWeak|CPF_RepSkip), TEXT("Function parameter"), NULL, EPropertyDeclarationStyle::None, (Function->FunctionFlags & FUNC_Net) ? EVariableCategory::ReplicatedParameter: EVariableCategory::RegularParameter);
+		GetVarType( AllClasses, TopNode, Property, ObjectFlags, ~(CPF_ParmFlags|CPF_AutoWeak|CPF_RepSkip), TEXT("Function parameter"), NULL, EPropertyDeclarationStyle::None, (Function->FunctionFlags & FUNC_Net) ? EVariableCategory::ReplicatedParameter: EVariableCategory::RegularParameter);
 		Property.PropertyFlags |= CPF_Parm;
 
 		if (bExpectCommaBeforeName)
@@ -5025,7 +4804,7 @@ void FHeaderParser::ParseParameterList(UFunction* Function, bool bExpectCommaBef
 	RequireSymbol( TEXT(")"), TEXT("parameter list") );
 }
 
-void FHeaderParser::CompileDelegateDeclaration(const TCHAR* DelegateIdentifier, EDelegateSpecifierAction::Type SpecifierAction)
+void FHeaderParser::CompileDelegateDeclaration(FClasses& AllClasses, const TCHAR* DelegateIdentifier, EDelegateSpecifierAction::Type SpecifierAction)
 {
 	TMap<FName, FString> MetaData;
 
@@ -5107,7 +4886,7 @@ void FHeaderParser::CompileDelegateDeclaration(const TCHAR* DelegateIdentifier,
 
 	if (bHasReturnValue)
 	{
-		GetVarType( TopNode, ReturnType, ReturnValueFlags, 0, NULL, NULL, EPropertyDeclarationStyle::None, EVariableCategory::Return );
+		GetVarType( AllClasses, TopNode, ReturnType, ReturnValueFlags, 0, NULL, NULL, EPropertyDeclarationStyle::None, EVariableCategory::Return );
 		RequireSymbol(TEXT(","), CurrentScopeName);
 	}
 
@@ -5167,7 +4946,7 @@ void FHeaderParser::CompileDelegateDeclaration(const TCHAR* DelegateIdentifier,
 	{
 		RequireSymbol(TEXT(","), CurrentScopeName);
 
-		ParseParameterList(DelegateSignatureFunction, /*bExpectCommaBeforeName=*/ true);
+		ParseParameterList(AllClasses, DelegateSignatureFunction, /*bExpectCommaBeforeName=*/ true);
 
 		// Check the expected versus actual number of parameters
 		int32 ParamCount = FoundParamCount - DelegateParameterCountStrings.GetTypedData() + 1;
@@ -5203,7 +4982,7 @@ void FHeaderParser::CompileDelegateDeclaration(const TCHAR* DelegateIdentifier,
 	DelegateSignatureFunction->Bind();
 
 	// End the nesting
-	PopNest( NEST_FunctionDeclaration, CurrentScopeName );
+	PopNest( AllClasses, NEST_FunctionDeclaration, CurrentScopeName );
 
 	// Don't allow delegate signatures to be redefined.
 	for (TFieldIterator<UFunction> FunctionIt(CastChecked<UStruct>(DelegateSignatureFunction->GetOuter()), EFieldIteratorFlags::IncludeSuper); FunctionIt; ++FunctionIt)
@@ -5219,7 +4998,7 @@ void FHeaderParser::CompileDelegateDeclaration(const TCHAR* DelegateIdentifier,
 /**
  * Parses and compiles a function declaration
  */
-void FHeaderParser::CompileFunctionDeclaration()
+void FHeaderParser::CompileFunctionDeclaration(FClasses& AllClasses)
 {
 	CheckAllow(TEXT("'Function'"), ALLOW_Function);
 
@@ -5440,7 +5219,7 @@ void FHeaderParser::CompileFunctionDeclaration()
 	// C++ style functions always have a return value type, even if it's void
 	if (!MatchIdentifier(TEXT("void")))
 	{
-		bHasReturnValue = GetVarType( TopNode, ReturnType, ReturnValueFlags, 0, NULL, NULL, EPropertyDeclarationStyle::None, EVariableCategory::Return );
+		bHasReturnValue = GetVarType( AllClasses, TopNode, ReturnType, ReturnValueFlags, 0, NULL, NULL, EPropertyDeclarationStyle::None, EVariableCategory::Return );
 	}
 
 	// Get function or operator name.
@@ -5502,7 +5281,7 @@ void FHeaderParser::CompileFunctionDeclaration()
 	FFunctionData* StoredFuncData = ClassData->AddFunction(FuncInfo);
 
 	// Get parameter list.
-	ParseParameterList(TopFunction, false, &MetaData);
+	ParseParameterList(AllClasses, TopFunction, false, &MetaData);
 
 	// Get return type, if any.
 	if (bHasReturnValue)
@@ -5731,7 +5510,7 @@ Found:
 	}
 
 	// Just declaring a function, so end the nesting.
-	PopNest( NEST_FunctionDeclaration, TypeOfFunction );
+	PopNest( AllClasses, NEST_FunctionDeclaration, TypeOfFunction );
 
 	// Optionally consume a semicolon
 	// This is optional to allow inline function definitions, as long as the function body is inside #if CPP ... #endif
@@ -5926,7 +5705,7 @@ void FHeaderParser::ValidatePropertyIsDeprecatedIfNecessary(FPropertyBase& VarPr
 	}
 }
 
-void FHeaderParser::CompileVariableDeclaration(UStruct* Struct, EPropertyDeclarationStyle::Type PropertyDeclarationStyle)
+void FHeaderParser::CompileVariableDeclaration(FClasses& AllClasses, UStruct* Struct, EPropertyDeclarationStyle::Type PropertyDeclarationStyle)
 {
 	uint64 DisallowFlags = CPF_ParmFlags;
 	uint64 EdFlags       = 0;
@@ -5934,7 +5713,7 @@ void FHeaderParser::CompileVariableDeclaration(UStruct* Struct, EPropertyDeclara
 	// Get variable type.
 	FPropertyBase OriginalProperty(CPT_None);
 	EObjectFlags ObjectFlags = RF_NoFlags;
-	GetVarType( Struct, OriginalProperty, ObjectFlags, DisallowFlags, TEXT("Member variable declaration"), /*OuterPropertyType=*/ NULL, PropertyDeclarationStyle, EVariableCategory::Member );
+	GetVarType( AllClasses, Struct, OriginalProperty, ObjectFlags, DisallowFlags, TEXT("Member variable declaration"), /*OuterPropertyType=*/ NULL, PropertyDeclarationStyle, EVariableCategory::Member );
 	OriginalProperty.PropertyFlags |= EdFlags;
 
 	FString* Category = OriginalProperty.MetaData.Find("Category");
@@ -6034,7 +5813,7 @@ void FHeaderParser::CompileVariableDeclaration(UStruct* Struct, EPropertyDeclara
 // Compile a statement: Either a declaration or a command.
 // Returns 1 if success, 0 if end of file.
 //
-bool FHeaderParser::CompileStatement()
+bool FHeaderParser::CompileStatement(FClasses& AllClasses)
 {
 	// Get a token and compile it.
 	FToken Token;
@@ -6043,7 +5822,7 @@ bool FHeaderParser::CompileStatement()
 		// End of file.
 		return false;
 	}
-	else if (!CompileDeclaration( Token ))
+	else if (!CompileDeclaration( AllClasses, Token ))
 	{
 		FError::Throwf(TEXT("'%s': Bad command or expression"), Token.Identifier );
 	}
@@ -6183,7 +5962,7 @@ void FHeaderParser::FinalizeScriptExposedFunctions(UClass* Class)
 // Parses the header associated with the specified class.
 // Returns result enumeration.
 //
-ECompilationResult::Type FHeaderParser::ParseHeaderForOneClass(FClassTree& AllClasses, UClass* InClass)
+ECompilationResult::Type FHeaderParser::ParseHeaderForOneClass(FClasses& AllClasses, FClass* InClass)
 {
 	// Early-out if this class has previously failed some aspect of parsing
 	if (FailedClassesAnnotation.Get(InClass))
@@ -6286,7 +6065,7 @@ ECompilationResult::Type FHeaderParser::ParseHeaderForOneClass(FClassTree& AllCl
 #endif
 	{
 		// Parse entire program.
-		while (CompileStatement())
+		while (CompileStatement(AllClasses))
 		{
 			bEmptyFile = false;
 
@@ -6417,14 +6196,10 @@ ECompilationResult::Type FHeaderParser::ParseHeaderForOneClass(FClassTree& AllCl
 -----------------------------------------------------------------------------*/
 
 // Parse Class's annotated headers and optionally its child classes.  Marks the class as CLASS_Parsed.
-ECompilationResult::Type FHeaderParser::ParseHeaders(FClassTree& AllClasses, FHeaderParser& HeaderParser, UClass* Class, bool bParseSubclasses)
+ECompilationResult::Type FHeaderParser::ParseHeaders(FClasses& AllClasses, FHeaderParser& HeaderParser, FClass* Class, bool bParseSubclasses)
 {
-	check(Class != NULL);
-
 	if (!GClassStrippedHeaderTextMap.Contains(Class))
-	{
 		return ECompilationResult::Succeeded;
-	}
 
 	ECompilationResult::Type Result = ECompilationResult::Succeeded;
 
@@ -6437,20 +6212,20 @@ ECompilationResult::Type FHeaderParser::ParseHeaders(FClassTree& AllClasses, FHe
 		// Check if the dependent class name came from #include directive in which case we allow it to fail.
 		DependentClassNameFromHeader(*DependentClassName, DependentClassName);
 
-		UClass* DependsOnClass = FindScriptClass(DependentClassName);
+		FClass* DependsOnClass = AllClasses.FindScriptClass(DependentClassName);
 		if (!DependsOnClass)
 		{
 			// Try again with actor class name
 			const FString DependentClassNameStripped  = GetClassNameWithoutPrefix(DependentClassName);
 			const FString ActorDependentClassName = FString(TEXT("A")) + DependentClassNameStripped;
 
-			DependsOnClass = FindScriptClass(ActorDependentClassName);
+			DependsOnClass = AllClasses.FindScriptClass(ActorDependentClassName);
 			if (!DependsOnClass)
 			{
 				// Try again with temporary class name. This may be struct only header.
 				const FString TemporaryDependentClassName = FString(TEXT("U")) + GenerateTemporaryClassName(*DependentClassNameStripped);
 
-				DependsOnClass = FindScriptClass(TemporaryDependentClassName);
+				DependsOnClass = AllClasses.FindScriptClass(TemporaryDependentClassName);
 				if (!DependsOnClass)
 				{
 					// Ingore the error and remove this entry from DependentOn array.
@@ -6473,23 +6248,23 @@ ECompilationResult::Type FHeaderParser::ParseHeaders(FClassTree& AllClasses, FHe
 		}
 
 		// Check for circular dependency. If the DependsOnClass is dependent on the SubClass, there is one.
-		if (HeaderParser.IsDependentOn(DependsOnClass, Class, AllClasses))
+		if (AllClasses.IsDependentOn(DependsOnClass, Class))
 		{
 			FError::Throwf(TEXT("Class %s DependsOn(%s) is a circular dependency."),*DependsOnClass->GetName(),*Class->GetName());
 		}
 
 		// Find first base class of DependsOnClass that is not a base class of Class.
-		TArray<UClass*> ClassesToParse;
+		TArray<FClass*> ClassesToParse;
 		ClassesToParse.Add(DependsOnClass);
 
-		for ( UClass* ParentClass = DependsOnClass->GetSuperClass(); ParentClass && !ParentClass->HasAnyClassFlags(CLASS_Parsed|CLASS_Intrinsic); ParentClass = ParentClass->GetSuperClass() )
+		for ( FClass* ParentClass = DependsOnClass->GetSuperClass(); ParentClass && !ParentClass->HasAnyClassFlags(CLASS_Parsed|CLASS_Intrinsic); ParentClass = ParentClass->GetSuperClass() )
 		{
 			ClassesToParse.Add(ParentClass);
 		}
 
 		while (ClassesToParse.Num() > 0)
 		{
-			UClass* NextClass = ClassesToParse.Pop();
+			FClass* NextClass = ClassesToParse.Pop();
 
 			ECompilationResult::Type ParseResult = ParseHeaders(AllClasses, HeaderParser, NextClass, true);
 
@@ -6530,15 +6305,8 @@ ECompilationResult::Type FHeaderParser::ParseHeaders(FClassTree& AllClasses, FHe
 	// Parse subclasses if instructed to do so.
 	if (bParseSubclasses)
 	{
-		const FClassTree* ClassLeaf = AllClasses.FindNode(Class);
-		TArray<const FClassTree*> ChildLeaves;
-		ClassLeaf->GetChildClasses(ChildLeaves);
-
-		for (int32 ClassIndex = 0; ClassIndex < ChildLeaves.Num(); ClassIndex++)
+		for (auto SubClass : AllClasses.GetDerivedClasses(Class))
 		{
-			const FClassTree* ChildLeaf = ChildLeaves[ClassIndex];
-			UClass* SubClass = ChildLeaf->GetClass();
-
 			//note: you must always pass in the root tree node here, since we may add new classes to the tree
 			// if an manual dependency (through dependson()) is encountered
 			ECompilationResult::Type ParseResult = ParseHeaders(AllClasses, HeaderParser, SubClass, bParseSubclasses);
@@ -6573,7 +6341,7 @@ bool FHeaderParser::DependentClassNameFromHeader(const TCHAR* HeaderFilename, FS
 }
 
 // Begins the process of exporting C++ class declarations for native classes in the specified package
-void FHeaderParser::ExportNativeHeaders( UPackage* CurrentPackage, FClassTree& AllClasses, bool bAllowSaveExportedHeaders, bool bUseRelativePaths )
+void FHeaderParser::ExportNativeHeaders( UPackage* CurrentPackage, FClasses& AllClasses, bool bAllowSaveExportedHeaders, bool bUseRelativePaths )
 {
 	// Build a list of header filenames
 	TArray<FString>	ClassHeaderFilenames;
@@ -6708,26 +6476,14 @@ FString FHeaderParser::RequireExactlyOneSpecifierValue(const FPropertySpecifier&
 	return Specifier.Values[0];
 }
 
-// Parse all headers for classes that are inside LimitOuter.
-ECompilationResult::Type FHeaderParser::ParseAllHeadersInside(FFeedbackContext* Warn, UPackage* LimitOuter, bool bAllowSaveExportedHeaders, bool bUseRelativePaths)
+// Parse all headers for classes that are inside CurrentPackage.
+ECompilationResult::Type FHeaderParser::ParseAllHeadersInside(FFeedbackContext* Warn, UPackage* CurrentPackage, bool bAllowSaveExportedHeaders, bool bUseRelativePaths)
 {
 	// Disable loading of objects outside of this package (or more exactly, objects which aren't UFields, CDO, or templates)
 	TGuardValue<bool> AutoRestoreVerifyObjectRefsFlag(GVerifyObjectReferencesOnly, true);
 
-	// Parse all classes in the package, no matter their parent
-	UClass* BaseClass = UObject::StaticClass();
-
-	// build a class inheritance tree for this package
-	FClassTree AllClasses(BaseClass);
-
-	for (TObjectIterator<UClass> ObjIt; ObjIt; ++ObjIt)
-	{
-		if ( ((LimitOuter == NULL) || ObjIt->IsIn(LimitOuter)) )
-		{
-			UClass* Class = *ObjIt;
-			AllClasses.AddClass(Class);
-		}
-	}
+	// Object which represents all parsed classes
+	FClasses AllClasses(CurrentPackage);
 
 	AllClasses.Validate();
 
@@ -6737,7 +6493,7 @@ ECompilationResult::Type FHeaderParser::ParseAllHeadersInside(FFeedbackContext*
 	Warn->SetContext(&HeaderParser);
 
 	// Set up a filename for the error context if we don't even get as far parsing a class
-	HeaderParser.Filename = IFileManager::Get().ConvertToAbsolutePathForExternalAppForRead(*GClassSourceFileMap[BaseClass]);
+	HeaderParser.Filename = IFileManager::Get().ConvertToAbsolutePathForExternalAppForRead(*GClassSourceFileMap[AllClasses.GetRootClass()]);
 
 	// Hierarchically parse all classes.
 	ECompilationResult::Type Result = ECompilationResult::Succeeded;
@@ -6748,7 +6504,7 @@ ECompilationResult::Type FHeaderParser::ParseAllHeadersInside(FFeedbackContext*
 		// Parse the headers
 		const bool bParseSubclasses = true;
 
-		Result = FHeaderParser::ParseHeaders(AllClasses, HeaderParser, BaseClass, bParseSubclasses);
+		Result = FHeaderParser::ParseHeaders(AllClasses, HeaderParser, AllClasses.GetRootClass(), bParseSubclasses);
 
 		// Export the autogenerated code wrappers
 		if (Result == ECompilationResult::Succeeded)
@@ -6758,7 +6514,7 @@ ECompilationResult::Type FHeaderParser::ParseAllHeadersInside(FFeedbackContext*
 			// from the feedback context.
 			Warn->SetContext(NULL);
 
-			ExportNativeHeaders(LimitOuter, AllClasses, bAllowSaveExportedHeaders, bUseRelativePaths);
+			ExportNativeHeaders(CurrentPackage, AllClasses, bAllowSaveExportedHeaders, bUseRelativePaths);
 
 			// Done with header generation
 			if (HeaderParser.LinesParsed > 0)
@@ -6784,48 +6540,6 @@ ECompilationResult::Type FHeaderParser::ParseAllHeadersInside(FFeedbackContext*
 	return Result;
 }
 
-/** 
- * Attempts to get class prefix. If the given class name does not start with a valid Unreal Prefix, it will return an empty string.
- *
- * @param InClassName - Name w/ potential prefix to check
- */
-FString FHeaderParser::GetClassPrefix( const FString InClassName )
-{
-	bool bIsLabledDeprecated;
-	return FHeaderParser::GetClassPrefix(InClassName, /*out*/ bIsLabledDeprecated);
-}
-
-/** 
- * Attempts to get class prefix. If the given class name does not start with a valid Unreal Prefix, it will return an empty string.
- *
- * @param InClassName - Name w/ potential prefix to check
- * @param out bIsLabeledDeprecated - Reference param set to True if the class name is marked as deprecated
- */
-FString FHeaderParser::GetClassPrefix( const FString InClassName, bool& bIsLabeledDeprecated )
-{
-	FString ClassPrefix = InClassName.Left(1);
-
-	bIsLabeledDeprecated = false;
-	// First make sure the class name starts with a valid prefix
-	if( ClassPrefix == TEXT("I") ||
-		ClassPrefix == TEXT("A") || 
-		ClassPrefix == TEXT("U") )
-	{
-		// If it is a class prefix, check for deprecated class prefix
-		if( InClassName.Mid(1, 11) == TEXT("DEPRECATED_") )
-		{
-			bIsLabeledDeprecated = true;
-			return InClassName.Left(12);
-		}
-	}
-	// Otherwise check for struct prefix. If it's not a class or struct prefix, it's invalid
-	else if( ClassPrefix != TEXT("F") && ClassPrefix != TEXT("T") )
-	{
-		ClassPrefix = TEXT("");
-	}
-	return ClassPrefix;
-}
-
 /** 
  * Returns True if the given class name includes a valid Unreal prefix and matches up with the given original class.
  *
@@ -6855,69 +6569,6 @@ bool FHeaderParser::ClassNameHasValidPrefix(const FString InNameToCheck, const F
 	return bNamesMatch;
 }
 
-/** 
- * Returns True if the given class name includes a valid Unreal prefix and matches based on the given class.
- *
- * @param InNameToCheck - Name w/ potential prefix to check
- * @param OriginalClass - Class to check against
- */
-bool FHeaderParser::ClassNameHasValidPrefix(const FString InNameToCheck, const UClass* OriginalClass)
-{
-	bool bIsLabledDeprecated;
-	GetClassPrefix( InNameToCheck, bIsLabledDeprecated );
-
-	// If the class is labeled deprecated, don't try to resolve it during header generation, valid results can't be guaranteed.
-	if (bIsLabledDeprecated)
-	{
-		return true;
-	}
-
-	const FString OriginalClassName = GetClassNameWithPrefix(OriginalClass);
-
-	bool bNamesMatch = (InNameToCheck == OriginalClassName);
-
-	if (!bNamesMatch)
-	{
-		//@TODO: UCREMOVAL: I/U interface hack - Ignoring prefixing for this call
-		if (OriginalClass->HasAnyClassFlags(CLASS_Interface))
-		{
-			bNamesMatch = InNameToCheck.Mid(1) == OriginalClassName.Mid(1);
-		}
-	}
-
-	return bNamesMatch;
-}
-
-/////////////////////////////////////////////////////
-// FHeaderPreParser
-
-FString FHeaderParser::GetClassNameWithPrefixRemoved( const FString InClassName )
-{
-	const FString ClassPrefix = GetClassPrefix( InClassName );
-	if( !ClassPrefix.IsEmpty() )
-	{	
-		return InClassName.Right(InClassName.Len() - ClassPrefix.Len());
-	}
-	return FString();
-}
-
-FString FHeaderParser::GetClassNameWithoutPrefix( const FString& InClassNameOrFilename )
-{
-	// Check for header names (they don't come with a full path so we only search for the first dot)
-	const int32 DotIndex = InClassNameOrFilename.Find(TEXT("."));
-	if (DotIndex == INDEX_NONE)
-	{
-		const FString ClassPrefix = GetClassPrefix( InClassNameOrFilename );
-		return InClassNameOrFilename.Right(InClassNameOrFilename.Len() - ClassPrefix.Len());
-	}
-	else
-	{
-		return InClassNameOrFilename.Mid(0, DotIndex);
-	}
-}
-
-/////////////////////////////////////////////////////
-
 void FHeaderParser::ParseClassName(const TCHAR* Temp, FString& ClassName)
 {
 	// Skip leading whitespace
diff --git a/Engine/Source/Programs/UnrealHeaderTool/Private/HeaderParser.h b/Engine/Source/Programs/UnrealHeaderTool/Private/HeaderParser.h
--- a/Engine/Source/Programs/UnrealHeaderTool/Private/HeaderParser.h
+++ b/Engine/Source/Programs/UnrealHeaderTool/Private/HeaderParser.h
@@ -9,6 +9,9 @@
 #include "ParserHelper.h"
 #include "BaseParser.h"
 
+class FClass;
+class FClasses;
+
 /*-----------------------------------------------------------------------------
 	Constants & types.
 -----------------------------------------------------------------------------*/
@@ -88,21 +91,6 @@ public:
 	//   ScriptText (text outside of #if CPP and #if DEFAULTS blocks)
 	static void SimplifiedClassParse(const TCHAR* Buffer, bool& bIsInterface, TArray<FName>& DependentOn, FString& out_ClassName, FString& out_ParentClassName, int32& out_ClassDeclLine, FStringOutputDevice& ScriptText);
 	
-	/** 
-	 * Attempts to get class prefix. If the given class name does not start with a valid Unreal Prefix, it will return an empty string.
-	 *
-	 * @param InClassName - Name w/ potential prefix to check
-	 */
-	static FString GetClassPrefix( const FString InClassName );
-
-	/** 
-	 * Attempts to get class prefix. If the given class name does not start with a valid Unreal Prefix, it will return an empty string.
-	 *
-	 * @param InClassName - Name w/ potential prefix to check
-	 * @param out bIsLabledDeprecated - Reference param set to True if the class name is marked as deprecated
-	 */
-	static FString GetClassPrefix( const FString InClassName, bool& bIsLabledDeprecated );
-
 	/** 
 	 * Returns True if the given class name includes a valid Unreal prefix and matches up with the given original class Name.
 	 *
@@ -110,27 +98,6 @@ public:
 	 * @param OriginalClassName - Name of class w/ no prefix to check against
 	 */
 	static bool ClassNameHasValidPrefix(const FString InNameToCheck, const FString OriginalClassName);
-	/** 
-	 * Returns True if the given class name includes a valid Unreal prefix and matches based on the given class.
-	 *
-	 * @param InNameToCheck - Name w/ potential prefix to check
-	 * @param OriginalClass - Class to check against
-	 */
-	static bool ClassNameHasValidPrefix(const FString InNameToCheck, const UClass* OriginalClass);
-
-	/** 
-	 * Attempts to strip the given class name of it's affixed prefix. If no prefix exists, will return a blank string.
-	 *
-	 * @param InClassName - Class Name with a prefix to be stripped
-	 */
-	static FString GetClassNameWithPrefixRemoved( const FString InClassName );
-
-	/** 
-	 * Attempts to strip the given class name of it's affixed prefix. If no prefix exists, will return unchanged string.
-	 *
-	 * @param InClassName - Class Name with a prefix to be stripped
-	 */
-	static FString GetClassNameWithoutPrefix( const FString& InClassNameOrFilename );
 
 	/**
 	 * Tries to convert the header file name to a class name (with 'U' prefix)
@@ -171,7 +138,7 @@ protected:
 	FFeedbackContext* Warn;
 
 	// Class that is currently being parsed.
-	UClass* Class;
+	FClass* Class;
 
 	// Filename currently being parsed
 	FString Filename;
@@ -299,7 +266,7 @@ protected:
 	bool IsBitfieldProperty();
 
 	// Parse the parameter list of a function or delegate declaration
-	void ParseParameterList(UFunction* Function, bool bExpectCommaBeforeName = false, TMap<FName, FString>* MetaData = NULL);
+	void ParseParameterList(FClasses& AllClasses, UFunction* Function, bool bExpectCommaBeforeName = false, TMap<FName, FString>* MetaData = NULL);
 
 	// Throws if a specifier value wasn't provided
 	void RequireSpecifierValue(const FPropertySpecifier& Specifier, bool bRequireExactlyOne = false);
@@ -315,7 +282,7 @@ protected:
 	 *
 	 * @return	Result enumeration.
 	 */
-	static ECompilationResult::Type ParseHeaders(FClassTree& AllClasses, FHeaderParser& HeaderParser, UClass* Class, bool bParseSubclasses);
+	static ECompilationResult::Type ParseHeaders(FClasses& AllClasses, FHeaderParser& HeaderParser, FClass* Class, bool bParseSubclasses);
 
 	//@TODO: Remove this method
 	static void ParseClassName(const TCHAR* Temp, FString& ClassName);
@@ -332,25 +299,25 @@ protected:
 	 * @param	CurrentPackage	the package being compiled
 	 * @param	AllClasses		the class tree for CurrentPackage
 	 */
-	static void ExportNativeHeaders( UPackage* CurrentPackage, FClassTree& AllClasses, bool bAllowSaveExportedHeaders, bool bUseRelativePaths );
+	static void ExportNativeHeaders( UPackage* CurrentPackage, FClasses& AllClasses, bool bAllowSaveExportedHeaders, bool bUseRelativePaths );
 
 	// FContextSupplier interface.
 	virtual FString GetContext() OVERRIDE;
 	// End of FContextSupplier interface.
 
 	// High-level compiling functions.
-	ECompilationResult::Type ParseHeaderForOneClass(FClassTree& AllClasses, UClass* InClass);
+	ECompilationResult::Type ParseHeaderForOneClass(FClasses& AllClasses, FClass* InClass);
 	void CompileDirective(UClass* Class);
 	void FinalizeScriptExposedFunctions(UClass* Class);
 	UEnum* CompileEnum(UClass* Owner);
-	UScriptStruct* CompileStructDeclaration(UClass* Owner);
-	bool CompileDeclaration(FToken& Token);
+	UScriptStruct* CompileStructDeclaration(FClasses& AllClasses, FClass* Owner);
+	bool CompileDeclaration(FClasses& AllClasses, FToken& Token);
 	/** Skip C++ (noexport) declaration. */
 	bool SkipDeclaration(FToken& Token);
 	/** Similar to MatchSymbol() but will return to the exact location as on entry if the symbol was not found. */
 	bool SafeMatchSymbol(const TCHAR* Match);
-	void HandleOneInheritedClass(FString InterfaceName);
-	void ParseClassNameDeclaration(FString& DeclaredClassName, FString& RequiredAPIMacroIfPresent);
+	void HandleOneInheritedClass(FClasses& AllClasses, FString InterfaceName);
+	void ParseClassNameDeclaration(FClasses& AllClasses, FString& DeclaredClassName, FString& RequiredAPIMacroIfPresent);
 
 	/** The property style of a variable declaration being parsed */
 	struct EPropertyDeclarationStyle
@@ -362,16 +329,16 @@ protected:
 		};
 	};
 
-	void CompileClassDeclaration    ();
-	void CompileDelegateDeclaration (const TCHAR* DelegateIdentifier, EDelegateSpecifierAction::Type SpecifierAction = EDelegateSpecifierAction::DontParse);
-	void CompileFunctionDeclaration ();
-	void CompileVariableDeclaration (UStruct* Struct, EPropertyDeclarationStyle::Type PropertyDeclarationStyle);
-	void CompileInterfaceDeclaration();
+	void CompileClassDeclaration    (FClasses& AllClasses);
+	void CompileDelegateDeclaration (FClasses& AllClasses, const TCHAR* DelegateIdentifier, EDelegateSpecifierAction::Type SpecifierAction = EDelegateSpecifierAction::DontParse);
+	void CompileFunctionDeclaration (FClasses& AllClasses);
+	void CompileVariableDeclaration (FClasses& AllClasses, UStruct* Struct, EPropertyDeclarationStyle::Type PropertyDeclarationStyle);
+	void CompileInterfaceDeclaration(FClasses& AllClasses);
 
-	void ParseInterfaceNameDeclaration(FString& DeclaredInterfaceName, FString& RequiredAPIMacroIfPresent);
-	void ParseSecondInterfaceClass();
+	void ParseInterfaceNameDeclaration(FClasses& AllClasses, FString& DeclaredInterfaceName, FString& RequiredAPIMacroIfPresent);
+	void ParseSecondInterfaceClass(FClasses& AllClasses);
 
-	bool CompileStatement();
+	bool CompileStatement(FClasses& AllClasses);
 
 	// Compute the function parameter size and save the return offset
 	static void ComputeFunctionParametersSize( UClass* InClass );
@@ -418,6 +385,7 @@ protected:
 	/**
 	 * Parses a variable or return value declaration and determines the variable type and property flags.
 	 *
+	 * @param   AllClasses                the class tree for CurrentPackage
 	 * @param   Scope                     struct to create the property in
 	 * @param   VarProperty               will be filled in with type and property flag data for the property declaration that was parsed
 	 * @param   ObjectFlags               will contain the object flags that should be assigned to all UProperties which are part of this declaration (will not be determined
@@ -430,7 +398,16 @@ protected:
 	 *
 	 * @return  true if the variable type was parsed
 	 */
-	bool GetVarType(UStruct* Scope, FPropertyBase& VarProperty, EObjectFlags& ObjectFlags, uint64 Disallow, const TCHAR* Thing, FToken* OuterPropertyType, EPropertyDeclarationStyle::Type PropertyDeclarationStyle, EVariableCategory::Type VariableCategory);
+	bool GetVarType(
+		FClasses&                       AllClasses,
+		UStruct*                        Scope,
+		FPropertyBase&                  VarProperty,
+		EObjectFlags&                   ObjectFlags,
+		uint64                          Disallow,
+		const TCHAR*                    Thing,
+		FToken*                         OuterPropertyType,
+		EPropertyDeclarationStyle::Type PropertyDeclarationStyle,
+		EVariableCategory::Type         VariableCategory);
 
 	/**
 	 * Parses a variable name declaration and creates a new UProperty object.
@@ -505,11 +482,11 @@ protected:
 
 	const TCHAR* NestTypeName( ENestType NestType );
 
-	UClass* GetQualifiedClass( const TCHAR* Thing );
+	FClass* GetQualifiedClass(const FClasses& AllClasses, const TCHAR* Thing);
 
 	// Nest management functions.
 	void PushNest( ENestType NestType, FName ThisName, UStruct* InNode );
-	void PopNest( ENestType NestType, const TCHAR* Descr );
+	void PopNest ( FClasses& AllClasses, ENestType NestType, const TCHAR* Descr );
 
 	/**
 	 * Binds all delegate properties declared in ValidationScope the delegate functions specified in the variable declaration, verifying that the function is a valid delegate
@@ -517,11 +494,12 @@ protected:
 	 *
 	 * @todo: this function will no longer be required once the post-parse fixup phase is added (TTPRO #13256)
 	 *
+	 * @param	AllClasses			the class tree for CurrentPackage
 	 * @param	ValidationScope		the scope to validate delegate properties for
 	 * @param	OwnerClass			the class currently being compiled.
 	 * @param	DelegateCache		cached map of delegates that have already been found; used for faster lookup.
 	 */
-	void FixupDelegateProperties( UStruct* ValidationScope, UClass* OwnerClass, TMap<FName, UFunction*>& DelegateCache );
+	void FixupDelegateProperties( FClasses& AllClasses, UStruct* ValidationScope, UClass* OwnerClass, TMap<FName, UFunction*>& DelegateCache );
 
 	/**
 	 * Verifies that all specified class's UProperties with CFG_RepNotify have valid callback targets with no parameters nor return values
@@ -534,17 +512,6 @@ protected:
 	void InitScriptLocation( FScriptLocation& Retry );
 	void ReturnToLocation( const FScriptLocation& Retry, bool Binary=1, bool Text=1 );
 
-	/**
-	 * Helper function to determine whether the class hierarchy rooted at Suspect is
-	 * dependent on the hierarchy rooted at Source.
-	 *
-	 * @param	Suspect		Root of hierarchy for suspect class
-	 * @param	Source		Root of hierarchy for source class
-	 * @param	AllClasses	Array of parsed classes
-	 * @return	true if the hierarchy rooted at Suspect is dependent on the one rooted at Source, false otherwise
-	 */
-	bool IsDependentOn( UClass* Suspect, UClass* Source, const FClassTree& AllClasses );
-
 	/**
 	 * If the property has already been seen during compilation, then return add. If not,
 	 * then return replace so that INI files don't mess with header exporting
diff --git a/Engine/Source/Programs/UnrealHeaderTool/Private/NativeClassExporter.h b/Engine/Source/Programs/UnrealHeaderTool/Private/NativeClassExporter.h
--- a/Engine/Source/Programs/UnrealHeaderTool/Private/NativeClassExporter.h
+++ b/Engine/Source/Programs/UnrealHeaderTool/Private/NativeClassExporter.h
@@ -34,7 +34,7 @@ struct FNativeClassHeaderGenerator
 private:
 
 	UClass*				CurrentClass;
-	TArray<UClass*>		Classes;
+	TArray<FClass*>		Classes;
 	FString				ClassHeaderFilename;
 	FString				API;
 	FString			ClassesHeaderPath;
@@ -72,9 +72,6 @@ private:
 	/** Names that have been exported so far **/
 	TSet<FName> ReferencedNames;
 
-	/** the class tree for this package */
-	const FClassTree&	ClassTree;
-
 	/** the existing disk version of the header for this package's names */
 	FString				OriginalNamesHeader;
 	/** the existing disk version of this header */
@@ -149,7 +146,7 @@ private:
 	 * @param	DependencyChain		[out] Receives dependency chain, if found.
 	 * @return	true if a dependency chain was found and filled in.
 	 */
-	bool FindDependencyChain( const UClass* Class, const FString& Header1, const FString& Header2, TArray<const UClass*>& DependencyChain );
+	bool FindDependencyChain( FClass* Class, const FString& Header1, const FString& Header2, TArray<FClass*>& DependencyChain );
 
 	/**
 	 * Finds a dependency chain between two class header files.
@@ -161,7 +158,7 @@ private:
 	 * @param	DependencyChain		[out] Receives dependency chain, if found. Must be empty before the call.
 	 * @return	true if a dependency chain was found and filled in.
 	 */
-	bool FindDependencyChainRecursive( const UClass* Class, const FString& Header1, const FString& Header2, bool bChainStarted, TArray<const UClass*>& DependencyChain );
+	bool FindDependencyChainRecursive( FClass* Class, const FString& Header1, const FString& Header2, bool bChainStarted, TArray<FClass*>& DependencyChain );
 
 	/**
 	 * Determines whether the glue version of the specified native function
@@ -489,5 +486,5 @@ private:
 public:
 
 	// Constructor
-	FNativeClassHeaderGenerator( UPackage* InPackage, FClassTree& inClassTree, bool InAllowSaveExportedHeaders, bool bInUseRelativePaths );
+	FNativeClassHeaderGenerator( UPackage* InPackage, FClasses& AllClasses, bool InAllowSaveExportedHeaders, bool bInUseRelativePaths );
 };
diff --git a/Engine/Source/Programs/UnrealHeaderTool/Private/ParserHelper.h b/Engine/Source/Programs/UnrealHeaderTool/Private/ParserHelper.h
--- a/Engine/Source/Programs/UnrealHeaderTool/Private/ParserHelper.h
+++ b/Engine/Source/Programs/UnrealHeaderTool/Private/ParserHelper.h
@@ -1,7 +1,6 @@
 // Copyright 1998-2014 Epic Games, Inc. All Rights Reserved.
 
 
-#include "ClassTree.h"
 #include "ClassMaps.h"
 
 #pragma once
diff --git a/Engine/Source/Programs/UnrealHeaderTool/Private/StringUtils.cpp b/Engine/Source/Programs/UnrealHeaderTool/Private/StringUtils.cpp
new file mode 100644
--- /dev/null
+++ b/Engine/Source/Programs/UnrealHeaderTool/Private/StringUtils.cpp
@@ -0,0 +1,60 @@
+// Copyright 1998-2014 Epic Games, Inc. All Rights Reserved.
+
+#include "UnrealHeaderTool.h"
+#include "StringUtils.h"
+
+FString GetClassNameWithPrefixRemoved(const FString InClassName)
+{
+	const FString ClassPrefix = GetClassPrefix( InClassName );
+	if( !ClassPrefix.IsEmpty() )
+	{	
+		return InClassName.Right(InClassName.Len() - ClassPrefix.Len());
+	}
+	return FString();
+}
+
+FString GetClassNameWithoutPrefix( const FString& InClassNameOrFilename )
+{
+	// Check for header names (they don't come with a full path so we only search for the first dot)
+	const int32 DotIndex = InClassNameOrFilename.Find(TEXT("."));
+	if (DotIndex == INDEX_NONE)
+	{
+		const FString ClassPrefix = GetClassPrefix( InClassNameOrFilename );
+		return InClassNameOrFilename.Right(InClassNameOrFilename.Len() - ClassPrefix.Len());
+	}
+	else
+	{
+		return InClassNameOrFilename.Mid(0, DotIndex);
+	}
+}
+
+FString GetClassPrefix( const FString InClassName )
+{
+	bool bIsLabledDeprecated;
+	return GetClassPrefix(InClassName, /*out*/ bIsLabledDeprecated);
+}
+
+FString GetClassPrefix( const FString InClassName, bool& bIsLabeledDeprecated )
+{
+	FString ClassPrefix = InClassName.Left(1);
+
+	bIsLabeledDeprecated = false;
+	// First make sure the class name starts with a valid prefix
+	if( ClassPrefix == TEXT("I") ||
+		ClassPrefix == TEXT("A") || 
+		ClassPrefix == TEXT("U") )
+	{
+		// If it is a class prefix, check for deprecated class prefix
+		if( InClassName.Mid(1, 11) == TEXT("DEPRECATED_") )
+		{
+			bIsLabeledDeprecated = true;
+			return InClassName.Left(12);
+		}
+	}
+	// Otherwise check for struct prefix. If it's not a class or struct prefix, it's invalid
+	else if( ClassPrefix != TEXT("F") && ClassPrefix != TEXT("T") )
+	{
+		ClassPrefix = TEXT("");
+	}
+	return ClassPrefix;
+}
diff --git a/Engine/Source/Programs/UnrealHeaderTool/Private/StringUtils.h b/Engine/Source/Programs/UnrealHeaderTool/Private/StringUtils.h
new file mode 100644
--- /dev/null
+++ b/Engine/Source/Programs/UnrealHeaderTool/Private/StringUtils.h
@@ -0,0 +1,36 @@
+// Copyright 1998-2014 Epic Games, Inc. All Rights Reserved.
+
+/*=============================================================================
+	StringUtils.h: Helper functions for manipulating strings.
+=============================================================================*/
+
+#pragma once
+
+/** 
+ * Attempts to strip the given class name of its affixed prefix. If no prefix exists, will return a blank string.
+ *
+ * @param InClassName - Class Name with a prefix to be stripped
+ */
+FString GetClassNameWithPrefixRemoved(const FString InClassName);
+
+/** 
+ * Attempts to strip the given class name of its affixed prefix. If no prefix exists, will return unchanged string.
+ *
+ * @param InClassName - Class Name with a prefix to be stripped
+ */
+FString GetClassNameWithoutPrefix(const FString& InClassNameOrFilename);
+
+/** 
+ * Attempts to get class prefix. If the given class name does not start with a valid Unreal Prefix, it will return an empty string.
+ *
+ * @param InClassName - Name w/ potential prefix to check
+ */
+FString GetClassPrefix(const FString InClassName);
+
+/** 
+ * Attempts to get class prefix. If the given class name does not start with a valid Unreal Prefix, it will return an empty string.
+ *
+ * @param InClassName - Name w/ potential prefix to check
+ * @param out bIsLabeledDeprecated - Reference param set to True if the class name is marked as deprecated
+ */
+FString GetClassPrefix(const FString InClassName, bool& bIsLabeledDeprecated);
diff --git a/Engine/Source/Runtime/CoreUObject/Public/UObject/ClassTree.h b/Engine/Source/Runtime/CoreUObject/Public/UObject/ClassTree.h
--- a/Engine/Source/Runtime/CoreUObject/Public/UObject/ClassTree.h
+++ b/Engine/Source/Runtime/CoreUObject/Public/UObject/ClassTree.h
@@ -16,15 +16,6 @@ enum EClassFlagMatchType
 	MATCH_All,
 };
 
-/**
- * Used by GetChildClasses<> as the comparison class when none is specified.
- */
-class FDefaultComparator
-{
-public:
-	bool IsValidClass( const class FClassTree* ) const { return true; }
-};
-
 /**
  * Manages an inheritance tree.  There is one FClassTree node for each UClass.  Each node 
  * stores pointers to its parent and child nodes.  New nodes should always be added to the root node.
@@ -347,7 +338,7 @@ public:
 		for ( int32 i = 0; i < Children.Num(); i++ )
 		{
 			FClassTree* ChildNode = Children[i];
-			if ( Mask.IsValidClass(ChildNode) )
+			if ( Mask(ChildNode->GetClass()) )
 			{
 				ChildClasses.Add( ChildNode->GetClass() );
 			}
@@ -434,22 +425,19 @@ public:
 	 */
 	const FClassTree* FindNode( UClass* SearchClass ) const
 	{
-		const FClassTree* Result(NULL);
-
 		if ( SearchClass == Class )
-			Result = this;
+			return this;
+
+		if (!SearchClass->IsChildOf(Class))
+			return NULL;
 
-		else if ( SearchClass->IsChildOf(Class) )
+		for (auto Child : Children)
 		{
-			for ( int32 i = 0; i < Children.Num(); i++ )
-			{
-				Result = Children[i]->FindNode(SearchClass);
-				if ( Result != NULL )
-					break;
-			}
+			if (auto Result = Child->FindNode(SearchClass))
+				return Result;
 		}
 
-		return Result;
+		return NULL;
 	}
 
 	/**
@@ -478,42 +466,30 @@ public:
 
 		// find the node associated with SearchClass's new SuperClass
 		FClassTree* NewParentNode = GetNode(NewParentClass);
-		if ( NewParentNode == NULL )
+		if (!NewParentNode)
 		{
 			// if that class hasn't been added yet, add it now.
-			if ( AddClass(NewParentClass) )
-			{
-				NewParentNode = GetNode(NewParentClass);
-			}
-			else
-			{
+			if (!AddClass(NewParentClass))
 				return false;
-			}
+
+			NewParentNode = GetNode(NewParentClass);
 		}
 		check(NewParentNode);
 
 		// find the node for the class that changed SuperClass
-		FClassTree* ClassNode = GetNode(SearchClass,true);
-		if ( ClassNode != NULL )
+		if (FClassTree* ClassNode = GetNode(SearchClass,true))
 		{
 			// if the node has an existing parent, remove it from the parent's Children array
-			if ( ClassNode->Parent != NULL )
+			if (ClassNode->Parent)
 			{
-				int32 ChildIndex = ClassNode->Parent->Children.Find(ClassNode);
-				if ( ChildIndex != INDEX_NONE )
-				{
-					ClassNode->Parent->Children.RemoveAt(ChildIndex);
-				}
+				ClassNode->Parent->Children.RemoveSingle(ClassNode);
 			}
-		}
-		else
-		{
-			return AddClass(SearchClass);
-		}
 
+			// move the node to the new SuperClass
+			return NewParentNode->AddChildNode(ClassNode) != INDEX_NONE;
+		}
 
-		// move the node to the new SuperClass
-		return NewParentNode->AddChildNode(ClassNode) != INDEX_NONE;
+		return AddClass(SearchClass);
 	}
 
 
