Fix for UHT build errors in non-unity.

diff --git a/Engine/Source/Programs/UnrealHeaderTool/Private/CodeGenerator.cpp b/Engine/Source/Programs/UnrealHeaderTool/Private/CodeGenerator.cpp
--- a/Engine/Source/Programs/UnrealHeaderTool/Private/CodeGenerator.cpp
+++ b/Engine/Source/Programs/UnrealHeaderTool/Private/CodeGenerator.cpp
@@ -9,6 +9,7 @@
 #include "HeaderParser.h"
 #include "ClassMaps.h"
 #include "Manifest.h"
+#include "StringUtils.h"
 
 /////////////////////////////////////////////////////
 // Globals
diff --git a/Engine/Source/Programs/UnrealHeaderTool/Private/HeaderParser.cpp b/Engine/Source/Programs/UnrealHeaderTool/Private/HeaderParser.cpp
--- a/Engine/Source/Programs/UnrealHeaderTool/Private/HeaderParser.cpp
+++ b/Engine/Source/Programs/UnrealHeaderTool/Private/HeaderParser.cpp
@@ -7,6 +7,7 @@
 #include "NativeClassExporter.h"
 #include "ClassMaps.h"
 #include "Classes.h"
+#include "StringUtils.h"
 
 #include "UObjectAnnotation.h"
 
diff --git a/Engine/Source/Programs/UnrealHeaderTool/Private/NativeClassExporter.h b/Engine/Source/Programs/UnrealHeaderTool/Private/NativeClassExporter.h
--- a/Engine/Source/Programs/UnrealHeaderTool/Private/NativeClassExporter.h
+++ b/Engine/Source/Programs/UnrealHeaderTool/Private/NativeClassExporter.h
@@ -6,6 +6,9 @@
 
 #pragma once
 
+#include "Class.h"
+#include "Classes.h"
+
 //
 //	FNativeClassHeaderGenerator
 //
