#ttp 332139 - UHT: Unclear error message when a function param is a pointer to a struct
#proj UE4
#branch UE4
#summary Fixed the error message.

diff --git a/Engine/Source/Programs/UnrealHeaderTool/Private/HeaderParser.cpp b/Engine/Source/Programs/UnrealHeaderTool/Private/HeaderParser.cpp
--- a/Engine/Source/Programs/UnrealHeaderTool/Private/HeaderParser.cpp
+++ b/Engine/Source/Programs/UnrealHeaderTool/Private/HeaderParser.cpp
@@ -3134,7 +3134,7 @@ bool FHeaderParser::GetVarType
 
 	if (MatchSymbol(TEXT("*")))
 	{
-		FError::Throwf(TEXT("Inappropriate '*' on variable of type '%s', cannot have an exposed pointer to this type (move it inside the CPP block)."), VarType.Identifier );
+		FError::Throwf(TEXT("Inappropriate '*' on variable of type '%s', cannot have an exposed pointer to this type."), VarType.Identifier );
 	}
 
 	//@TODO: UCREMOVAL: 'const' member variables that will get written post-construction by defaultproperties
