#UE4 - Added a starting point for UAttributeProperty.  This is a starting point to get it working for UMG but it isnt ready for runtime yet.

diff --git a/Engine/Source/Editor/PropertyEditor/Private/ItemPropertyNode.cpp b/Engine/Source/Editor/PropertyEditor/Private/ItemPropertyNode.cpp
--- a/Engine/Source/Editor/PropertyEditor/Private/ItemPropertyNode.cpp
+++ b/Engine/Source/Editor/PropertyEditor/Private/ItemPropertyNode.cpp
@@ -96,7 +96,8 @@ void FItemPropertyNode::InitExpansionFlags (void)
 	if(	Cast<UStructProperty>(MyProperty)
 		||	( Cast<UArrayProperty>(MyProperty) && GetReadAddress(false,Addresses) )
 		||  HasNodeFlags(EPropertyNodeFlags::EditInline)
-		||	( Property->ArrayDim > 1 && ArrayIndex == -1 ) )
+		||	( Property->ArrayDim > 1 && ArrayIndex == -1 ) 
+		|| Cast<UAttributeProperty>(MyProperty) )
 	{
 		SetNodeFlags(EPropertyNodeFlags::CanBeExpanded, true);
 	}
@@ -111,6 +112,7 @@ void FItemPropertyNode::InitChildNodes()
 	UStructProperty* StructProperty = Cast<UStructProperty>(Property);
 	UArrayProperty* ArrayProperty = Cast<UArrayProperty>(Property);
 	UObjectPropertyBase* ObjectProperty = Cast<UObjectPropertyBase>(Property);
+	UAttributeProperty* AttributeProperty = Cast<UAttributeProperty>(Property);
 
 	bool bShouldShowHiddenProperties = !!HasNodeFlags(EPropertyNodeFlags::ShouldShowHiddenProperties);
 
@@ -267,6 +269,21 @@ void FItemPropertyNode::InitChildNodes()
 			}
 		}
 	}
+	else if( AttributeProperty )
+	{
+		TSharedPtr<FItemPropertyNode> NewItemNode(new FItemPropertyNode);
+
+		FPropertyNodeInitParams InitParams;
+		InitParams.ParentNode = SharedThis(this);
+		InitParams.Property = AttributeProperty->Inner;
+		InitParams.ArrayOffset = 0;
+		InitParams.ArrayIndex = INDEX_NONE;
+		InitParams.bAllowChildren = true;
+		InitParams.bForceHiddenPropertyVisibility = bShouldShowHiddenProperties;
+
+		NewItemNode->InitNode(InitParams);
+		AddChildNode(NewItemNode);
+	}
 
 	//needs to be after all the children are created
 	if ( FPropertySettings::Get().ExpandDistributions() == true)
@@ -300,7 +317,11 @@ FString FItemPropertyNode::GetDisplayName() const
 	else 
 	{
 		const UProperty* PropertyPtr = GetProperty();
-		if( GetArrayIndex()==-1 && PropertyPtr != NULL  )
+		if (GetParentNode() && GetParentNode()->GetProperty() && GetParentNode()->GetProperty()->IsA<UAttributeProperty>())
+		{
+			FinalDisplayName = NSLOCTEXT("PropertyEditor", "AttributeValue", "Value").ToString();
+		}
+		else if( GetArrayIndex()==-1 && PropertyPtr != NULL  )
 		{
 			// This item is not a member of an array, get a traditional display name
 			if ( FPropertySettings::Get().ShowFriendlyPropertyNames() )
diff --git a/Engine/Source/Editor/PropertyEditor/Private/PropertyEditorHelpers.cpp b/Engine/Source/Editor/PropertyEditor/Private/PropertyEditorHelpers.cpp
--- a/Engine/Source/Editor/PropertyEditor/Private/PropertyEditorHelpers.cpp
+++ b/Engine/Source/Editor/PropertyEditor/Private/PropertyEditorHelpers.cpp
@@ -11,6 +11,7 @@
 #include "SPropertyEditor.h"
 #include "SPropertyEditorNumeric.h"
 #include "SPropertyEditorArray.h"
+#include "SPropertyEditorAttribute.h"
 #include "SPropertyEditorCombo.h"
 #include "SPropertyEditorEditInline.h"
 #include "SPropertyEditorText.h"
@@ -144,6 +145,14 @@ TSharedRef<SWidget> SPropertyValueWidget::ConstructPropertyEditorWidget( TShared
 
 			ArrayWidget->GetDesiredWidth( MinDesiredWidth, MaxDesiredWidth );
 		}
+		else if( SPropertyEditorAttribute::Supports(PropertyEditorRef) )
+		{
+			TSharedRef<SPropertyEditorAttribute> AttributeWidget =
+				SAssignNew(PropertyWidget, SPropertyEditorAttribute, PropertyEditorRef)
+				.Font(FontStyle);
+
+			AttributeWidget->GetDesiredWidth(MinDesiredWidth, MaxDesiredWidth);
+		}
 		else if ( SPropertyEditorAsset::Supports( PropertyEditorRef ) )
 		{
 			TSharedRef<SPropertyEditorAsset> AssetWidget = 
diff --git a/Engine/Source/Editor/PropertyEditor/Private/SDetailsView.cpp b/Engine/Source/Editor/PropertyEditor/Private/SDetailsView.cpp
--- a/Engine/Source/Editor/PropertyEditor/Private/SDetailsView.cpp
+++ b/Engine/Source/Editor/PropertyEditor/Private/SDetailsView.cpp
@@ -454,12 +454,17 @@ void SDetailsView::UpdatePropertyMapRecursive( FPropertyNode& InNode, FDetailLay
 				// Is this a property of an array
 				bool bIsChildOfArray = PropertyEditorHelpers::IsChildOfArray( ChildNode );
 
+				bool bIsChildOfAttribute = ParentProperty ? ParentProperty->IsA<UAttributeProperty>() : false;
+
 				// Edit inline new properties should be visible by default
 				bVisibleByDefault |= bIsEditInlineNew;
 
-				// Children of arrays are not visible directly,
+				// Children of arrays are not visible directly
 				bVisibleByDefault &= !bIsChildOfArray;
 
+				// Children of attributes are not visible directly
+				bVisibleByDefault &= !bIsChildOfAttribute;
+
 				// See if the user requested specific property visibility 
 				const bool bIsUserVisible = IsPropertyVisible.IsBound() ? IsPropertyVisible.Execute( Property ) : true;
 
@@ -482,7 +487,7 @@ void SDetailsView::UpdatePropertyMapRecursive( FPropertyNode& InNode, FDetailLay
 				}
 
 				// Do not add children of customized in struct properties or arrays
-				if( !bIsChildOfCustomizedStruct && !bIsChildOfArray )
+				if( !bIsChildOfCustomizedStruct && !bIsChildOfArray && !bIsChildOfAttribute )
 				{
 					// Get the class property map
 					FClassInstanceToPropertyMap& ClassInstanceMap = ClassToPropertyMap.FindOrAdd( Property->GetOwnerStruct()->GetFName() );
diff --git a/Engine/Source/Programs/UnrealHeaderTool/Private/CodeGenerator.cpp b/Engine/Source/Programs/UnrealHeaderTool/Private/CodeGenerator.cpp
--- a/Engine/Source/Programs/UnrealHeaderTool/Private/CodeGenerator.cpp
+++ b/Engine/Source/Programs/UnrealHeaderTool/Private/CodeGenerator.cpp
@@ -713,6 +713,10 @@ FString FNativeClassHeaderGenerator::PropertyNew(FString& Meta, UProperty* Prop,
 		UFunction *TargetFunction = MulticastDelegateProperty->SignatureFunction;
 		ExtraArgs = FString::Printf(TEXT(", %s"), *GetSingletonName(TargetFunction));
 	}
+	else if( UAttributeProperty* AttributeProperty = Cast<UAttributeProperty>(Prop) )
+	{
+		ExtraArgs = FString::Printf(TEXT(", sizeof(TAttribute<%s>)"), *AttributeProperty->Inner->GetCPPType() );
+	}
 
 	FString Constructor = FString::Printf(TEXT("new(%s, TEXT(\"%s\"), RF_Public|RF_Transient|RF_Native) U%s(%s, 0x%016llx%s);"),
 		*OuterString,
@@ -820,12 +824,19 @@ void FNativeClassHeaderGenerator::OutputProperty(FString& Meta, FOutputDevice& O
 		OutputDevice.Log(*PropertyNew(Meta, Prop, OuterString, PropMacroOuterClass, TEXT(""), Spaces, *SourceStruct));
 	}
 	UArrayProperty* ArrayProperty = Cast<UArrayProperty>(Prop);
+	UAttributeProperty* AttributeProperty = Cast<UAttributeProperty>(Prop);
 	if (ArrayProperty)
 	{
 		FString InnerOuterString = FString::Printf(TEXT("NewProp_%s"), *Prop->GetName());
 		FString PropMacroOuterArray = TEXT("FPostConstructInitializeProperties(), EC_CppProperty, 0");
 		OutputDevice.Log(*PropertyNew(Meta, ArrayProperty->Inner, InnerOuterString, PropMacroOuterArray, TEXT("_Inner"), Spaces));
 	}
+	else if(AttributeProperty)
+	{
+		FString InnerOuterString = FString::Printf(TEXT("NewProp_%s"), *Prop->GetName());
+		FString PropMacroOuterAttrib = TEXT("FPostConstructInitializeProperties(), EC_CppProperty, 0");
+		OutputDevice.Log(*PropertyNew(Meta, AttributeProperty->Inner, InnerOuterString, PropMacroOuterAttrib, TEXT("_Inner"), Spaces));
+	}
 }
 
 static FString BackSlashAndIndent(const FString& Text)
@@ -878,10 +889,16 @@ static void FindNoExportStructs(TArray<UScriptStruct*>& Structs, UStruct* Start)
 		{
 			UStructProperty* InnerStruct = Cast<UStructProperty>(*ItInner);
 			UArrayProperty* InnerArray = Cast<UArrayProperty>(*ItInner);
+			UAttributeProperty* InnerAttribute = Cast<UAttributeProperty>(*ItInner);
 			if (InnerArray)
 			{
 				InnerStruct = Cast<UStructProperty>(InnerArray->Inner);
 			}
+			else if(InnerAttribute)
+			{
+				InnerStruct = Cast<UStructProperty>(InnerAttribute->Inner);
+			}
+
 			if (InnerStruct)
 			{
 				FindNoExportStructs(Structs, InnerStruct->Struct);
@@ -2616,16 +2633,19 @@ void ExportProtoDeclaration(FOutputDevice& Out, const FString& MessageName, TFie
 		{
 			bool bIsRepeated = false;
 			UArrayProperty* ArrayProperty = Cast<UArrayProperty>(Property);
-			if (ArrayProperty != NULL)
+			UAttributeProperty* AttributeProperty = Cast<UAttributeProperty>(Property);
+			if (ArrayProperty != NULL || AttributeProperty != NULL)
 			{
-				UClass* InnerPropClass = ArrayProperty->Inner->GetClass();
+				UProperty* Inner = ArrayProperty ? ArrayProperty->Inner : AttributeProperty->Inner;
+
+				UClass* InnerPropClass = Inner->GetClass();
 				if (InnerPropClass != UInterfaceProperty::StaticClass() && InnerPropClass != UObjectProperty::StaticClass())
 				{
 					FString InnerExtendedTypeText;
-					FString InnerTypeText = ArrayProperty->Inner->GetCPPType(&InnerExtendedTypeText, CPPF_None);
+					FString InnerTypeText = Inner->GetCPPType(&InnerExtendedTypeText, CPPF_None);
 					TypeText = InnerTypeText;
 					ExtendedTypeText = InnerExtendedTypeText;
-					Property = ArrayProperty->Inner;
+					Property = Inner;
 					bIsRepeated = true;
 				}
 				else
@@ -3062,7 +3082,8 @@ void FNativeClassHeaderGenerator::ExportEventParms( const TArray<UFunction*>& In
 				Cast<UStrProperty>(Prop) ||
 				Cast<UTextProperty>(Prop) ||
 				Cast<UArrayProperty>(Prop) ||
-				Cast<UInterfaceProperty>(Prop)
+				Cast<UInterfaceProperty>(Prop) ||
+				Cast<UAttributeProperty>(Prop)
 				)
 			{
 				bNeedsOutput = false;
@@ -3134,7 +3155,8 @@ FString FNativeClassHeaderGenerator::GetNullParameterValue( UProperty* Prop, boo
 	}
 	else if ( PropClass == UArrayProperty::StaticClass()
 		||    PropClass == UDelegateProperty::StaticClass()
-		||    PropClass == UMulticastDelegateProperty::StaticClass() )
+		||    PropClass == UMulticastDelegateProperty::StaticClass()
+		||    PropClass == UAttributeProperty::StaticClass() )
 	{
 		FString Type, ExtendedType;
 		Type = Prop->GetCPPType(&ExtendedType,CPPF_OptionalValue);
@@ -3663,10 +3685,15 @@ void FNativeClassHeaderGenerator::ApplyAlternatePropertyExportText( UProperty* P
 	if ( ObjectProp == NULL )
 	{
 		UArrayProperty* ArrayProp = Cast<UArrayProperty>(Prop);
+		UAttributeProperty* AttributeProp = Cast<UAttributeProperty>(Prop);
 		if ( ArrayProp != NULL )
 		{
 			ObjectProp = Cast<UObjectProperty>(ArrayProp->Inner);
 		}
+		else if( AttributeProp != NULL )
+		{
+			ObjectProp = Cast<UObjectProperty>(AttributeProp->Inner);
+		}
 	}
 	if ( ObjectProp != NULL && PropData && PropData->Token.ExportInfo.Len() > 0)
 	{
@@ -3678,10 +3705,15 @@ void FNativeClassHeaderGenerator::ApplyAlternatePropertyExportText( UProperty* P
 	if ( StructProp == NULL )
 	{
 		UArrayProperty* ArrayProp = Cast<UArrayProperty>(Prop);
+		UAttributeProperty* AttributeProp = Cast<UAttributeProperty>(Prop);
 		if ( ArrayProp != NULL )
 		{
 			StructProp = Cast<UStructProperty>(ArrayProp->Inner);
 		}
+		else if (AttributeProp != NULL)
+		{
+			StructProp = Cast<UStructProperty>(AttributeProp->Inner);
+		}
 	}
 
 	if ( StructProp )
diff --git a/Engine/Source/Programs/UnrealHeaderTool/Private/HeaderParser.cpp b/Engine/Source/Programs/UnrealHeaderTool/Private/HeaderParser.cpp
--- a/Engine/Source/Programs/UnrealHeaderTool/Private/HeaderParser.cpp
+++ b/Engine/Source/Programs/UnrealHeaderTool/Private/HeaderParser.cpp
@@ -2767,6 +2767,33 @@ bool FHeaderParser::GetVarType
 		VarProperty.ArrayType = EArrayType::Dynamic;
 		RequireSymbol( TEXT(">"), TEXT("'tarray'") );
 	}
+	else if( VarType.Matches(TEXT("TAttribute")) )
+	{
+		RequireSymbol( TEXT("<"), TEXT("'tattribute'") );
+
+		// GetVarType() clears the property flags of the array var, so use dummy 
+		// flags when getting the inner property
+		EObjectFlags InnerFlags;
+		uint64 OriginalVarTypeFlags = VarType.PropertyFlags;
+		VarType.PropertyFlags |= Flags;
+
+		GetVarType(AllClasses, Scope, VarProperty, InnerFlags, Disallow, TEXT("'tattribute'"), &VarType, EPropertyDeclarationStyle::None, VariableCategory );
+		if( VarProperty.Type == CPT_None )
+		{
+			FError::Throwf( TEXT("Unsupported attribute type.") );
+		}
+		else if( VarProperty.bIsAttribute == true )
+		{
+			FError::Throwf( TEXT("Attributes within attributes are not supported") );
+		}
+
+		VarProperty.bIsAttribute = true;
+		OriginalVarTypeFlags |= VarProperty.PropertyFlags & (CPF_ContainsInstancedReference | CPF_InstancedReference);
+		VarType.PropertyFlags = OriginalVarTypeFlags;
+	
+		RequireSymbol(TEXT(">"), TEXT("'tattribute'"));
+
+	}
 	else if ( VarType.Matches(TEXT("FString")) )
 	{
 		VarProperty = FPropertyBase(CPT_String);
@@ -3454,6 +3481,7 @@ UProperty* FHeaderParser::GetVarNameAndDim
 			Prev = *It;
 		}
 
+		UProperty* Attribute = NULL;
 		UProperty* Array    = NULL;
 		UObject*   NewScope = Scope;
 		int32      ArrayDim = 1; // 1 = not a static array, 2 = static array
@@ -3467,6 +3495,12 @@ UProperty* FHeaderParser::GetVarNameAndDim
 		{
 			ArrayDim = 2;
 		}
+		else if( VarProperty.bIsAttribute == true )
+		{
+			Attribute = new(Scope,PropertyName,ObjectFlags)UAttributeProperty(FPostConstructInitializeProperties());	
+			NewScope = Attribute;
+			ObjectFlags = RF_Public;
+		}
 
 		if (VarProperty.Type == CPT_Byte)
 		{
@@ -3643,6 +3677,24 @@ UProperty* FHeaderParser::GetVarNameAndDim
 			}
 			NewProperty = Array;
 		}
+		else if( Attribute )
+		{
+			check(NewProperty);
+
+			UAttributeProperty* AttributeProp = CastChecked<UAttributeProperty>(Attribute);
+			AttributeProp->Inner = NewProperty;
+	
+			// Copy some of the property flags to the inner property.
+			NewProperty->PropertyFlags |= (VarProperty.PropertyFlags&CPF_PropagateToArrayInner);
+
+			if (NewProperty->PropertyFlags & (CPF_ContainsInstancedReference | CPF_InstancedReference))
+			{
+				VarProperty.PropertyFlags |= CPF_ContainsInstancedReference;
+				VarProperty.PropertyFlags &= ~CPF_InstancedReference; //this was propagated to the inner
+			}
+			NewProperty = Attribute;
+		}
+
 		NewProperty->ArrayDim = ArrayDim;
 		if (ArrayDim == 2)
 		{
diff --git a/Engine/Source/Programs/UnrealHeaderTool/Private/ParserHelper.h b/Engine/Source/Programs/UnrealHeaderTool/Private/ParserHelper.h
--- a/Engine/Source/Programs/UnrealHeaderTool/Private/ParserHelper.h
+++ b/Engine/Source/Programs/UnrealHeaderTool/Private/ParserHelper.h
@@ -125,6 +125,8 @@ public:
 
 	EPointerType::Type PointerType;
 
+	/** Is the property a TAttribute */
+	bool bIsAttribute;
 public:
 	/** @name Constructors */
 	//@{
@@ -140,6 +142,7 @@ public:
 	, RepNotifyName      (NAME_None)
 	, ReferenceType      (InRefType)
 	, PointerType        (EPointerType::None)
+	, bIsAttribute		 (false)
 	{
 	}
 
@@ -155,6 +158,7 @@ public:
 	, RepNotifyName      (NAME_None)
 	, ReferenceType      (InRefType)
 	, PointerType        (EPointerType::None)
+	, bIsAttribute		 (false)
 	{
 	}
 
@@ -170,6 +174,7 @@ public:
 	, RepNotifyName      (NAME_None)
 	, ReferenceType      (InRefType)
 	, PointerType        (EPointerType::None)
+	, bIsAttribute		 (false)
 	{
 		// if this is an interface class, we use the UInterfaceProperty class instead of UObjectProperty
 		if ( InClass->HasAnyClassFlags(CLASS_Interface) )
@@ -228,6 +233,7 @@ public:
 	, RepNotifyName      (NAME_None)
 	, ReferenceType      (InRefType)
 	, PointerType        (EPointerType::None)
+	, bIsAttribute		 (false)
 	{
 	}
 
@@ -235,6 +241,7 @@ public:
 	: PropertyExportFlags(PROPEXPORT_Public)
 	, DelegateName       (NAME_None)
 	, RepNotifyName      (NAME_None)
+	, bIsAttribute		 (false)
 	{
 		checkSlow(Property);
 
@@ -252,6 +259,12 @@ public:
 			Property = CastChecked<UArrayProperty>(Property)->Inner;
 			ClassOfProperty = Property->GetClass();
 		}
+		else if( ClassOfProperty == UAttributeProperty::StaticClass() )
+		{
+			Property = CastChecked<UAttributeProperty>(Property)->Inner;
+			ClassOfProperty = Property->GetClass();
+			bIsAttribute = true;
+		}
 
 		if( ClassOfProperty==UByteProperty::StaticClass() )
 		{
diff --git a/Engine/Source/Runtime/Core/Public/Core.h b/Engine/Source/Runtime/Core/Public/Core.h
--- a/Engine/Source/Runtime/Core/Public/Core.h
+++ b/Engine/Source/Runtime/Core/Public/Core.h
@@ -73,6 +73,7 @@ class               UTextProperty;
 class               UArrayProperty;
 class				UDelegateProperty;
 class				UMulticastDelegateProperty;
+class				UAttributeProperty;
 class			UStruct;
 class				UFunction;
 class				UClass;
diff --git a/Engine/Source/Runtime/Core/Public/UObject/UnrealNames.h b/Engine/Source/Runtime/Core/Public/UObject/UnrealNames.h
--- a/Engine/Source/Runtime/Core/Public/UObject/UnrealNames.h
+++ b/Engine/Source/Runtime/Core/Public/UObject/UnrealNames.h
@@ -57,6 +57,7 @@ REGISTER_NAME(23,Int64Property)
 REGISTER_NAME(25,Int16Property)
 REGISTER_NAME(26,Int8Property)
 REGISTER_NAME(27,AssetSubclassOfProperty)
+REGISTER_NAME(28,AttributeProperty)
 
 // Special packages.
 REGISTER_NAME(30,Core)
diff --git a/Engine/Source/Runtime/CoreUObject/Private/UObject/Class.cpp b/Engine/Source/Runtime/CoreUObject/Private/UObject/Class.cpp
--- a/Engine/Source/Runtime/CoreUObject/Private/UObject/Class.cpp
+++ b/Engine/Source/Runtime/CoreUObject/Private/UObject/Class.cpp
@@ -958,6 +958,10 @@ void UStruct::SerializeTaggedProperties(FArchive& Ar, uint8* Data, UStruct* Defa
 			{
 				UE_LOG(LogClass, Warning, TEXT("Array Inner Type mismatch in %s of %s - Previous (%s) Current(%s) for package:  %s"), *Tag.Name.ToString(), *GetName(), *Tag.InnerType.ToString(), *CastChecked<UArrayProperty>(Property)->Inner->GetID().ToString(), *Ar.GetArchiveName() );
 			}
+			else if (Tag.Type == NAME_AttributeProperty && Tag.InnerType != NAME_None && Tag.InnerType != CastChecked<UAttributeProperty>(Property)->Inner->GetID())
+			{
+				UE_LOG(LogClass, Warning, TEXT("Attribute Inner Type mismatch in %s of %s - Previous (%s) Current(%s) for package:  %s"), *Tag.Name.ToString(), *GetName(), *Tag.InnerType.ToString(), *CastChecked<UAttributeProperty>(Property)->Inner->GetID().ToString(), *Ar.GetArchiveName());
+			}
 			else if( Tag.Type==NAME_StructProperty && Tag.StructName!=CastChecked<UStructProperty>(Property)->Struct->GetFName() 
 				&& CastChecked<UStructProperty>(Property)->UseBinaryOrNativeSerialization(Ar) )
 			{
diff --git a/Engine/Source/Runtime/CoreUObject/Private/UObject/PropertyAttribute.cpp b/Engine/Source/Runtime/CoreUObject/Private/UObject/PropertyAttribute.cpp
new file mode 100644
--- /dev/null
+++ b/Engine/Source/Runtime/CoreUObject/Private/UObject/PropertyAttribute.cpp
@@ -0,0 +1,167 @@
+// Copyright 1998-2014 Epic Games, Inc. All Rights Reserved.
+
+#include "CoreUObjectPrivate.h"
+#include "PropertyHelper.h"
+
+void UAttributeProperty::LinkInternal(FArchive& Ar)
+{
+	ULinkerLoad* MyLinker = GetLinker();
+	if( MyLinker )
+	{
+		MyLinker->Preload(this);
+	}
+	Ar.Preload(Inner);
+	Inner->Link(Ar);
+}
+
+bool UAttributeProperty::Identical( const void* A, const void* B, uint32 PortFlags ) const
+{
+	checkSlow(Inner);
+
+	return Inner->Identical( A, B, PortFlags );
+}
+
+void UAttributeProperty::SerializeItem( FArchive& Ar, void* Value, int32 MaxReadBytes, void const* Defaults ) const
+{
+	checkSlow(Inner);
+
+	// Ensure that the Inner itself has been loaded before calling SerializeItem() on it
+	Ar.Preload(Inner);
+
+	Inner->SerializeItem( Ar, Value, MaxReadBytes, Defaults);
+}
+
+bool UAttributeProperty::NetSerializeItem( FArchive& Ar, UPackageMap* Map, void* Data, TArray<uint8> * MetaData ) const
+{
+	UE_LOG( LogProperty, Fatal, TEXT( "Deprecated code path" ) );
+	return 1;
+}
+
+void UAttributeProperty::Serialize( FArchive& Ar )
+{
+	Super::Serialize( Ar );
+	Ar << Inner;
+	checkSlow(Inner||HasAnyFlags(RF_ClassDefaultObject));
+}
+
+void UAttributeProperty::AddReferencedObjects(UObject* InThis, FReferenceCollector& Collector)
+{
+	UAttributeProperty* This = CastChecked<UAttributeProperty>(InThis);
+	Collector.AddReferencedObject( This->Inner, This );
+	Super::AddReferencedObjects( This, Collector );
+}
+
+FString UAttributeProperty::GetCPPType( FString* ExtendedTypeText/*=NULL*/, uint32 CPPExportFlags/*=0*/ ) const
+{
+	checkSlow(Inner);
+
+	if ( ExtendedTypeText != NULL )
+	{
+		FString InnerExtendedTypeText;
+		FString InnerTypeText = Inner->GetCPPType(&InnerExtendedTypeText, CPPExportFlags & ~CPPF_ArgumentOrReturnValue); // we won't consider array inners to be "arguments or return values"
+		if ( InnerExtendedTypeText.Len() && InnerExtendedTypeText.Right(1) == TEXT(">") )
+		{
+			// if our internal property type is a template class, add a space between the closing brackets b/c VS.NET cannot parse this correctly
+			InnerExtendedTypeText += TEXT(" ");
+		}
+		else if ( !InnerExtendedTypeText.Len() && InnerTypeText.Len() && InnerTypeText.Right(1) == TEXT(">") )
+		{
+			// if our internal property type is a template class, add a space between the closing brackets b/c VS.NET cannot parse this correctly
+			InnerExtendedTypeText += TEXT(" ");
+		}
+		*ExtendedTypeText = FString::Printf(TEXT("<%s%s>"), *InnerTypeText, *InnerExtendedTypeText);
+	}
+	return TEXT("TAttribute");
+}
+
+FString UAttributeProperty::GetCPPMacroType( FString& ExtendedTypeText ) const
+{
+	checkSlow(Inner);
+	ExtendedTypeText = Inner->GetCPPType();
+	return TEXT("TATTRIBUTE");
+}
+
+void UAttributeProperty::ExportTextItem( FString& ValueStr, const void* PropertyValue, const void* DefaultValue, UObject* Parent, int32 PortFlags, UObject* ExportRootScope ) const
+{
+	checkSlow(Inner);
+	Inner->ExportTextItem( ValueStr, PropertyValue, DefaultValue, Parent, PortFlags, ExportRootScope );
+}
+
+const TCHAR* UAttributeProperty::ImportText_Internal( const TCHAR* Buffer, void* Data, int32 PortFlags, UObject* Parent, FOutputDevice* ErrorText ) const
+{
+	checkSlow(Inner);
+	Inner->ImportText( Buffer, Data, PortFlags, Parent, ErrorText );
+
+	return Buffer;
+}
+
+void UAttributeProperty::CopyValuesInternal(void* Dest, void const* Src, int32 Count) const
+{
+	checkSlow(Inner);
+	Inner->CopyCompleteValue(Dest, Src);
+}
+
+void UAttributeProperty::ClearValueInternal(void* Data) const 
+{
+	checkSlow(Inner);
+	Inner->ClearValue(Data);
+}
+
+void UAttributeProperty::DestroyValueInternal(void* Dest) const
+{
+	checkSlow(Inner);
+	Inner->DestroyValue(Dest);
+}
+
+void UAttributeProperty::AddCppProperty( UProperty* Property )
+{
+	check(!Inner);
+	check(Property);
+
+	Inner = Property;
+}
+
+bool UAttributeProperty::IsLocalized() const
+{
+	return Inner->IsLocalized() ? true : Super::IsLocalized();
+}
+
+bool UAttributeProperty::PassCPPArgsByRef() const
+{
+	return true;
+}
+
+void UAttributeProperty::InstanceSubobjects( void* Data, void const* DefaultData, UObject* Owner, FObjectInstancingGraph* InstanceGraph )
+{
+	check(Inner)
+	Inner->InstanceSubobjects( Data, DefaultData, Owner, InstanceGraph );
+}
+
+bool UAttributeProperty::ContainsObjectReference() const
+{
+	check(Inner);
+	return Inner->ContainsObjectReference();
+}
+
+bool UAttributeProperty::ContainsWeakObjectReference() const
+{
+	check(Inner);
+	return Inner->ContainsWeakObjectReference();
+}
+
+void UAttributeProperty::EmitReferenceInfo(FGCReferenceTokenStream* ReferenceTokenStream, int32 BaseOffset)
+{
+	check(Inner);
+	Inner->EmitReferenceInfo(ReferenceTokenStream, BaseOffset);
+}
+
+bool UAttributeProperty::SameType(const UProperty* Other) const
+{
+	return Super::SameType(Other) && Inner && Inner->SameType(((UAttributeProperty*)Other)->Inner);
+}
+
+IMPLEMENT_CORE_INTRINSIC_CLASS(UAttributeProperty, UProperty,
+	{
+		Class->EmitObjectReference( STRUCT_OFFSET( UAttributeProperty, Inner ) );
+	}
+);
diff --git a/Engine/Source/Runtime/CoreUObject/Private/UObject/SavePackage.cpp b/Engine/Source/Runtime/CoreUObject/Private/UObject/SavePackage.cpp
--- a/Engine/Source/Runtime/CoreUObject/Private/UObject/SavePackage.cpp
+++ b/Engine/Source/Runtime/CoreUObject/Private/UObject/SavePackage.cpp
@@ -1450,7 +1450,8 @@ class FExportReferenceSorter : public FArchiveUObject
 			UObjectPropertyBase::StaticClass(),
 			ULazyObjectProperty::StaticClass(),
 			UAssetObjectProperty::StaticClass(),
-			UAssetClassProperty::StaticClass()
+			UAssetClassProperty::StaticClass(),
+			UAttributeProperty::StaticClass()
 		};
 
 		for ( int32 CoreClassIndex = 0; CoreClassIndex < ARRAY_COUNT(CoreClassList); CoreClassIndex++ )
diff --git a/Engine/Source/Runtime/CoreUObject/Public/Templates/Casts.h b/Engine/Source/Runtime/CoreUObject/Public/Templates/Casts.h
--- a/Engine/Source/Runtime/CoreUObject/Public/Templates/Casts.h
+++ b/Engine/Source/Runtime/CoreUObject/Public/Templates/Casts.h
@@ -160,3 +160,4 @@ DECLARE_CAST_BY_FLAG(UTextProperty)
 DECLARE_CAST_BY_FLAG(UArrayProperty)
 DECLARE_CAST_BY_FLAG(UDelegateProperty)
 DECLARE_CAST_BY_FLAG(UMulticastDelegateProperty)
+DECLARE_CAST_BY_FLAG(UAttributeProperty)
diff --git a/Engine/Source/Runtime/CoreUObject/Public/UObject/ObjectBase.h b/Engine/Source/Runtime/CoreUObject/Public/UObject/ObjectBase.h
--- a/Engine/Source/Runtime/CoreUObject/Public/UObject/ObjectBase.h
+++ b/Engine/Source/Runtime/CoreUObject/Public/UObject/ObjectBase.h
@@ -269,6 +269,7 @@ typedef uint64 EClassCastFlags;
 #define CASTCLASS_UInt16Property				DECLARE_UINT64(0x0000000080000000)
 #define CASTCLASS_UDoubleProperty				DECLARE_UINT64(0x0000000100000000)
 #define CASTCLASS_UAssetClassProperty			DECLARE_UINT64(0x0000000200000000)
+#define CASTCLASS_UAttributeProperty			DECLARE_UINT64(0x0000000400000000)
 
 #define CASTCLASS_AllFlags						DECLARE_UINT64(0xFFFFFFFFFFFFFFFF)
 
diff --git a/Engine/Source/Runtime/CoreUObject/Public/UObject/PropertyTag.h b/Engine/Source/Runtime/CoreUObject/Public/UObject/PropertyTag.h
--- a/Engine/Source/Runtime/CoreUObject/Public/UObject/PropertyTag.h
+++ b/Engine/Source/Runtime/CoreUObject/Public/UObject/PropertyTag.h
@@ -18,7 +18,7 @@ struct FPropertyTag
 	FName	Name;		// Name of property.
 	FName	StructName;	// Struct name if UStructProperty.
 	FName	EnumName;	// Enum name if UByteProperty
-	FName InnerType; // Inner type if UArrayProperty
+	FName InnerType; // Inner type if UArrayProperty or UAttributeProperty
 	int32		Size;       // Property size.
 	int32		ArrayIndex;	// Index if an array; else 0.
 	int32		SizeOffset;	// location in stream of tag size member
@@ -53,11 +53,16 @@ struct FPropertyTag
 					EnumName = ByteProp->Enum->GetFName();
 				}
 			}
-			else if (Property->IsA(UArrayProperty::StaticClass()))
+			else if (Property->IsA(UArrayProperty::StaticClass()) )
 			{
 				UArrayProperty* ArrayProp = CastChecked<UArrayProperty>(Property);
 				InnerType = ArrayProp->Inner->GetID();
 			}
+			else if( Property->IsA(UAttributeProperty::StaticClass() ) )
+			{
+				UAttributeProperty* AttributeProp = CastChecked<UAttributeProperty>(Property);
+				InnerType = AttributeProp->Inner->GetID();
+			}
 			else if (Property->IsA(UBoolProperty::StaticClass()))
 			{
 				UBoolProperty* Bool = CastChecked<UBoolProperty>(Property);
@@ -101,7 +106,7 @@ struct FPropertyTag
 			Ar << Tag.EnumName;
 		}
 		// only need to serialize this for arrays
-		else if (Tag.Type == NAME_ArrayProperty)
+		else if (Tag.Type == NAME_ArrayProperty || Tag.Type == NAME_AttributeProperty)
 		{
 			if (Ar.UE4Ver() >= VAR_UE4_ARRAY_PROPERTY_INNER_TAGS)
 			{
diff --git a/Engine/Source/Runtime/CoreUObject/Public/UObject/UnrealType.h b/Engine/Source/Runtime/CoreUObject/Public/UObject/UnrealType.h
--- a/Engine/Source/Runtime/CoreUObject/Public/UObject/UnrealType.h
+++ b/Engine/Source/Runtime/CoreUObject/Public/UObject/UnrealType.h
@@ -2443,6 +2443,57 @@ public:
 	}
 };
 
+/*-----------------------------------------------------------------------------
+	UAttributeProperty.
+-----------------------------------------------------------------------------*/
+
+//
+// Describes a TAttribute .
+//
+
+class COREUOBJECT_API UAttributeProperty : public UProperty
+{
+	DECLARE_CASTED_CLASS_INTRINSIC(UAttributeProperty,UProperty,0,CoreUObject,CASTCLASS_UAttributeProperty)
+
+	// Variables.
+	UProperty* Inner;
+public:
+	UAttributeProperty( const class FPostConstructInitializeProperties& PCIP, ECppProperty, int32 InOffset, uint64 InFlags, int32 InElementSize )
+	:	UProperty( PCIP, EC_CppProperty, InOffset, InFlags)
+	{
+		ElementSize = InElementSize;
+	}
+
+	// UObject interface
+	virtual void Serialize( FArchive& Ar ) OVERRIDE;
+	static void AddReferencedObjects(UObject* InThis, FReferenceCollector& Collector);
+	// End of UObject interface
+
+	// UField interface
+	virtual void AddCppProperty( UProperty* Property ) OVERRIDE;
+	// End of UField interface
+
+	// UProperty interface
+	virtual FString GetCPPMacroType(FString& ExtendedTypeText) const  OVERRIDE;
+	virtual FString GetCPPType(FString* ExtendedTypeText, uint32 CPPExportFlags) const OVERRIDE;
+	virtual void LinkInternal(FArchive& Ar) OVERRIDE;
+	virtual bool Identical(const void* A, const void* B, uint32 PortFlags) const OVERRIDE;
+	virtual void SerializeItem(FArchive& Ar, void* Value, int32 MaxReadBytes, void const* Defaults) const OVERRIDE;
+	virtual bool NetSerializeItem(FArchive& Ar, UPackageMap* Map, void* Data, TArray<uint8> * MetaData = NULL) const OVERRIDE;
+	virtual void ExportTextItem(FString& ValueStr, const void* PropertyValue, const void* DefaultValue, UObject* Parent, int32 PortFlags, UObject* ExportRootScope) const OVERRIDE;
+	virtual const TCHAR* ImportText_Internal(const TCHAR* Buffer, void* Data, int32 PortFlags, UObject* OwnerObject, FOutputDevice* ErrorText) const OVERRIDE;
+	virtual void CopyValuesInternal(void* Dest, void const* Src, int32 Count) const OVERRIDE;
+	virtual void ClearValueInternal(void* Data) const OVERRIDE;
+	virtual void DestroyValueInternal(void* Dest) const OVERRIDE;
+	virtual bool IsLocalized() const OVERRIDE;
+	virtual bool PassCPPArgsByRef() const OVERRIDE;
+	virtual void InstanceSubobjects(void* Data, void const* DefaultData, UObject* Owner, struct FObjectInstancingGraph* InstanceGraph) OVERRIDE;
+	virtual bool ContainsObjectReference() const OVERRIDE;
+	virtual bool ContainsWeakObjectReference() const OVERRIDE;
+	virtual void EmitReferenceInfo(FGCReferenceTokenStream* ReferenceTokenStream, int32 BaseOffset) OVERRIDE;
+	virtual bool SameType(const UProperty* Other) const OVERRIDE;
+	// End of UProperty interface
+};
 
 /*-----------------------------------------------------------------------------
 	UStructProperty.
