#ttp 303393 UObject improvement: Remove header groups
#proj core
#branch UE4
#summary Removed header groups, removed some redundant includes and replaced with forward declarations. Fixed compilation errors.

diff --git a/Engine/Plugins/Experimental/LiveEditor/Source/LiveEditor/Private/K2Node_LiveEditObject.cpp b/Engine/Plugins/Experimental/LiveEditor/Source/LiveEditor/Private/K2Node_LiveEditObject.cpp
--- a/Engine/Plugins/Experimental/LiveEditor/Source/LiveEditor/Private/K2Node_LiveEditObject.cpp
+++ b/Engine/Plugins/Experimental/LiveEditor/Source/LiveEditor/Private/K2Node_LiveEditObject.cpp
@@ -5,7 +5,6 @@
 #include "BlueprintUtilities.h"
 #include "Kismet2/BlueprintEditorUtils.h"
 #include "KismetCompiler.h"
-#include "EngineKismetLibraryClasses.h"
 
 #define LOCTEXT_NAMESPACE "UK2Node_LiveEditObject"
 
diff --git a/Engine/Source/Developer/AssetTools/Private/AssetTypeActions/AssetTypeActions_ForceFeedbackEffect.cpp b/Engine/Source/Developer/AssetTools/Private/AssetTypeActions/AssetTypeActions_ForceFeedbackEffect.cpp
--- a/Engine/Source/Developer/AssetTools/Private/AssetTypeActions/AssetTypeActions_ForceFeedbackEffect.cpp
+++ b/Engine/Source/Developer/AssetTools/Private/AssetTypeActions/AssetTypeActions_ForceFeedbackEffect.cpp
@@ -1,7 +1,6 @@
 // Copyright 1998-2014 Epic Games, Inc. All Rights Reserved.
 
 #include "AssetToolsPrivatePCH.h"
-#include "EngineUserInterfaceClasses.h"
 
 UClass* FAssetTypeActions_ForceFeedbackEffect::GetSupportedClass() const
 {
diff --git a/Engine/Source/Developer/AssetTools/Private/AssetTypeActions/AssetTypeActions_InstancedFoliageSettings.cpp b/Engine/Source/Developer/AssetTools/Private/AssetTypeActions/AssetTypeActions_InstancedFoliageSettings.cpp
--- a/Engine/Source/Developer/AssetTools/Private/AssetTypeActions/AssetTypeActions_InstancedFoliageSettings.cpp
+++ b/Engine/Source/Developer/AssetTools/Private/AssetTypeActions/AssetTypeActions_InstancedFoliageSettings.cpp
@@ -1,7 +1,6 @@
 // Copyright 1998-2014 Epic Games, Inc. All Rights Reserved.
 
 #include "AssetToolsPrivatePCH.h"
-#include "EngineFoliageClasses.h"
 
 UClass* FAssetTypeActions_InstancedFoliageSettings::GetSupportedClass() const
 {
diff --git a/Engine/Source/Developer/AssetTools/Private/AssetTypeActions/AssetTypeActions_TouchInterface.cpp b/Engine/Source/Developer/AssetTools/Private/AssetTypeActions/AssetTypeActions_TouchInterface.cpp
--- a/Engine/Source/Developer/AssetTools/Private/AssetTypeActions/AssetTypeActions_TouchInterface.cpp
+++ b/Engine/Source/Developer/AssetTools/Private/AssetTypeActions/AssetTypeActions_TouchInterface.cpp
@@ -1,7 +1,6 @@
 // Copyright 1998-2014 Epic Games, Inc. All Rights Reserved.
 
 #include "AssetToolsPrivatePCH.h"
-#include "EngineUserInterfaceClasses.h"
 
 UClass* FAssetTypeActions_TouchInterface::GetSupportedClass() const
 {
diff --git a/Engine/Source/Developer/MovieSceneTools/Private/K2Node_PlayMovieScene.cpp b/Engine/Source/Developer/MovieSceneTools/Private/K2Node_PlayMovieScene.cpp
--- a/Engine/Source/Developer/MovieSceneTools/Private/K2Node_PlayMovieScene.cpp
+++ b/Engine/Source/Developer/MovieSceneTools/Private/K2Node_PlayMovieScene.cpp
@@ -4,7 +4,6 @@
 #include "ScopedTransaction.h"
 #include "K2Node_PlayMovieScene.h"
 #include "BlueprintUtilities.h"
-#include "EngineLevelScriptClasses.h"		// For ALevelScriptActor
 #include "MovieScene.h"
 #include "Kismet2/BlueprintEditorUtils.h"
 #include "MovieSceneBindings.h"
diff --git a/Engine/Source/Editor/AnimGraph/Classes/AnimGraphNode_Base.h b/Engine/Source/Editor/AnimGraph/Classes/AnimGraphNode_Base.h
--- a/Engine/Source/Editor/AnimGraph/Classes/AnimGraphNode_Base.h
+++ b/Engine/Source/Editor/AnimGraph/Classes/AnimGraphNode_Base.h
@@ -1,10 +1,13 @@
 // Copyright 1998-2014 Epic Games, Inc. All Rights Reserved.
 
 #pragma once
-
+#include "K2Node.h"
 #include "AnimGraphNode_Base.generated.h"
 
 //
+// Forward declarations.
+//
+struct FEdGraphSchemaAction_K2NewNode;
 
 struct FPoseLinkMappingRecord
 {
diff --git a/Engine/Source/Editor/AnimGraph/Classes/AnimGraphNode_BlendListByInt.h b/Engine/Source/Editor/AnimGraph/Classes/AnimGraphNode_BlendListByInt.h
--- a/Engine/Source/Editor/AnimGraph/Classes/AnimGraphNode_BlendListByInt.h
+++ b/Engine/Source/Editor/AnimGraph/Classes/AnimGraphNode_BlendListByInt.h
@@ -1,6 +1,7 @@
 // Copyright 1998-2014 Epic Games, Inc. All Rights Reserved.
 
 #pragma once
+#include "AnimGraphNode_BlendListBase.h"
 #include "AnimGraphNode_BlendListByInt.generated.h"
 
 UCLASS(MinimalAPI)
diff --git a/Engine/Source/Editor/AnimGraph/Classes/AnimGraphNode_BlendSpaceBase.h b/Engine/Source/Editor/AnimGraph/Classes/AnimGraphNode_BlendSpaceBase.h
--- a/Engine/Source/Editor/AnimGraph/Classes/AnimGraphNode_BlendSpaceBase.h
+++ b/Engine/Source/Editor/AnimGraph/Classes/AnimGraphNode_BlendSpaceBase.h
@@ -1,6 +1,7 @@
 // Copyright 1998-2014 Epic Games, Inc. All Rights Reserved.
 
 #pragma once
+#include "AnimGraphNode_Base.h"
 #include "AnimGraphNode_BlendSpaceBase.generated.h"
 
 UCLASS(Abstract, MinimalAPI)
diff --git a/Engine/Source/Editor/AnimGraph/Classes/AnimGraphNode_BlendSpaceEvaluator.h b/Engine/Source/Editor/AnimGraph/Classes/AnimGraphNode_BlendSpaceEvaluator.h
--- a/Engine/Source/Editor/AnimGraph/Classes/AnimGraphNode_BlendSpaceEvaluator.h
+++ b/Engine/Source/Editor/AnimGraph/Classes/AnimGraphNode_BlendSpaceEvaluator.h
@@ -1,7 +1,7 @@
 // Copyright 1998-2014 Epic Games, Inc. All Rights Reserved.
 
 #pragma once
-
+#include "AnimGraphNode_BlendSpaceBase.h"
 #include "Animation/AnimNode_BlendSpaceEvaluator.h"
 #include "AnimGraphNode_BlendSpaceEvaluator.generated.h"
 
diff --git a/Engine/Source/Editor/AnimGraph/Classes/AnimGraphNode_BlendSpacePlayer.h b/Engine/Source/Editor/AnimGraph/Classes/AnimGraphNode_BlendSpacePlayer.h
--- a/Engine/Source/Editor/AnimGraph/Classes/AnimGraphNode_BlendSpacePlayer.h
+++ b/Engine/Source/Editor/AnimGraph/Classes/AnimGraphNode_BlendSpacePlayer.h
@@ -1,6 +1,7 @@
 // Copyright 1998-2014 Epic Games, Inc. All Rights Reserved.
 
 #pragma once
+#include "AnimGraphNode_BlendSpaceBase.h"
 #include "AnimGraphNode_BlendSpacePlayer.generated.h"
 
 UCLASS(MinimalAPI)
diff --git a/Engine/Source/Editor/AnimGraph/Classes/AnimGraphNode_CustomTransitionResult.h b/Engine/Source/Editor/AnimGraph/Classes/AnimGraphNode_CustomTransitionResult.h
--- a/Engine/Source/Editor/AnimGraph/Classes/AnimGraphNode_CustomTransitionResult.h
+++ b/Engine/Source/Editor/AnimGraph/Classes/AnimGraphNode_CustomTransitionResult.h
@@ -1,6 +1,7 @@
 // Copyright 1998-2014 Epic Games, Inc. All Rights Reserved.
 
 #pragma once
+#include "AnimGraphNode_Root.h"
 #include "AnimGraphNode_CustomTransitionResult.generated.h"
 
 UCLASS(MinimalAPI)
diff --git a/Engine/Source/Editor/AnimGraph/Classes/AnimGraphNode_Root.h b/Engine/Source/Editor/AnimGraph/Classes/AnimGraphNode_Root.h
--- a/Engine/Source/Editor/AnimGraph/Classes/AnimGraphNode_Root.h
+++ b/Engine/Source/Editor/AnimGraph/Classes/AnimGraphNode_Root.h
@@ -1,7 +1,7 @@
 // Copyright 1998-2014 Epic Games, Inc. All Rights Reserved.
 
-
 #pragma once
+#include "AnimGraphNode_Base.h"
 #include "AnimGraphNode_Root.generated.h"
 
 UCLASS(MinimalAPI)
diff --git a/Engine/Source/Editor/AnimGraph/Classes/AnimGraphNode_RotationOffsetBlendSpace.h b/Engine/Source/Editor/AnimGraph/Classes/AnimGraphNode_RotationOffsetBlendSpace.h
--- a/Engine/Source/Editor/AnimGraph/Classes/AnimGraphNode_RotationOffsetBlendSpace.h
+++ b/Engine/Source/Editor/AnimGraph/Classes/AnimGraphNode_RotationOffsetBlendSpace.h
@@ -1,6 +1,7 @@
 // Copyright 1998-2014 Epic Games, Inc. All Rights Reserved.
 
 #pragma once
+#include "AnimGraphNode_BlendSpaceBase.h"
 #include "AnimGraphNode_RotationOffsetBlendSpace.generated.h"
 
 UCLASS(MinimalAPI)
diff --git a/Engine/Source/Editor/AnimGraph/Classes/AnimGraphNode_StateMachineBase.h b/Engine/Source/Editor/AnimGraph/Classes/AnimGraphNode_StateMachineBase.h
--- a/Engine/Source/Editor/AnimGraph/Classes/AnimGraphNode_StateMachineBase.h
+++ b/Engine/Source/Editor/AnimGraph/Classes/AnimGraphNode_StateMachineBase.h
@@ -1,6 +1,7 @@
 // Copyright 1998-2014 Epic Games, Inc. All Rights Reserved.
 
 #pragma once
+#include "AnimGraphNode_Base.h"
 #include "AnimGraphNode_StateMachineBase.generated.h"
 
 UCLASS(Abstract)
diff --git a/Engine/Source/Editor/AnimGraph/Classes/AnimGraphNode_StateResult.h b/Engine/Source/Editor/AnimGraph/Classes/AnimGraphNode_StateResult.h
--- a/Engine/Source/Editor/AnimGraph/Classes/AnimGraphNode_StateResult.h
+++ b/Engine/Source/Editor/AnimGraph/Classes/AnimGraphNode_StateResult.h
@@ -2,6 +2,7 @@
 
 
 #pragma once
+#include "AnimGraphNode_Root.h"
 #include "AnimGraphNode_StateResult.generated.h"
 
 UCLASS(MinimalAPI)
diff --git a/Engine/Source/Editor/AnimGraph/Classes/AnimGraphNode_TransitionPoseEvaluator.h b/Engine/Source/Editor/AnimGraph/Classes/AnimGraphNode_TransitionPoseEvaluator.h
--- a/Engine/Source/Editor/AnimGraph/Classes/AnimGraphNode_TransitionPoseEvaluator.h
+++ b/Engine/Source/Editor/AnimGraph/Classes/AnimGraphNode_TransitionPoseEvaluator.h
@@ -1,7 +1,7 @@
 // Copyright 1998-2014 Epic Games, Inc. All Rights Reserved.
 
 #pragma once
-
+#include "AnimGraphNode_Base.h"
 #include "AnimGraphNode_TransitionPoseEvaluator.generated.h"
 
 UCLASS(MinimalAPI)
diff --git a/Engine/Source/Editor/AnimGraph/Classes/AnimPreviewInstance.h b/Engine/Source/Editor/AnimGraph/Classes/AnimPreviewInstance.h
--- a/Engine/Source/Editor/AnimGraph/Classes/AnimPreviewInstance.h
+++ b/Engine/Source/Editor/AnimGraph/Classes/AnimPreviewInstance.h
@@ -4,6 +4,7 @@
 
 #include "BlueprintGraphDefinitions.h"
 #include "AnimGraphNode_ModifyBone.h"
+#include "Animation/AnimSingleNodeInstance.h"
 #include "AnimPreviewInstance.generated.h"
 
 /** Enum to know how montage is being played */
diff --git a/Engine/Source/Editor/AnimGraph/Classes/AnimStateConduitNode.h b/Engine/Source/Editor/AnimGraph/Classes/AnimStateConduitNode.h
--- a/Engine/Source/Editor/AnimGraph/Classes/AnimStateConduitNode.h
+++ b/Engine/Source/Editor/AnimGraph/Classes/AnimStateConduitNode.h
@@ -1,6 +1,7 @@
 // Copyright 1998-2014 Epic Games, Inc. All Rights Reserved.
 
 #pragma once
+#include "AnimStateNodeBase.h"
 #include "AnimStateConduitNode.generated.h"
 
 UCLASS(MinimalAPI)
diff --git a/Engine/Source/Editor/AnimGraph/Classes/AnimStateNode.h b/Engine/Source/Editor/AnimGraph/Classes/AnimStateNode.h
--- a/Engine/Source/Editor/AnimGraph/Classes/AnimStateNode.h
+++ b/Engine/Source/Editor/AnimGraph/Classes/AnimStateNode.h
@@ -1,6 +1,7 @@
 // Copyright 1998-2014 Epic Games, Inc. All Rights Reserved.
 
 #pragma once
+#include "AnimStateNodeBase.h"
 #include "AnimStateNode.generated.h"
 
 UENUM()
diff --git a/Engine/Source/Editor/AnimGraph/Classes/AnimStateNodeBase.h b/Engine/Source/Editor/AnimGraph/Classes/AnimStateNodeBase.h
--- a/Engine/Source/Editor/AnimGraph/Classes/AnimStateNodeBase.h
+++ b/Engine/Source/Editor/AnimGraph/Classes/AnimStateNodeBase.h
@@ -1,6 +1,7 @@
 // Copyright 1998-2014 Epic Games, Inc. All Rights Reserved.
 
 #pragma once
+#include "EdGraph/EdGraphNode.h"
 #include "AnimStateNodeBase.generated.h"
 
 UCLASS(MinimalAPI, abstract)
diff --git a/Engine/Source/Editor/AnimGraph/Classes/AnimStateTransitionNode.h b/Engine/Source/Editor/AnimGraph/Classes/AnimStateTransitionNode.h
--- a/Engine/Source/Editor/AnimGraph/Classes/AnimStateTransitionNode.h
+++ b/Engine/Source/Editor/AnimGraph/Classes/AnimStateTransitionNode.h
@@ -1,9 +1,9 @@
 // Copyright 1998-2014 Epic Games, Inc. All Rights Reserved.
 
 #pragma once
+#include "AnimStateNodeBase.h"
 #include "AnimStateTransitionNode.generated.h"
 
-
 UCLASS(MinimalAPI, DependsOn=UAnimInstance, config=Editor)
 class UAnimStateTransitionNode : public UAnimStateNodeBase
 {
diff --git a/Engine/Source/Editor/AnimGraph/Classes/AnimationCustomTransitionGraph.h b/Engine/Source/Editor/AnimGraph/Classes/AnimationCustomTransitionGraph.h
--- a/Engine/Source/Editor/AnimGraph/Classes/AnimationCustomTransitionGraph.h
+++ b/Engine/Source/Editor/AnimGraph/Classes/AnimationCustomTransitionGraph.h
@@ -1,7 +1,7 @@
 // Copyright 1998-2014 Epic Games, Inc. All Rights Reserved.
 
 #pragma once
-
+#include "AnimationGraph.h"
 #include "AnimationCustomTransitionGraph.generated.h"
 
 UCLASS(MinimalAPI)
diff --git a/Engine/Source/Editor/AnimGraph/Classes/AnimationCustomTransitionSchema.h b/Engine/Source/Editor/AnimGraph/Classes/AnimationCustomTransitionSchema.h
--- a/Engine/Source/Editor/AnimGraph/Classes/AnimationCustomTransitionSchema.h
+++ b/Engine/Source/Editor/AnimGraph/Classes/AnimationCustomTransitionSchema.h
@@ -1,6 +1,7 @@
 // Copyright 1998-2014 Epic Games, Inc. All Rights Reserved.
 
 #pragma once
+#include "AnimationGraphSchema.h"
 #include "AnimationCustomTransitionSchema.generated.h"
 
 UCLASS(MinimalAPI)
diff --git a/Engine/Source/Editor/AnimGraph/Classes/AnimationStateGraph.h b/Engine/Source/Editor/AnimGraph/Classes/AnimationStateGraph.h
--- a/Engine/Source/Editor/AnimGraph/Classes/AnimationStateGraph.h
+++ b/Engine/Source/Editor/AnimGraph/Classes/AnimationStateGraph.h
@@ -1,6 +1,7 @@
 // Copyright 1998-2014 Epic Games, Inc. All Rights Reserved.
 
 #pragma once
+#include "AnimationGraph.h"
 #include "AnimationStateGraph.generated.h"
 
 UCLASS(MinimalAPI)
diff --git a/Engine/Source/Editor/AnimGraph/Classes/AnimationStateMachineGraph.h b/Engine/Source/Editor/AnimGraph/Classes/AnimationStateMachineGraph.h
--- a/Engine/Source/Editor/AnimGraph/Classes/AnimationStateMachineGraph.h
+++ b/Engine/Source/Editor/AnimGraph/Classes/AnimationStateMachineGraph.h
@@ -1,7 +1,7 @@
 // Copyright 1998-2014 Epic Games, Inc. All Rights Reserved.
 
 #pragma once
-
+#include "EdGraph/EdGraph.h"
 #include "AnimationStateMachineGraph.generated.h"
 
 UCLASS(MinimalAPI)
diff --git a/Engine/Source/Editor/AnimGraph/Classes/AnimationTransitionGraph.h b/Engine/Source/Editor/AnimGraph/Classes/AnimationTransitionGraph.h
--- a/Engine/Source/Editor/AnimGraph/Classes/AnimationTransitionGraph.h
+++ b/Engine/Source/Editor/AnimGraph/Classes/AnimationTransitionGraph.h
@@ -1,6 +1,7 @@
 // Copyright 1998-2014 Epic Games, Inc. All Rights Reserved.
 
 #pragma once
+#include "AnimationGraph.h"
 #include "AnimationTransitionGraph.generated.h"
 
 UCLASS(MinimalAPI)
diff --git a/Engine/Source/Editor/AnimGraph/Private/AnimationConduitGraphSchema.cpp b/Engine/Source/Editor/AnimGraph/Private/AnimationConduitGraphSchema.cpp
--- a/Engine/Source/Editor/AnimGraph/Private/AnimationConduitGraphSchema.cpp
+++ b/Engine/Source/Editor/AnimGraph/Private/AnimationConduitGraphSchema.cpp
@@ -1,6 +1,10 @@
 // Copyright 1998-2014 Epic Games, Inc. All Rights Reserved.
 
 #include "AnimGraphPrivatePCH.h"
+#include "AnimationConduitGraphSchema.h"
+#include "AnimGraphNode_TransitionResult.h"
+#include "AnimationTransitionGraph.h"
+#include "AnimStateConduitNode.h"
 
 /////////////////////////////////////////////////////
 // UAnimationConduitGraphSchema
diff --git a/Engine/Source/Editor/AnimGraph/Private/AnimationCustomTransitionGraph.cpp b/Engine/Source/Editor/AnimGraph/Private/AnimationCustomTransitionGraph.cpp
--- a/Engine/Source/Editor/AnimGraph/Private/AnimationCustomTransitionGraph.cpp
+++ b/Engine/Source/Editor/AnimGraph/Private/AnimationCustomTransitionGraph.cpp
@@ -1,7 +1,7 @@
 // Copyright 1998-2014 Epic Games, Inc. All Rights Reserved.
 
 #include "AnimGraphPrivatePCH.h"
-
+#include "AnimationCustomTransitionGraph.h"
 
 #define LOCTEXT_NAMESPACE "AnimationCustomTransitionGraph"
 
diff --git a/Engine/Source/Editor/AnimGraph/Private/AnimationCustomTransitionSchema.cpp b/Engine/Source/Editor/AnimGraph/Private/AnimationCustomTransitionSchema.cpp
--- a/Engine/Source/Editor/AnimGraph/Private/AnimationCustomTransitionSchema.cpp
+++ b/Engine/Source/Editor/AnimGraph/Private/AnimationCustomTransitionSchema.cpp
@@ -1,7 +1,10 @@
 // Copyright 1998-2014 Epic Games, Inc. All Rights Reserved.
 
 #include "AnimGraphPrivatePCH.h"
-
+#include "AnimationCustomTransitionSchema.h"
+#include "AnimGraphNode_CustomTransitionResult.h"
+#include "AnimGraphNode_TransitionPoseEvaluator.h"
+#include "AnimStateTransitionNode.h"
 
 /////////////////////////////////////////////////////
 // UAnimationCustomTransitionSchema
diff --git a/Engine/Source/Editor/AnimGraph/Private/AnimationGraph.cpp b/Engine/Source/Editor/AnimGraph/Private/AnimationGraph.cpp
--- a/Engine/Source/Editor/AnimGraph/Private/AnimationGraph.cpp
+++ b/Engine/Source/Editor/AnimGraph/Private/AnimationGraph.cpp
@@ -1,7 +1,7 @@
 // Copyright 1998-2014 Epic Games, Inc. All Rights Reserved.
 
 #include "AnimGraphPrivatePCH.h"
-
+#include "AnimationGraph.h"
 
 #define LOCTEXT_NAMESPACE "AnimationGraph"
 
diff --git a/Engine/Source/Editor/AnimGraph/Private/AnimationGraphSchema.cpp b/Engine/Source/Editor/AnimGraph/Private/AnimationGraphSchema.cpp
--- a/Engine/Source/Editor/AnimGraph/Private/AnimationGraphSchema.cpp
+++ b/Engine/Source/Editor/AnimGraph/Private/AnimationGraphSchema.cpp
@@ -10,6 +10,14 @@
 #include "ScopedTransaction.h"
 #include "AssetData.h"
 #include "K2ActionMenuBuilder.h" // for FK2ActionMenuBuilder::AddNewNodeAction()
+#include "AnimationGraphSchema.h"
+#include "K2Node_TransitionRuleGetter.h"
+#include "AnimStateNode.h"
+#include "AnimGraphNode_BlendSpacePlayer.h"
+#include "AnimGraphNode_ComponentToLocalSpace.h"
+#include "AnimGraphNode_LocalToComponentSpace.h"
+#include "AnimGraphNode_RotationOffsetBlendSpace.h"
+#include "AnimGraphNode_SequencePlayer.h"
 
 #define LOCTEXT_NAMESPACE "AnimationGraphSchema"
 
diff --git a/Engine/Source/Editor/AnimGraph/Private/AnimationStateGraph.cpp b/Engine/Source/Editor/AnimGraph/Private/AnimationStateGraph.cpp
--- a/Engine/Source/Editor/AnimGraph/Private/AnimationStateGraph.cpp
+++ b/Engine/Source/Editor/AnimGraph/Private/AnimationStateGraph.cpp
@@ -1,7 +1,7 @@
 // Copyright 1998-2014 Epic Games, Inc. All Rights Reserved.
 
 #include "AnimGraphPrivatePCH.h"
-
+#include "AnimationStateGraph.h"
 
 #define LOCTEXT_NAMESPACE "AnimationStateGraph"
 
diff --git a/Engine/Source/Editor/AnimGraph/Private/AnimationStateGraphSchema.cpp b/Engine/Source/Editor/AnimGraph/Private/AnimationStateGraphSchema.cpp
--- a/Engine/Source/Editor/AnimGraph/Private/AnimationStateGraphSchema.cpp
+++ b/Engine/Source/Editor/AnimGraph/Private/AnimationStateGraphSchema.cpp
@@ -5,6 +5,7 @@
 #include "GraphEditorActions.h"
 #include "ScopedTransaction.h"
 #include "AssetData.h"
+#include "AnimationStateGraphSchema.h"
 
 #define LOCTEXT_NAMESPACE "AnimationStateGraphSchema"
 
diff --git a/Engine/Source/Editor/AnimGraph/Private/AnimationStateMachineGraph.cpp b/Engine/Source/Editor/AnimGraph/Private/AnimationStateMachineGraph.cpp
--- a/Engine/Source/Editor/AnimGraph/Private/AnimationStateMachineGraph.cpp
+++ b/Engine/Source/Editor/AnimGraph/Private/AnimationStateMachineGraph.cpp
@@ -1,7 +1,7 @@
 // Copyright 1998-2014 Epic Games, Inc. All Rights Reserved.
 
 #include "AnimGraphPrivatePCH.h"
-
+#include "AnimationStatemachineGraph.h"
 
 #define LOCTEXT_NAMESPACE "AnimationGraph"
 
diff --git a/Engine/Source/Editor/AnimGraph/Private/AnimationStateMachineSchema.cpp b/Engine/Source/Editor/AnimGraph/Private/AnimationStateMachineSchema.cpp
--- a/Engine/Source/Editor/AnimGraph/Private/AnimationStateMachineSchema.cpp
+++ b/Engine/Source/Editor/AnimGraph/Private/AnimationStateMachineSchema.cpp
@@ -6,6 +6,9 @@
 
 #include "AnimGraphPrivatePCH.h"
 
+#include "AnimationStateMachineSchema.h"
+#include "AnimStateEntryNode.h"
+
 #include "ScopedTransaction.h"
 #include "GraphEditorActions.h"
 
diff --git a/Engine/Source/Editor/AnimGraph/Private/AnimationTransitionSchema.cpp b/Engine/Source/Editor/AnimGraph/Private/AnimationTransitionSchema.cpp
--- a/Engine/Source/Editor/AnimGraph/Private/AnimationTransitionSchema.cpp
+++ b/Engine/Source/Editor/AnimGraph/Private/AnimationTransitionSchema.cpp
@@ -5,6 +5,7 @@
 =============================================================================*/
 
 #include "AnimGraphPrivatePCH.h"
+#include "AnimationTransitionSchema.h"
 
 #include "BlueprintUtilities.h"
 #include "AnimGraphDefinitions.h"
diff --git a/Engine/Source/Editor/AnimGraph/Public/AnimGraphDefinitions.h b/Engine/Source/Editor/AnimGraph/Public/AnimGraphDefinitions.h
--- a/Engine/Source/Editor/AnimGraph/Public/AnimGraphDefinitions.h
+++ b/Engine/Source/Editor/AnimGraph/Public/AnimGraphDefinitions.h
@@ -1,5 +1,3 @@
 // Copyright 1998-2014 Epic Games, Inc. All Rights Reserved.
 
 #include "BlueprintGraphDefinitions.h"
-#include "EngineAnimationNodeClasses.h"
-#include "AnimGraphClasses.h"
diff --git a/Engine/Source/Editor/BlueprintGraph/Classes/K2Node.h b/Engine/Source/Editor/BlueprintGraph/Classes/K2Node.h
--- a/Engine/Source/Editor/BlueprintGraph/Classes/K2Node.h
+++ b/Engine/Source/Editor/BlueprintGraph/Classes/K2Node.h
@@ -1,8 +1,7 @@
 // Copyright 1998-2014 Epic Games, Inc. All Rights Reserved.
 
 #pragma once
-
-
+#include "EdGraph/EdGraphNode.h"
 #include "K2Node.generated.h"
 
 /** Helper struct to allow us to redirect properties and functions through renames and additionally between classes if necessary */
diff --git a/Engine/Source/Editor/BlueprintGraph/Private/CallFunctionHandler.cpp b/Engine/Source/Editor/BlueprintGraph/Private/CallFunctionHandler.cpp
--- a/Engine/Source/Editor/BlueprintGraph/Private/CallFunctionHandler.cpp
+++ b/Engine/Source/Editor/BlueprintGraph/Private/CallFunctionHandler.cpp
@@ -1,8 +1,6 @@
 // Copyright 1998-2014 Epic Games, Inc. All Rights Reserved.
 
 #include "BlueprintGraphPrivatePCH.h"
-#include "EngineKismetLibraryClasses.h"
-#include "EngineLevelScriptClasses.h"
 
 #include "CallFunctionHandler.h"
 
diff --git a/Engine/Source/Editor/BlueprintGraph/Private/EdGraphSchema_K2.cpp b/Engine/Source/Editor/BlueprintGraph/Private/EdGraphSchema_K2.cpp
--- a/Engine/Source/Editor/BlueprintGraph/Private/EdGraphSchema_K2.cpp
+++ b/Engine/Source/Editor/BlueprintGraph/Private/EdGraphSchema_K2.cpp
@@ -13,7 +13,6 @@
 
 #include "Kismet2/KismetEditorUtilities.h"
 #include "Kismet2/KismetDebugUtilities.h"
-#include "EngineKismetLibraryClasses.h"
 #include "ComponentAssetBroker.h"
 #include "AssetData.h"
 #include "Editor/UnrealEd/Public/EdGraphUtilities.h"
diff --git a/Engine/Source/Editor/BlueprintGraph/Private/K2Node_ActorBoundEvent.cpp b/Engine/Source/Editor/BlueprintGraph/Private/K2Node_ActorBoundEvent.cpp
--- a/Engine/Source/Editor/BlueprintGraph/Private/K2Node_ActorBoundEvent.cpp
+++ b/Engine/Source/Editor/BlueprintGraph/Private/K2Node_ActorBoundEvent.cpp
@@ -2,7 +2,6 @@
 
 #include "BlueprintGraphPrivatePCH.h"
 
-#include "EngineLevelScriptClasses.h"
 #include "KismetCompiler.h"
 #include "EventEntryHandler.h"
 
diff --git a/Engine/Source/Editor/BlueprintGraph/Private/K2Node_CallFunction.cpp b/Engine/Source/Editor/BlueprintGraph/Private/K2Node_CallFunction.cpp
--- a/Engine/Source/Editor/BlueprintGraph/Private/K2Node_CallFunction.cpp
+++ b/Engine/Source/Editor/BlueprintGraph/Private/K2Node_CallFunction.cpp
@@ -3,8 +3,6 @@
 
 #include "BlueprintGraphPrivatePCH.h"
 
-#include "EngineKismetLibraryClasses.h"
-
 #include "CompilerResultsLog.h"
 #include "CallFunctionHandler.h"
 
diff --git a/Engine/Source/Editor/BlueprintGraph/Private/K2Node_ComponentBoundEvent.cpp b/Engine/Source/Editor/BlueprintGraph/Private/K2Node_ComponentBoundEvent.cpp
--- a/Engine/Source/Editor/BlueprintGraph/Private/K2Node_ComponentBoundEvent.cpp
+++ b/Engine/Source/Editor/BlueprintGraph/Private/K2Node_ComponentBoundEvent.cpp
@@ -2,8 +2,6 @@
 
 #include "BlueprintGraphPrivatePCH.h"
 
-#include "EngineLevelScriptClasses.h"
-
 #define LOCTEXT_NAMESPACE "K2Node"
 
 UK2Node_ComponentBoundEvent::UK2Node_ComponentBoundEvent(const class FPostConstructInitializeProperties& PCIP)
diff --git a/Engine/Source/Editor/BlueprintGraph/Private/K2Node_CustomEvent.cpp b/Engine/Source/Editor/BlueprintGraph/Private/K2Node_CustomEvent.cpp
--- a/Engine/Source/Editor/BlueprintGraph/Private/K2Node_CustomEvent.cpp
+++ b/Engine/Source/Editor/BlueprintGraph/Private/K2Node_CustomEvent.cpp
@@ -3,7 +3,6 @@
 
 #include "BlueprintGraphPrivatePCH.h"
 
-#include "EngineLevelScriptClasses.h"
 #include "Kismet2NameValidators.h"
 #include "K2Node_BaseMCDelegate.h"
 #include "CompilerResultsLog.h"
diff --git a/Engine/Source/Editor/BlueprintGraph/Private/K2Node_InputActionEvent.cpp b/Engine/Source/Editor/BlueprintGraph/Private/K2Node_InputActionEvent.cpp
--- a/Engine/Source/Editor/BlueprintGraph/Private/K2Node_InputActionEvent.cpp
+++ b/Engine/Source/Editor/BlueprintGraph/Private/K2Node_InputActionEvent.cpp
@@ -2,8 +2,6 @@
 
 #include "BlueprintGraphPrivatePCH.h"
 
-#include "EngineLevelScriptClasses.h"
-
 UK2Node_InputActionEvent::UK2Node_InputActionEvent(const class FPostConstructInitializeProperties& PCIP)
 	: Super(PCIP)
 {
diff --git a/Engine/Source/Editor/BlueprintGraph/Private/K2Node_Message.cpp b/Engine/Source/Editor/BlueprintGraph/Private/K2Node_Message.cpp
--- a/Engine/Source/Editor/BlueprintGraph/Private/K2Node_Message.cpp
+++ b/Engine/Source/Editor/BlueprintGraph/Private/K2Node_Message.cpp
@@ -3,8 +3,6 @@
 #include "BlueprintGraphPrivatePCH.h"
 #include "KismetCompiler.h"
 
-#include "EngineKismetLibraryClasses.h"
-
 #define LOCTEXT_NAMESPACE "K2Node_Message"
 
 UK2Node_Message::UK2Node_Message(const class FPostConstructInitializeProperties& PCIP)
diff --git a/Engine/Source/Editor/Cascade/Classes/CascadeParticleSystemComponent.h b/Engine/Source/Editor/Cascade/Classes/CascadeParticleSystemComponent.h
--- a/Engine/Source/Editor/Cascade/Classes/CascadeParticleSystemComponent.h
+++ b/Engine/Source/Editor/Cascade/Classes/CascadeParticleSystemComponent.h
@@ -8,7 +8,7 @@
 /** 
  * Used to provide an extended particle system component to allow collision to function in the preview window.
  */
-UCLASS(HeaderGroup=Component, hidecategories=Object, hidecategories=Physics, hidecategories=Collision, editinlinenew)
+UCLASS(hidecategories=Object, hidecategories=Physics, hidecategories=Collision, editinlinenew)
 class UCascadeParticleSystemComponent : public UParticleSystemComponent
 {
 	GENERATED_UCLASS_BODY()
diff --git a/Engine/Source/Editor/ComponentVisualizers/Private/ComponentVisualizers.cpp b/Engine/Source/Editor/ComponentVisualizers/Private/ComponentVisualizers.cpp
--- a/Engine/Source/Editor/ComponentVisualizers/Private/ComponentVisualizers.cpp
+++ b/Engine/Source/Editor/ComponentVisualizers/Private/ComponentVisualizers.cpp
@@ -4,7 +4,6 @@
 #include "ComponentVisualizers.h"
 
 #include "SoundDefinitions.h"
-#include "EngineSplineClasses.h"
 
 #include "PointLightComponentVisualizer.h"
 #include "SpotLightComponentVisualizer.h"
diff --git a/Engine/Source/Editor/ComponentVisualizers/Private/SplineComponentVisualizer.cpp b/Engine/Source/Editor/ComponentVisualizers/Private/SplineComponentVisualizer.cpp
--- a/Engine/Source/Editor/ComponentVisualizers/Private/SplineComponentVisualizer.cpp
+++ b/Engine/Source/Editor/ComponentVisualizers/Private/SplineComponentVisualizer.cpp
@@ -1,7 +1,6 @@
 // Copyright 1998-2014 Epic Games, Inc. All Rights Reserved.
 
 #include "ComponentVisualizersPrivatePCH.h"
-#include "EngineSplineClasses.h"
 #include "SplineComponentVisualizer.h"
 
 void FSplineComponentVisualizer::DrawVisualization(const UActorComponent* Component, const FSceneView* View, FPrimitiveDrawInterface* PDI)
diff --git a/Engine/Source/Editor/DetailCustomizations/Private/ActorMaterialCategory.cpp b/Engine/Source/Editor/DetailCustomizations/Private/ActorMaterialCategory.cpp
--- a/Engine/Source/Editor/DetailCustomizations/Private/ActorMaterialCategory.cpp
+++ b/Engine/Source/Editor/DetailCustomizations/Private/ActorMaterialCategory.cpp
@@ -2,7 +2,6 @@
 
 #include "DetailCustomizationsPrivatePCH.h"
 #include "ActorMaterialCategory.h"
-#include "EngineDecalClasses.h"
 #include "AssetThumbnail.h"
 #include "ActorEditorUtils.h"
 
diff --git a/Engine/Source/Editor/DetailCustomizations/Private/AnimMontageSegmentDetails.cpp b/Engine/Source/Editor/DetailCustomizations/Private/AnimMontageSegmentDetails.cpp
--- a/Engine/Source/Editor/DetailCustomizations/Private/AnimMontageSegmentDetails.cpp
+++ b/Engine/Source/Editor/DetailCustomizations/Private/AnimMontageSegmentDetails.cpp
@@ -2,10 +2,10 @@
 
 #include "DetailCustomizationsPrivatePCH.h"
 #include "AnimMontageSegmentDetails.h"
+#include "AnimPreviewInstance.h"
 #include "Runtime/Engine/Public/Slate/SceneViewport.h"
 #include "Editor/KismetWidgets/Public/SScrubControlPanel.h"
 #include "Runtime/Engine/Public/FXSystem.h"
-
 #define LOCTEXT_NAMESPACE "AnimMontageSegmentDetails"
 
 /////////////////////////////////////////////////////////////////////////
diff --git a/Engine/Source/Editor/DetailCustomizations/Private/AnimSequenceDetails.cpp b/Engine/Source/Editor/DetailCustomizations/Private/AnimSequenceDetails.cpp
--- a/Engine/Source/Editor/DetailCustomizations/Private/AnimSequenceDetails.cpp
+++ b/Engine/Source/Editor/DetailCustomizations/Private/AnimSequenceDetails.cpp
@@ -3,6 +3,7 @@
 #include "DetailCustomizationsPrivatePCH.h"
 #include "AnimSequenceDetails.h"
 #include "AnimMontageSegmentDetails.h"
+#include "AnimPreviewInstance.h"
 #include "Runtime/Engine/Public/Slate/SceneViewport.h"
 #include "Editor/KismetWidgets/Public/SScrubControlPanel.h"
 #include "Runtime/Engine/Public/FXSystem.h"
diff --git a/Engine/Source/Editor/DetailCustomizations/Private/AnimTransitionNodeDetails.cpp b/Engine/Source/Editor/DetailCustomizations/Private/AnimTransitionNodeDetails.cpp
--- a/Engine/Source/Editor/DetailCustomizations/Private/AnimTransitionNodeDetails.cpp
+++ b/Engine/Source/Editor/DetailCustomizations/Private/AnimTransitionNodeDetails.cpp
@@ -2,8 +2,12 @@
 
 #include "DetailCustomizationsPrivatePCH.h"
 
+#include "AnimationTransitionGraph.h"
 #include "AnimGraphDefinitions.h"
-
+#include "AnimGraphNode_TransitionResult.h"
+#include "AnimStateConduitNode.h"
+#include "AnimStateNodeBase.h"
+#include "AnimStateTransitionNode.h"
 #include "AnimTransitionNodeDetails.h"
 #include "KismetEditorUtilities.h"
 #include "SKismetLinearExpression.h"
diff --git a/Engine/Source/Editor/DetailCustomizations/Private/CameraDetails.cpp b/Engine/Source/Editor/DetailCustomizations/Private/CameraDetails.cpp
--- a/Engine/Source/Editor/DetailCustomizations/Private/CameraDetails.cpp
+++ b/Engine/Source/Editor/DetailCustomizations/Private/CameraDetails.cpp
@@ -2,7 +2,6 @@
 
 #include "DetailCustomizationsPrivatePCH.h"
 #include "CameraDetails.h"
-#include "EngineDecalClasses.h"
 
 #define LOCTEXT_NAMESPACE "CameraDetails"
 
diff --git a/Engine/Source/Editor/DetailCustomizations/Private/ReflectionCaptureDetails.h b/Engine/Source/Editor/DetailCustomizations/Private/ReflectionCaptureDetails.h
--- a/Engine/Source/Editor/DetailCustomizations/Private/ReflectionCaptureDetails.h
+++ b/Engine/Source/Editor/DetailCustomizations/Private/ReflectionCaptureDetails.h
@@ -2,8 +2,6 @@
 
 #pragma once
 
-#include "EngineDecalClasses.h"
-
 class FReflectionCaptureDetails : public IDetailCustomization
 {
 public:
diff --git a/Engine/Source/Editor/DetailCustomizations/Private/SkyLightComponentDetails.h b/Engine/Source/Editor/DetailCustomizations/Private/SkyLightComponentDetails.h
--- a/Engine/Source/Editor/DetailCustomizations/Private/SkyLightComponentDetails.h
+++ b/Engine/Source/Editor/DetailCustomizations/Private/SkyLightComponentDetails.h
@@ -2,8 +2,6 @@
 
 #pragma once
 
-#include "EngineDecalClasses.h"
-
 class FSkyLightComponentDetails : public IDetailCustomization
 {
 public:
diff --git a/Engine/Source/Editor/DetailCustomizations/Private/TransitionPoseEvaluatorNodeDetails.cpp b/Engine/Source/Editor/DetailCustomizations/Private/TransitionPoseEvaluatorNodeDetails.cpp
--- a/Engine/Source/Editor/DetailCustomizations/Private/TransitionPoseEvaluatorNodeDetails.cpp
+++ b/Engine/Source/Editor/DetailCustomizations/Private/TransitionPoseEvaluatorNodeDetails.cpp
@@ -3,6 +3,7 @@
 #include "DetailCustomizationsPrivatePCH.h"
 
 #include "AnimGraphDefinitions.h"
+#include "AnimGraphNode_TransitionPoseEvaluator.h"
 
 #include "AnimTransitionNodeDetails.h"
 #include "KismetEditorUtilities.h"
diff --git a/Engine/Source/Editor/DetailCustomizations/Private/TransitionPoseEvaluatorNodeDetails.h b/Engine/Source/Editor/DetailCustomizations/Private/TransitionPoseEvaluatorNodeDetails.h
--- a/Engine/Source/Editor/DetailCustomizations/Private/TransitionPoseEvaluatorNodeDetails.h
+++ b/Engine/Source/Editor/DetailCustomizations/Private/TransitionPoseEvaluatorNodeDetails.h
@@ -5,6 +5,11 @@
 #include "PreviewScene.h"
 #include "AnimGraphDefinitions.h"
 
+//
+// Forward declarations.
+//
+class UAnimGraphNode_TransitionPoseEvaluator;
+
 class FTransitionPoseEvaluatorNodeDetails : public IDetailCustomization
 {
 public:
diff --git a/Engine/Source/Editor/FoliageEdit/Private/FoliageEdMode.cpp b/Engine/Source/Editor/FoliageEdit/Private/FoliageEdMode.cpp
--- a/Engine/Source/Editor/FoliageEdit/Private/FoliageEdMode.cpp
+++ b/Engine/Source/Editor/FoliageEdit/Private/FoliageEdMode.cpp
@@ -8,7 +8,6 @@
 #include "ObjectTools.h"
 #include "FoliageEdMode.h"
 #include "ScopedTransaction.h"
-#include "EngineFoliageClasses.h"
 
 #include "FoliageEdModeToolkit.h"
 #include "ModuleManager.h"
diff --git a/Engine/Source/Editor/GraphEditor/Private/AnimationNodes/SGraphNodeAnimationResult.cpp b/Engine/Source/Editor/GraphEditor/Private/AnimationNodes/SGraphNodeAnimationResult.cpp
--- a/Engine/Source/Editor/GraphEditor/Private/AnimationNodes/SGraphNodeAnimationResult.cpp
+++ b/Engine/Source/Editor/GraphEditor/Private/AnimationNodes/SGraphNodeAnimationResult.cpp
@@ -3,6 +3,7 @@
 #include "GraphEditorCommon.h"
 #include "AnimGraphDefinitions.h"
 #include "SGraphNodeAnimationResult.h"
+#include "EdGraph/EdGraphNode.h"
 
 /////////////////////////////////////////////////////
 // SGraphNodeAnimationResult
diff --git a/Engine/Source/Editor/GraphEditor/Private/AnimationNodes/SGraphNodeStateMachineInstance.cpp b/Engine/Source/Editor/GraphEditor/Private/AnimationNodes/SGraphNodeStateMachineInstance.cpp
--- a/Engine/Source/Editor/GraphEditor/Private/AnimationNodes/SGraphNodeStateMachineInstance.cpp
+++ b/Engine/Source/Editor/GraphEditor/Private/AnimationNodes/SGraphNodeStateMachineInstance.cpp
@@ -1,10 +1,12 @@
 // Copyright 1998-2014 Epic Games, Inc. All Rights Reserved.
 
 #include "GraphEditorCommon.h"
-#include "../KismetNodes/SGraphNodeK2Base.h"
-#include "../KismetNodes/SGraphNodeK2Composite.h"
+#include "KismetNodes/SGraphNodeK2Base.h"
+#include "KismetNodes/SGraphNodeK2Composite.h"
 #include "SGraphNodeStateMachineInstance.h"
 #include "SGraphPreviewer.h"
+#include "AnimationStateMachineGraph.h"
+#include "AnimGraphNode_StateMachineBase.h"
 
 /////////////////////////////////////////////////////
 // SGraphNodeStateMachineInstance
diff --git a/Engine/Source/Editor/GraphEditor/Private/AnimationStateNodes/SGraphNodeAnimState.cpp b/Engine/Source/Editor/GraphEditor/Private/AnimationStateNodes/SGraphNodeAnimState.cpp
--- a/Engine/Source/Editor/GraphEditor/Private/AnimationStateNodes/SGraphNodeAnimState.cpp
+++ b/Engine/Source/Editor/GraphEditor/Private/AnimationStateNodes/SGraphNodeAnimState.cpp
@@ -6,6 +6,8 @@
 #include "Editor/UnrealEd/Public/Kismet2/BlueprintEditorUtils.h"
 #include "SGraphNode.h"
 #include "IDocumentation.h"
+#include "AnimStateConduitNode.h"
+#include "AnimationStateMachineGraph.h"
 
 /////////////////////////////////////////////////////
 // SStateMachineOutputPin
diff --git a/Engine/Source/Editor/GraphEditor/Private/AnimationStateNodes/SGraphNodeAnimState.h b/Engine/Source/Editor/GraphEditor/Private/AnimationStateNodes/SGraphNodeAnimState.h
--- a/Engine/Source/Editor/GraphEditor/Private/AnimationStateNodes/SGraphNodeAnimState.h
+++ b/Engine/Source/Editor/GraphEditor/Private/AnimationStateNodes/SGraphNodeAnimState.h
@@ -2,6 +2,13 @@
 
 #pragma once
 
+//
+// Forward declarations.
+//
+class UAnimStateNodeBase;
+class UAnimStateConduitNode;
+class UAnimStateEntryNode;
+
 class SGraphNodeAnimState : public SGraphNode
 {
 public:
diff --git a/Engine/Source/Editor/GraphEditor/Private/AnimationStateNodes/SGraphNodeAnimTransition.cpp b/Engine/Source/Editor/GraphEditor/Private/AnimationStateNodes/SGraphNodeAnimTransition.cpp
--- a/Engine/Source/Editor/GraphEditor/Private/AnimationStateNodes/SGraphNodeAnimTransition.cpp
+++ b/Engine/Source/Editor/GraphEditor/Private/AnimationStateNodes/SGraphNodeAnimTransition.cpp
@@ -12,6 +12,9 @@
 #include "ConnectionDrawingPolicy.h"
 #include "IDocumentation.h"
 
+#include "AnimationTransitionGraph.h"
+#include "AnimGraphNode_TransitionResult.h"
+
 #define LOCTEXT_NAMESPACE "TransitionNodes"
 
 /////////////////////////////////////////////////////
diff --git a/Engine/Source/Editor/GraphEditor/Private/ConnectionDrawingPolicy.cpp b/Engine/Source/Editor/GraphEditor/Private/ConnectionDrawingPolicy.cpp
--- a/Engine/Source/Editor/GraphEditor/Private/ConnectionDrawingPolicy.cpp
+++ b/Engine/Source/Editor/GraphEditor/Private/ConnectionDrawingPolicy.cpp
@@ -8,6 +8,10 @@
 #include "KismetNodes/KismetNodeInfoContext.h"
 #include "SoundDefinitions.h"
 #include "AnimationStateNodes/SGraphNodeAnimTransition.h"
+#include "AnimStateTransitionNode.h"
+#include "AnimStateEntryNode.h"
+#include "AnimationGraphSchema.h"
+#include "AnimGraphNode_Base.h"
 
 DEFINE_LOG_CATEGORY_STATIC(LogConnectionDrawingPolicy, Log, All);
 
diff --git a/Engine/Source/Editor/GraphEditor/Private/NodeFactory.cpp b/Engine/Source/Editor/GraphEditor/Private/NodeFactory.cpp
--- a/Engine/Source/Editor/GraphEditor/Private/NodeFactory.cpp
+++ b/Engine/Source/Editor/GraphEditor/Private/NodeFactory.cpp
@@ -26,6 +26,12 @@
 #include "KismetNodes/SGraphNodeFormatText.h"
 #include "KismetNodes/SGraphNodeK2ArrayFunction.h"
 
+#include "AnimGraphNode_Root.h"
+#include "AnimGraphNode_SequencePlayer.h"
+#include "AnimStateNode.h"
+
+#include "AnimationStateMachineSchema.h"
+
 #include "AnimationStateNodes/SGraphNodeAnimState.h"
 #include "AnimationStateNodes/SGraphNodeAnimTransition.h"
 #include "AnimationStateNodes/SGraphNodeAnimStateEntry.h"
diff --git a/Engine/Source/Editor/GraphEditor/Private/SGraphPanel.cpp b/Engine/Source/Editor/GraphEditor/Private/SGraphPanel.cpp
--- a/Engine/Source/Editor/GraphEditor/Private/SGraphPanel.cpp
+++ b/Engine/Source/Editor/GraphEditor/Private/SGraphPanel.cpp
@@ -15,6 +15,8 @@
 #include "KismetNodes/KismetNodeInfoContext.h"
 #include "GraphDiffControl.h"
 
+#include "AnimationGraphSchema.h"
+
 DEFINE_LOG_CATEGORY_STATIC(LogGraphPanel, Log, All);
 
 /**
diff --git a/Engine/Source/Editor/Kismet/Private/BlueprintEditor.cpp b/Engine/Source/Editor/Kismet/Private/BlueprintEditor.cpp
--- a/Engine/Source/Editor/Kismet/Private/BlueprintEditor.cpp
+++ b/Engine/Source/Editor/Kismet/Private/BlueprintEditor.cpp
@@ -4,7 +4,6 @@
 #include "ScopedTransaction.h"
 #include "Editor/UnrealEd/Public/Kismet2/KismetDebugUtilities.h"
 #include "Editor/UnrealEd/Public/Kismet2/BlueprintEditorUtils.h"
-#include "EngineLevelScriptClasses.h"
 #include "GraphEditor.h"
 #include "BlueprintUtilities.h"
 #include "Editor/KismetCompiler/Public/KismetCompilerModule.h"
@@ -61,6 +60,11 @@
 #include "BlueprintEditorTabFactories.h"
 #include "SPinTypeSelector.h"
 #include "AnimGraphDefinitions.h"
+#include "AnimStateConduitNode.h"
+#include "AnimationGraphSchema.h"
+#include "AnimationStateGraph.h"
+#include "AnimationStateMachineSchema.h"
+#include "AnimationTransitionGraph.h"
 #include "BlueprintEditorModes.h"
 
 #include "EngineAnalytics.h"
diff --git a/Engine/Source/Editor/Kismet/Private/SBlueprintPalette.cpp b/Engine/Source/Editor/Kismet/Private/SBlueprintPalette.cpp
--- a/Engine/Source/Editor/Kismet/Private/SBlueprintPalette.cpp
+++ b/Engine/Source/Editor/Kismet/Private/SBlueprintPalette.cpp
@@ -22,6 +22,11 @@
 #include "BlueprintPaletteFavorites.h"
 #include "STutorialWrapper.h"
 #include "ClassIconFinder.h"
+#include "AnimationStateMachineGraph.h"
+#include "AnimationStateMachineSchema.h"
+#include "AnimationStateGraph.h"
+#include "AnimStateConduitNode.h"
+#include "AnimationTransitionGraph.h"
 
 #define LOCTEXT_NAMESPACE "BlueprintPalette"
 
diff --git a/Engine/Source/Editor/KismetCompiler/Private/KismetCompiler.cpp b/Engine/Source/Editor/KismetCompiler/Private/KismetCompiler.cpp
--- a/Engine/Source/Editor/KismetCompiler/Private/KismetCompiler.cpp
+++ b/Engine/Source/Editor/KismetCompiler/Private/KismetCompiler.cpp
@@ -6,8 +6,6 @@
 
 #include "KismetCompilerPrivatePCH.h"
 #include "KismetCompilerBackend.h"
-#include "EngineKismetLibraryClasses.h"
-#include "EngineLevelScriptClasses.h"
 #include "Editor/UnrealEd/Public/Kismet2/KismetDebugUtilities.h"
 #include "Editor/UnrealEd/Public/Kismet2/KismetReinstanceUtilities.h"
 #include "Editor/UnrealEd/Public/Kismet2/BlueprintEditorUtils.h"
diff --git a/Engine/Source/Editor/LandscapeEditor/Private/LandscapeEdMode.cpp b/Engine/Source/Editor/LandscapeEditor/Private/LandscapeEdMode.cpp
--- a/Engine/Source/Editor/LandscapeEditor/Private/LandscapeEdMode.cpp
+++ b/Engine/Source/Editor/LandscapeEditor/Private/LandscapeEdMode.cpp
@@ -5,8 +5,6 @@
 #include "ObjectTools.h"
 #include "LandscapeEdMode.h"
 #include "ScopedTransaction.h"
-#include "EngineTerrainClasses.h"
-#include "EngineFoliageClasses.h"
 #include "Landscape/LandscapeEdit.h"
 #include "Landscape/LandscapeRender.h"
 #include "Landscape/LandscapeDataAccess.h"
diff --git a/Engine/Source/Editor/LandscapeEditor/Private/LandscapeEdModeBrushes.cpp b/Engine/Source/Editor/LandscapeEditor/Private/LandscapeEdModeBrushes.cpp
--- a/Engine/Source/Editor/LandscapeEditor/Private/LandscapeEdModeBrushes.cpp
+++ b/Engine/Source/Editor/LandscapeEditor/Private/LandscapeEdModeBrushes.cpp
@@ -5,7 +5,6 @@
 #include "ObjectTools.h"
 #include "LandscapeEdMode.h"
 #include "ScopedTransaction.h"
-#include "EngineTerrainClasses.h"
 #include "Runtime/Engine/Public/Landscape/LandscapeEdit.h"
 #include "Runtime/Engine/Public/Landscape/LandscapeRender.h"
 
diff --git a/Engine/Source/Editor/LandscapeEditor/Private/LandscapeEdModeComponentTools.cpp b/Engine/Source/Editor/LandscapeEditor/Private/LandscapeEdModeComponentTools.cpp
--- a/Engine/Source/Editor/LandscapeEditor/Private/LandscapeEdModeComponentTools.cpp
+++ b/Engine/Source/Editor/LandscapeEditor/Private/LandscapeEdModeComponentTools.cpp
@@ -4,8 +4,6 @@
 #include "ObjectTools.h"
 #include "LandscapeEdMode.h"
 #include "ScopedTransaction.h"
-#include "EngineTerrainClasses.h"
-#include "EngineFoliageClasses.h"
 #include "Landscape/LandscapeEdit.h"
 #include "Landscape/LandscapeRender.h"
 #include "Landscape/LandscapeDataAccess.h"
diff --git a/Engine/Source/Editor/LandscapeEditor/Private/LandscapeEdModeRampTool.cpp b/Engine/Source/Editor/LandscapeEditor/Private/LandscapeEdModeRampTool.cpp
--- a/Engine/Source/Editor/LandscapeEditor/Private/LandscapeEdModeRampTool.cpp
+++ b/Engine/Source/Editor/LandscapeEditor/Private/LandscapeEdModeRampTool.cpp
@@ -4,7 +4,6 @@
 //#include "ObjectTools.h"
 #include "LandscapeEdMode.h"
 #include "ScopedTransaction.h"
-#include "EngineTerrainClasses.h"
 //#include "EngineFoliageClasses.h"
 #include "Landscape/LandscapeEdit.h"
 #include "Landscape/LandscapeRender.h"
diff --git a/Engine/Source/Editor/Matinee/Private/Matinee.cpp b/Engine/Source/Editor/Matinee/Private/Matinee.cpp
--- a/Engine/Source/Editor/Matinee/Private/Matinee.cpp
+++ b/Engine/Source/Editor/Matinee/Private/Matinee.cpp
@@ -18,7 +18,6 @@
 #include "CameraController.h"
 #include "MatineeConstants.h"
 
-#include "EngineInterpolationClasses.h"
 #include "SubtitleManager.h"
 #include "InterpolationHitProxy.h"
 
diff --git a/Engine/Source/Editor/Matinee/Private/MatineeMenus.cpp b/Engine/Source/Editor/Matinee/Private/MatineeMenus.cpp
--- a/Engine/Source/Editor/Matinee/Private/MatineeMenus.cpp
+++ b/Engine/Source/Editor/Matinee/Private/MatineeMenus.cpp
@@ -7,7 +7,6 @@
 #include "MatineeModule.h"
 #include "Matinee.h"
 
-#include "EngineInterpolationClasses.h"
 #include "CameraController.h"
 #include "DesktopPlatformModule.h"
 #include "PackageTools.h"
diff --git a/Engine/Source/Editor/Persona/Private/AnimationEditorViewportClient.cpp b/Engine/Source/Editor/Persona/Private/AnimationEditorViewportClient.cpp
--- a/Engine/Source/Editor/Persona/Private/AnimationEditorViewportClient.cpp
+++ b/Engine/Source/Editor/Persona/Private/AnimationEditorViewportClient.cpp
@@ -8,7 +8,9 @@
 #include "SAnimViewportToolBar.h"
 #include "AnimViewportShowCommands.h"
 #include "AnimGraphDefinitions.h"
+#include "AnimPreviewInstance.h"
 #include "AnimationEditorViewportClient.h"
+#include "AnimGraphNode_SkeletalControlBase.h"
 #include "Runtime/Engine/Classes/Components/ReflectionCaptureComponent.h"
 #include "Runtime/Engine/Classes/Components/SphereReflectionCaptureComponent.h"
 #include "UnrealWidget.h"
diff --git a/Engine/Source/Editor/Persona/Private/Persona.cpp b/Engine/Source/Editor/Persona/Private/Persona.cpp
--- a/Engine/Source/Editor/Persona/Private/Persona.cpp
+++ b/Engine/Source/Editor/Persona/Private/Persona.cpp
@@ -40,7 +40,12 @@
 #include "Editor/Persona/Private/AnimationEditorViewportClient.h"
 
 #include "ComponentAssetBroker.h"
-
+#include "AnimGraphNode_BlendListByInt.h"
+#include "AnimGraphNode_BlendSpaceEvaluator.h"
+#include "AnimGraphNode_BlendSpacePlayer.h"
+#include "AnimGraphNode_LayeredBoneBlend.h"
+#include "AnimGraphNode_SequencePlayer.h"
+#include "AnimGraphNode_SequenceEvaluator.h"
 #define LOCTEXT_NAMESPACE "FPersona"
 
 /////////////////////////////////////////////////////
diff --git a/Engine/Source/Editor/PhAT/Classes/PhATEdSkeletalMeshComponent.h b/Engine/Source/Editor/PhAT/Classes/PhATEdSkeletalMeshComponent.h
--- a/Engine/Source/Editor/PhAT/Classes/PhATEdSkeletalMeshComponent.h
+++ b/Engine/Source/Editor/PhAT/Classes/PhATEdSkeletalMeshComponent.h
@@ -5,7 +5,7 @@
 #include "PhATEdSkeletalMeshComponent.generated.h"
 
 
-UCLASS(HeaderGroup=Component)
+UCLASS()
 class UPhATEdSkeletalMeshComponent : public USkeletalMeshComponent
 {
 	GENERATED_UCLASS_BODY()
diff --git a/Engine/Source/Editor/UnrealEd/Private/AssetSelection.cpp b/Engine/Source/Editor/UnrealEd/Private/AssetSelection.cpp
--- a/Engine/Source/Editor/UnrealEd/Private/AssetSelection.cpp
+++ b/Engine/Source/Editor/UnrealEd/Private/AssetSelection.cpp
@@ -2,8 +2,6 @@
 
 #include "UnrealEd.h"
 #include "AssetSelection.h"
-#include "EngineLightClasses.h"
-#include "EngineDecalClasses.h"
 #include "ScopedTransaction.h"
 
 #include "LevelUtils.h"
diff --git a/Engine/Source/Editor/UnrealEd/Private/Editor.cpp b/Engine/Source/Editor/UnrealEd/Private/Editor.cpp
--- a/Engine/Source/Editor/UnrealEd/Private/Editor.cpp
+++ b/Engine/Source/Editor/UnrealEd/Private/Editor.cpp
@@ -7,8 +7,6 @@
 
 // needed for the RemotePropagator
 #include "SoundDefinitions.h"
-#include "EngineFoliageClasses.h"
-#include "EngineDecalClasses.h"
 #include "Database.h"
 #include "SurfaceIterators.h"
 #include "ScopedTransaction.h"
diff --git a/Engine/Source/Editor/UnrealEd/Private/EditorActor.cpp b/Engine/Source/Editor/UnrealEd/Private/EditorActor.cpp
--- a/Engine/Source/Editor/UnrealEd/Private/EditorActor.cpp
+++ b/Engine/Source/Editor/UnrealEd/Private/EditorActor.cpp
@@ -4,7 +4,6 @@
 #include "UnrealEd.h"
 #include "ScopedTransaction.h"
 #include "Factories.h"
-#include "EngineFoliageClasses.h"
 #include "LevelUtils.h"
 #include "BusyCursor.h"
 #include "BSPOps.h"
@@ -14,7 +13,6 @@
 #include "Editor/GeometryMode/Public/GeometryEdMode.h"
 #include "Editor/GeometryMode/Public/EditorGeometry.h"
 #include "ActorEditorUtils.h"
-#include "EngineLevelScriptClasses.h"
 
 #define LOCTEXT_NAMESPACE "UnrealEd.EditorActor"
 
diff --git a/Engine/Source/Editor/UnrealEd/Private/EditorAnimUtils.cpp b/Engine/Source/Editor/UnrealEd/Private/EditorAnimUtils.cpp
--- a/Engine/Source/Editor/UnrealEd/Private/EditorAnimUtils.cpp
+++ b/Engine/Source/Editor/UnrealEd/Private/EditorAnimUtils.cpp
@@ -1,6 +1,7 @@
 // Copyright 1998-2014 Epic Games, Inc. All Rights Reserved.
 
 #include "UnrealEd.h"
+#include "AnimGraphNode_Base.h"
 #include "AssetData.h"
 #include "EditorAnimUtils.h"
 #include "BlueprintEditorUtils.h"
@@ -10,6 +11,7 @@
 #include "NotificationManager.h"
 #include "Editor/Persona/Public/PersonaModule.h"
 
+
 #define LOCTEXT_NAMESPACE "EditorAnimUtils"
 
 namespace EditorAnimUtils
diff --git a/Engine/Source/Editor/UnrealEd/Private/EditorExporters.cpp b/Engine/Source/Editor/UnrealEd/Private/EditorExporters.cpp
--- a/Engine/Source/Editor/UnrealEd/Private/EditorExporters.cpp
+++ b/Engine/Source/Editor/UnrealEd/Private/EditorExporters.cpp
@@ -6,8 +6,6 @@
 
 #include "UnrealEd.h"
 #include "SoundDefinitions.h"
-#include "EngineMaterialClasses.h"
-#include "EngineFoliageClasses.h"
 #include "Landscape/LandscapeDataAccess.h"
 #include "Kismet2/DebuggerCommands.h"
 #include "MaterialCompiler.h"
diff --git a/Engine/Source/Editor/UnrealEd/Private/Factories/ActorFactory.cpp b/Engine/Source/Editor/UnrealEd/Private/Factories/ActorFactory.cpp
--- a/Engine/Source/Editor/UnrealEd/Private/Factories/ActorFactory.cpp
+++ b/Engine/Source/Editor/UnrealEd/Private/Factories/ActorFactory.cpp
@@ -6,12 +6,8 @@ ActorFactory.cpp:
 
 #include "UnrealEd.h"
 #include "ParticleDefinitions.h"
-#include "EngineMaterialClasses.h"
-#include "EngineDecalClasses.h"
-#include "EngineFoliageClasses.h"
 #include "SoundDefinitions.h"
 #include "BlueprintUtilities.h"
-#include "EngineDecalClasses.h"
 #include "Kismet2/BlueprintEditorUtils.h"
 #include "Kismet2/KismetEditorUtilities.h"
 #include "AssetData.h"
diff --git a/Engine/Source/Editor/UnrealEd/Private/Factories/AnimBlueprintFactory.cpp b/Engine/Source/Editor/UnrealEd/Private/Factories/AnimBlueprintFactory.cpp
--- a/Engine/Source/Editor/UnrealEd/Private/Factories/AnimBlueprintFactory.cpp
+++ b/Engine/Source/Editor/UnrealEd/Private/Factories/AnimBlueprintFactory.cpp
@@ -7,13 +7,9 @@
 #include "UnrealEd.h"
 #include "Factories.h"
 #include "SoundDefinitions.h"
-#include "EngineParticleClasses.h"
 #include "BlueprintUtilities.h"
 #include "Kismet2/KismetEditorUtilities.h"
-#include "EngineInterpolationClasses.h"
-#include "EngineFoliageClasses.h"
 #include "BusyCursor.h"
-#include "EngineMaterialClasses.h"
 #include "BSPOps.h"
 #include "LevelUtils.h"
 #include "ObjectTools.h"
@@ -27,7 +23,6 @@
 #include "ClassViewerModule.h"
 #include "ClassViewerFilter.h"
 #include "AssetRegistryModule.h"
-#include "EngineLevelScriptClasses.h"		// So we can filter out ALevelScriptActors
 
 #define LOCTEXT_NAMESPACE "AnimBlueprintFactory"
 
diff --git a/Engine/Source/Editor/UnrealEd/Private/Factories/ApexDestructibleAssetImport.cpp b/Engine/Source/Editor/UnrealEd/Private/Factories/ApexDestructibleAssetImport.cpp
--- a/Engine/Source/Editor/UnrealEd/Private/Factories/ApexDestructibleAssetImport.cpp
+++ b/Engine/Source/Editor/UnrealEd/Private/Factories/ApexDestructibleAssetImport.cpp
@@ -16,9 +16,7 @@ DEFINE_LOG_CATEGORY_STATIC(LogApexDestructibleAssetImport, Log, All);
 
 #include "Engine.h"
 #include "TextureLayout.h"
-#include "EngineMaterialClasses.h"
 #include "SkelImport.h"
-#include "EngineInterpolationClasses.h"
 #include "EditorPhysXSupport.h"
 #include "ApexDestructibleAssetImport.h"
 #include "Developer/MeshUtilities/Public/MeshUtilities.h"
diff --git a/Engine/Source/Editor/UnrealEd/Private/Factories/EditorFactories.cpp b/Engine/Source/Editor/UnrealEd/Private/Factories/EditorFactories.cpp
--- a/Engine/Source/Editor/UnrealEd/Private/Factories/EditorFactories.cpp
+++ b/Engine/Source/Editor/UnrealEd/Private/Factories/EditorFactories.cpp
@@ -7,15 +7,11 @@
 #include "UnrealEd.h"
 #include "Factories.h"
 #include "SoundDefinitions.h"
-#include "EngineParticleClasses.h"
 #include "BlueprintUtilities.h"
 #include "Kismet2/KismetEditorUtilities.h"
-#include "EngineInterpolationClasses.h"
-#include "EngineFoliageClasses.h"
 #include "BmpImageSupport.h"
 #include "ScopedTransaction.h"
 #include "BusyCursor.h"
-#include "EngineMaterialClasses.h"
 #include "BSPOps.h"
 #include "LevelUtils.h"
 #include "ObjectTools.h"
@@ -33,7 +29,6 @@
 #include "ContentBrowserModule.h"
 #include "ClassViewerModule.h"
 #include "ClassViewerFilter.h"
-#include "EngineLevelScriptClasses.h"		// So we can filter out ALevelScriptActors
 #include "SClassPickerDialog.h"
 #include "MessageLog.h"
 #include "EnumEditorUtils.h"
diff --git a/Engine/Source/Editor/UnrealEd/Private/Fbx/FbxLightImport.cpp b/Engine/Source/Editor/UnrealEd/Private/Fbx/FbxLightImport.cpp
--- a/Engine/Source/Editor/UnrealEd/Private/Fbx/FbxLightImport.cpp
+++ b/Engine/Source/Editor/UnrealEd/Private/Fbx/FbxLightImport.cpp
@@ -7,7 +7,6 @@
 #include "UnrealEd.h"
 #include "Factories.h"
 #include "Engine.h"
-#include "EngineMaterialClasses.h"
 #include "FbxImporter.h"
 
 DEFINE_LOG_CATEGORY_STATIC(LogFbxLightImport, Log, All);
diff --git a/Engine/Source/Editor/UnrealEd/Private/Fbx/FbxMaterialImport.cpp b/Engine/Source/Editor/UnrealEd/Private/Fbx/FbxMaterialImport.cpp
--- a/Engine/Source/Editor/UnrealEd/Private/Fbx/FbxMaterialImport.cpp
+++ b/Engine/Source/Editor/UnrealEd/Private/Fbx/FbxMaterialImport.cpp
@@ -32,7 +32,6 @@
 #include "Factories.h"
 #include "Engine.h"
 #include "TextureLayout.h"
-#include "EngineMaterialClasses.h"
 #include "FbxImporter.h"
 #include "ObjectTools.h"
 #include "PackageTools.h"
diff --git a/Engine/Source/Editor/UnrealEd/Private/Fbx/FbxSkeletalMeshImport.cpp b/Engine/Source/Editor/UnrealEd/Private/Fbx/FbxSkeletalMeshImport.cpp
--- a/Engine/Source/Editor/UnrealEd/Private/Fbx/FbxSkeletalMeshImport.cpp
+++ b/Engine/Source/Editor/UnrealEd/Private/Fbx/FbxSkeletalMeshImport.cpp
@@ -37,9 +37,7 @@
 
 #include "Engine.h"
 #include "TextureLayout.h"
-#include "EngineMaterialClasses.h"
 #include "SkelImport.h"
-#include "EngineInterpolationClasses.h"
 #include "FbxImporter.h"
 #include "AnimEncoding.h"
 #include "SSkeletonWidget.h"
diff --git a/Engine/Source/Editor/UnrealEd/Private/Kismet2/BlueprintEditorUtils.cpp b/Engine/Source/Editor/UnrealEd/Private/Kismet2/BlueprintEditorUtils.cpp
--- a/Engine/Source/Editor/UnrealEd/Private/Kismet2/BlueprintEditorUtils.cpp
+++ b/Engine/Source/Editor/UnrealEd/Private/Kismet2/BlueprintEditorUtils.cpp
@@ -3,11 +3,18 @@
 #include "UnrealEd.h"
 
 #include "LatentActions.h"
-#include "EngineLevelScriptClasses.h"
 #include "Kismet2/CompilerResultsLog.h"
 
 #include "Editor/KismetCompiler/Public/KismetCompilerModule.h"
 #include "AnimGraphDefinitions.h"
+#include "AnimGraphNode_StateMachineBase.h"
+#include "AnimStateNode.h"
+#include "AnimStateNodeBase.h"
+#include "AnimStateTransitionNode.h"
+#include "AnimationTransitionSchema.h"
+#include "AnimationGraph.h"
+#include "AnimationGraphSchema.h"
+#include "AnimationStateMachineGraph.h"
 #include "Editor/UnrealEd/Public/Kismet2/KismetEditorUtilities.h"
 #include "Editor/UnrealEd/Public/Kismet2/KismetDebugUtilities.h"
 #include "Editor/UnrealEd/Public/Kismet2/BlueprintEditorUtils.h"
@@ -21,7 +28,6 @@
 #include "DefaultValueHelper.h"
 #include "ObjectEditorUtils.h"
 #include "ToolkitManager.h"
-
 #define LOCTEXT_NAMESPACE "Blueprint"
 
 DEFINE_LOG_CATEGORY(LogBlueprintDebug);
diff --git a/Engine/Source/Editor/UnrealEd/Private/Kismet2/Kismet2.cpp b/Engine/Source/Editor/UnrealEd/Private/Kismet2/Kismet2.cpp
--- a/Engine/Source/Editor/UnrealEd/Private/Kismet2/Kismet2.cpp
+++ b/Engine/Source/Editor/UnrealEd/Private/Kismet2/Kismet2.cpp
@@ -4,6 +4,8 @@
 #include "UnrealEd.h"
 #include "BlueprintUtilities.h"
 #include "AnimGraphDefinitions.h"
+#include "AnimationGraph.h"
+#include "AnimationGraphSchema.h"
 #include "Kismet2/KismetEditorUtilities.h"
 #include "Kismet2/KismetReinstanceUtilities.h"
 #include "Kismet2/BlueprintEditorUtils.h"
diff --git a/Engine/Source/Editor/UnrealEd/Private/Kismet2/Kismet2NameValidators.cpp b/Engine/Source/Editor/UnrealEd/Private/Kismet2/Kismet2NameValidators.cpp
--- a/Engine/Source/Editor/UnrealEd/Private/Kismet2/Kismet2NameValidators.cpp
+++ b/Engine/Source/Editor/UnrealEd/Private/Kismet2/Kismet2NameValidators.cpp
@@ -4,6 +4,8 @@
 #include "Slate.h"
 #include "Kismet2NameValidators.h"
 #include "AnimGraphDefinitions.h"
+#include "AnimStateTransitionNode.h"
+#include "AnimationStateMachineGraph.h"
 #include "KismetEditorUtilities.h"
 #include "KismetDebugUtilities.h"
 #include "BlueprintEditorUtils.h"
diff --git a/Engine/Source/Editor/UnrealEd/Private/LevelEditorViewport.cpp b/Engine/Source/Editor/UnrealEd/Private/LevelEditorViewport.cpp
--- a/Engine/Source/Editor/UnrealEd/Private/LevelEditorViewport.cpp
+++ b/Engine/Source/Editor/UnrealEd/Private/LevelEditorViewport.cpp
@@ -2,9 +2,7 @@
 
 
 #include "UnrealEd.h"
-#include "EngineMaterialClasses.h"
 #include "SoundDefinitions.h"
-#include "EngineSplineClasses.h"
 #include "CameraController.h"
 #include "MouseDeltaTracker.h"
 #include "ScopedTransaction.h"
diff --git a/Engine/Source/Editor/UnrealEd/Private/StaticLightingSystem/StaticLightingExport.cpp b/Engine/Source/Editor/UnrealEd/Private/StaticLightingSystem/StaticLightingExport.cpp
--- a/Engine/Source/Editor/UnrealEd/Private/StaticLightingSystem/StaticLightingExport.cpp
+++ b/Engine/Source/Editor/UnrealEd/Private/StaticLightingSystem/StaticLightingExport.cpp
@@ -8,10 +8,9 @@
 
 #include "StaticMeshLight.h"
 #include "ModelLight.h"
-#include "EngineFoliageClasses.h"
 #include "Runtime/Engine/Public/Landscape/LandscapeLight.h"
 
-#include "../Lightmass/Lightmass.h"
+#include "Lightmass/Lightmass.h"
 
 // Doxygen cannot parse these definitions correctly since the declaration is in a header from a different module
 #if !UE_BUILD_DOCS
diff --git a/Engine/Source/Editor/UnrealEd/Public/UnrealEd.h b/Engine/Source/Editor/UnrealEd/Public/UnrealEd.h
--- a/Engine/Source/Editor/UnrealEd/Public/UnrealEd.h
+++ b/Engine/Source/Editor/UnrealEd/Public/UnrealEd.h
@@ -11,8 +11,6 @@
 
 #include "Engine.h"
 
-#include "EngineUserInterfaceClasses.h"
-
 #include "Slate.h"
 #include "EditorStyle.h"
 
diff --git a/Engine/Source/Programs/UnrealHeaderTool/Private/CodeGenerator.cpp b/Engine/Source/Programs/UnrealHeaderTool/Private/CodeGenerator.cpp
--- a/Engine/Source/Programs/UnrealHeaderTool/Private/CodeGenerator.cpp
+++ b/Engine/Source/Programs/UnrealHeaderTool/Private/CodeGenerator.cpp
@@ -2052,7 +2052,7 @@ void FNativeClassHeaderGenerator::ExportClassHeaderWrapper( UClass* Class, bool
 void FNativeClassHeaderGenerator::ExportClassHeaderRecursive( UClass* Class, TArray<UClass*>& DependencyChain, bool bCheckDependenciesOnly )
 {
 	bool bIsExportClass   = GClassStrippedHeaderTextMap.Contains(Class) && Class->HasAnyClassFlags(CLASS_Native|CLASS_Temporary) && !Class->HasAnyClassFlags(CLASS_NoExport);
-	bool bIsCorrectHeader = GClassHeaderFilenameMap.FindRef(Class) == ClassHeaderFilename && Class->GetOuter() == Package;
+	bool bIsCorrectHeader = Class->GetOuter() == Package;
 
 	// Check for circular header dependencies between export classes.
 	if ( bIsExportClass )
@@ -2070,17 +2070,6 @@ void FNativeClassHeaderGenerator::ExportClassHeaderRecursive( UClass* Class, TAr
 				bCheckDependenciesOnly = true;
 			}
 		}
-		else if ( bCheckDependenciesOnly && !ClassHeaderFilename.IsEmpty() )
-		{
-			FString DependencyChainString;
-			for ( int32 DependencyIndex=0; DependencyIndex < DependencyChain.Num(); DependencyIndex++ )
-			{
-				UClass* DependencyClass = DependencyChain[DependencyIndex];
-				DependencyChainString += DependencyClass->GetName() + TEXT(" -> ");
-			}
-			DependencyChainString += Class->GetName();
-			UE_LOG(LogCompile, Warning, TEXT("Circular header dependency detected (%s), while exporting %s!"), *DependencyChainString, *(FPackageName::GetShortName(Package) + ClassHeaderFilename + TEXT("Classes.h")) );
-		}
 	}
 
 	// Check if the Class has already been exported, after we've checked for circular header dependencies.
@@ -3744,53 +3733,6 @@ void FNativeClassHeaderGenerator::ApplyAlternatePropertyExportText( UProperty* P
 	}
 }
 
-/**
- * Sorts the list of header files being exported from a package according to their dependency on each other.
- *
- * @param	HeaderDependencyMap		a mapping of header filenames to a list of header filenames that must be processed before that one.
- * @param	SortedHeaderFilenames	[out] receives the sorted list of header filenames.
- */
-bool FNativeClassHeaderGenerator::SortHeaderDependencyMap(const TMap<const FString*, HeaderDependents>& HeaderDependencyMap, TArray<const FString*>& SortedHeaderFilenames) const
-{
-	SortedHeaderFilenames.Empty(HeaderDependencyMap.Num());
-
-	while (SortedHeaderFilenames.Num() < HeaderDependencyMap.Num())
-	{
-		bool bAddedSomething = false;
-
-		// Find headers with no dependencies and add those to the list.
-		for (auto It = HeaderDependencyMap.CreateConstIterator(); It; ++It)
-		{
-			auto Header = It->Key;
-			if (SortedHeaderFilenames.Contains(Header))
-				continue;
-
-			bool bHasRemainingDependencies = false;
-			for (auto It2 = It->Value.CreateConstIterator(); It2; ++It2)
-			{
-				if (!SortedHeaderFilenames.Contains(*It2))
-				{
-					bHasRemainingDependencies = true;
-					break;
-				}
-			}
-
-			if (!bHasRemainingDependencies)
-			{
-				// Add it to the list.
-				SortedHeaderFilenames.AddUnique(Header);
-				bAddedSomething           = true;
-			}
-		}
-
-		// Circular dependency error?
-		if (!bAddedSomething)
-			return false;
-	}
-
-	return true;
-}
-
 bool FNativeClassHeaderGenerator::FindInterDependency( TMap<const FString*, HeaderDependents>& HeaderDependencyMap, const FString* Header, const FString*& OutHeader1, const FString*& OutHeader2 )
 {
 	TSet<const FString*> VisitedHeaders;
@@ -3828,72 +3770,6 @@ bool FNativeClassHeaderGenerator::FindInterDependencyRecursive( TMap<const FStri
 	return false;
 }
 
-/**
- * Finds a dependency chain between two class header files.
- * Wrapper around FindDependencyChainRecursive().
- *
- * @param	Class				A class to scan for a dependency chain between the two headers.
- * @param	Header1				First class header filename.
- * @param	Header2				Second class header filename.
- * @param	DependencyChain		[out] Receives dependency chain, if found.
- * @return	true if a dependency chain was found and filled in.
- */
-bool FNativeClassHeaderGenerator::FindDependencyChain( FClass* Class, const FString& Header1, const FString& Header2, TArray<FClass*>& DependencyChain )
-{
-	DependencyChain.Empty();
-	return FindDependencyChainRecursive( Class, Header1, Header2, false, DependencyChain );
-}
-
-/**
- * Finds a dependency chain between two class header files.
- *
- * @param	Class				A class to scan for a dependency chain between the two headers.
- * @param	Header1				First class header filename.
- * @param	Header2				Second class header filename.
- * @param	bChainStarted		Whether Header1 has been found and we've started to fill in DependencyChain. Must be false to begin with.
- * @param	DependencyChain		[out] Receives dependency chain, if found. Must be empty before the call.
- * @return	true if a dependency chain was found and filled in.
- */
-bool FNativeClassHeaderGenerator::FindDependencyChainRecursive( FClass* Class, const FString& Header1, const FString& Header2, bool bChainStarted, TArray<FClass*>& DependencyChain )
-{
-	bool bIsExportClass = Class->HasAnyClassFlags(CLASS_Native) && !Class->HasAnyClassFlags(CLASS_NoExport|CLASS_Intrinsic|CLASS_Temporary);
-	if ( bIsExportClass )
-	{
-		auto HeaderFilename = GClassHeaderFilenameMap.FindRef(Class);
-
-		bChainStarted = bChainStarted || HeaderFilename == Header1;
-
-		if ( bChainStarted )
-		{
-			DependencyChain.Add( Class );
-			if ( HeaderFilename == Header2 )
-			{
-				return true;
-			}
-		}
-	}
-
-	FClass* SuperClass = Class->GetSuperClass();
-	if ( SuperClass && SuperClass->GetOuter() == Class->GetOuter() && FindDependencyChainRecursive( SuperClass, Header1, Header2, bChainStarted, DependencyChain ) )
-	{
-		return true;
-	}
-
-	for (FClass* InterfaceClass : Class->GetInterfaceTypes())
-	{
-		if ( InterfaceClass->GetOuter() == Class->GetOuter() && FindDependencyChainRecursive( InterfaceClass, Header1, Header2, bChainStarted, DependencyChain ) )
-		{
-			return true;
-		}
-	}
-
-	if ( bIsExportClass && bChainStarted )
-	{
-		DependencyChain.RemoveAtSwap( DependencyChain.Num() - 1 );
-	}
-	return false;
-}
-
 // Constructor.
 FNativeClassHeaderGenerator::FNativeClassHeaderGenerator( UPackage* InPackage, FClasses& AllClasses, bool InAllowSaveExportedHeaders, bool bInUseRelativePaths )
 	: CurrentClass                          (NULL)
@@ -3905,7 +3781,6 @@ FNativeClassHeaderGenerator::FNativeClassHeaderGenerator( UPackage* InPackage, F
 	, bUseRelativePaths                     (bInUseRelativePaths)
 {
 	const FString PackageName = FPackageName::GetShortName(Package);
-
 	GeneratedCPPFilenameBase =  PackageName + TEXT(".generated");
 	GeneratedProtoFilenameBase = PackageName + TEXT(".generated");
 	GeneratedMCPFilenameBase = PackageName + TEXT(".generated");
@@ -3915,148 +3790,15 @@ FNativeClassHeaderGenerator::FNativeClassHeaderGenerator( UPackage* InPackage, F
 
 	Classes = AllClasses.GetClassesInPackage();
 
-	TMap<FName,UClass*> ClassNameMap;
-	for ( int32 ClassIndex = 0; ClassIndex < Classes.Num(); ClassIndex++ )
-	{
-		ClassNameMap.Add(Classes[ClassIndex]->GetFName(), Classes[ClassIndex]);
-	}
-
-	class UnsortedHeaderFilenamesType
+	int ExportAndTemporaryClassCount = 0;
+	for (FClass* Cls : AllClasses.GetClassesInPackage(Package))
 	{
-	public:
-		FString* AddUnique(const FString& Str)
+		bool bIsExportClass = Cls->HasAnyClassFlags(CLASS_Native) && !Cls->HasAnyClassFlags(CLASS_NoExport | CLASS_Intrinsic | CLASS_Temporary);
+		bool bIsTemporaryClass = Cls->HasAnyClassFlags(CLASS_Temporary);
+		if (bIsExportClass || bIsTemporaryClass)
 		{
-			for (auto It = Data.CreateConstIterator(); It; ++It)
-			{
-				auto& Ptr = *It;
-
-				if (*Ptr == Str)
-					return &*Ptr;
-			}
-
-			return Data[Data.Add(MakeUnique<FString>(Str))].Get();
-		}
-
-		TArray<TUniquePtr<FString>> Data;
-	};
-
-	UnsortedHeaderFilenamesType UnsortedHeaderFilenames;
-	TArray<const FString*>      HeaderFilenames; // These reference the strings stored in UnsortedHeaderFilenames
-
-	bool bCircularDependencyDetected = false;
- 	{
-	 	TMap<const FString*, HeaderDependents> HeaderDependencyMap;
-		TArray<FClass*> ClassesInPackage = AllClasses.GetClassesInPackage(Package);
-		for (FClass* Cls : ClassesInPackage)
-		{
-			auto ClsHeaderFilename = GClassHeaderFilenameMap.FindRef(Cls);
-
-			bool bIsExportClass = Cls->HasAnyClassFlags(CLASS_Native) && !Cls->HasAnyClassFlags(CLASS_NoExport|CLASS_Intrinsic|CLASS_Temporary);
-			if ( bIsExportClass )
-			{
-				auto* HeaderFilename = UnsortedHeaderFilenames.AddUnique(ClsHeaderFilename);
-				auto& Dependencies   = HeaderDependencyMap.FindOrAdd(HeaderFilename);
-
-				// Add the super class' header as a dependency if it's different.
-				UClass* SuperClass = Cls->GetSuperClass();
-				bool bIsDependentExportClass = SuperClass && SuperClass->HasAnyClassFlags(CLASS_Native) && !SuperClass->HasAnyClassFlags(CLASS_NoExport|CLASS_Intrinsic|CLASS_Temporary);
-				if ( bIsDependentExportClass && SuperClass->GetOuter() == Cls->GetOuter() )
-				{
-					auto SuperClassHeaderFilename = GClassHeaderFilenameMap.FindRef(SuperClass);
-					if ( SuperClassHeaderFilename != ClsHeaderFilename )
-					{
-						auto* DependentHeaderFilename = UnsortedHeaderFilenames.AddUnique(SuperClassHeaderFilename);
-						Dependencies.AddUnique(DependentHeaderFilename);
-					}
-				}
-
-				// Add base interface headers as dependencies, if they're different.
-				for (TArray<FImplementedInterface>::TIterator It(Cls->Interfaces); It; ++It)
-				{
-					UClass* InterfaceClass = It->Class;
-					bool bIsDependentExportClass = InterfaceClass->HasAnyClassFlags(CLASS_Native) && !InterfaceClass->HasAnyClassFlags(CLASS_NoExport|CLASS_Intrinsic|CLASS_Temporary);
-					if ( bIsDependentExportClass && InterfaceClass->GetOuter() == Cls->GetOuter() )
-					{
-						auto InterfaceClassHeaderFilename = GClassHeaderFilenameMap.FindRef(InterfaceClass);
-						if ( InterfaceClassHeaderFilename != ClsHeaderFilename )
-						{
-							auto* DependentHeaderFilename = UnsortedHeaderFilenames.AddUnique(InterfaceClassHeaderFilename);
-							Dependencies.AddUnique(DependentHeaderFilename);
-						}
-					}
-				}
-			}
-			else if (Cls->HasAnyClassFlags(CLASS_Temporary))
-			{
-				// Struct only headers can generate header groups too so make sure their added too.
-				auto* HeaderFilename = UnsortedHeaderFilenames.AddUnique(ClsHeaderFilename);
-				HeaderDependencyMap.FindOrAdd(HeaderFilename);
-			}
-		}
-		bCircularDependencyDetected = !SortHeaderDependencyMap(HeaderDependencyMap, HeaderFilenames);
-		if (bCircularDependencyDetected)
-		{
-			// Find one circular path (though there may be multiple).
-			for (auto HeaderIt = HeaderDependencyMap.CreateConstIterator(); HeaderIt; ++HeaderIt)
-			{
-				const FString* Header = HeaderIt->Key;
-
-				const FString* ClassHeaderFilename1;
-				const FString* ClassHeaderFilename2;
-				if (!FindInterDependency(HeaderDependencyMap, Header, ClassHeaderFilename1, ClassHeaderFilename2))
-					continue;
-
-				TArray<FClass*> DependencyChain1;
-				for (FClass* Class : ClassesInPackage)
-				{
-					FindDependencyChain( Class, *ClassHeaderFilename1, *ClassHeaderFilename2, DependencyChain1 );
-
-					if (DependencyChain1.Num())
-						break;
-				}
-
-				TArray<FClass*> DependencyChain2;
-				for (FClass* Class : ClassesInPackage)
-				{
-					FindDependencyChain( Class, *ClassHeaderFilename2, *ClassHeaderFilename1, DependencyChain2 );
-
-					if (DependencyChain2.Num())
-						break;
-				}
-
-				if (!DependencyChain1.Num() || !DependencyChain2.Num())
-					continue;
-
-				FString DependencyChainString1;
-				for (FClass* Dependency : DependencyChain1)
-				{
-					if (!DependencyChainString1.IsEmpty())
-					{
-						DependencyChainString1 += TEXT(" -> ");
-					}
-
-					DependencyChainString1 += Dependency->GetName();
-				}
-
-				FString DependencyChainString2;
-				for (FClass* Dependency : DependencyChain2)
-				{
-					if (!DependencyChainString1.IsEmpty())
-					{
-						DependencyChainString2 += TEXT(" -> ");
-					}
-
-					DependencyChainString2 += Dependency->GetName();
-				}
-
-				UE_LOG(LogCompile, Error, TEXT("Header interdependency: %s <-> %s (%s and %s)."),
-					*(PackageName + *ClassHeaderFilename1 + TEXT("Classes.h")),
-					*(PackageName + *ClassHeaderFilename2 + TEXT("Classes.h")),
-					*DependencyChainString1,
-					*DependencyChainString2 );
-				break;
-			}
-			UE_LOG(LogCompile, Error, TEXT("Interdependent headers detected - aborting!") );
+			ExportAndTemporaryClassCount = 1;
+			break;
 		}
 	}
 
@@ -4072,24 +3814,20 @@ FNativeClassHeaderGenerator::FNativeClassHeaderGenerator( UPackage* InPackage, F
 	GeneratedFunctionDeclarations.Empty();
 	GeneratedFunctionBodyTextSplit.Empty();
 	
-	for (auto It = HeaderFilenames.CreateConstIterator(); !bCircularDependencyDetected && It; ++It)
+	if (ExportAndTemporaryClassCount == 1)
 	{
-		ClassHeaderFilename = **It;
-
 		FString PkgDir;
 		FString GeneratedIncludeDirectory;
 		if (MakeCommandlet_FindPackageLocation(*PackageName, PkgDir, GeneratedIncludeDirectory) == false)
 		{
 			UE_LOG(LogCompile, Error, TEXT("Failed to find path for package %s"), *PackageName);
 		}
-		ClassesHeaderPath = GeneratedIncludeDirectory / PackageName + ClassHeaderFilename + TEXT("Classes.h");
+		ClassesHeaderPath = GeneratedIncludeDirectory / PackageName + TEXT("Classes.h");
 
 		int32 ClassCount = 0;
-		for( int32 ClassIndex = 0; ClassIndex < Classes.Num(); ClassIndex++ )
+		for (UClass* Class : Classes)
 		{
-			UClass* Class = Classes[ClassIndex];
-
-			if( Class->GetOuter()==Package && Class->HasAnyClassFlags(CLASS_Native|CLASS_Temporary) && GClassHeaderFilenameMap.FindRef(Class) == ClassHeaderFilename )
+			if (Class->GetOuter() == Package && Class->HasAnyClassFlags(CLASS_Native | CLASS_Temporary))
 			{
 				if (GClassStrippedHeaderTextMap.Contains(Class) && !Class->HasAnyClassFlags(CLASS_NoExport))
 				{
@@ -4108,79 +3846,69 @@ FNativeClassHeaderGenerator::FNativeClassHeaderGenerator( UPackage* InPackage, F
 			}
 		}
 
-		if( ClassCount == 0 )
+		if (ClassCount != 0)
 		{
-			continue;
-		}
-
-		CurrentClass = NULL;
-		ListOfPublicClassesUObjectHeaderIncludes.Empty();
-		ListOfAllUObjectHeaderIncludes.Empty();
-		OriginalHeader.Empty();
-		PreHeaderText.Empty();
+			CurrentClass = NULL;
+			ListOfPublicClassesUObjectHeaderIncludes.Empty();
+			ListOfAllUObjectHeaderIncludes.Empty();
+			OriginalHeader.Empty();
+			PreHeaderText.Empty();
 
-		// Load the original header file into memory
-		FFileHelper::LoadFileToString(OriginalHeader,*ClassesHeaderPath);
-		
-		UE_LOG(LogCompile, Log,  TEXT("Autogenerating C++ header: %s"), *ClassHeaderFilename );
+			// Load the original header file into memory
+			FFileHelper::LoadFileToString(OriginalHeader, *ClassesHeaderPath);
 
-		// Write the classes and enums header prefixes.
+			// Write the classes and enums header prefixes.
 
-		PreHeaderText.Logf(
-			TEXT("// Copyright 1998-2014 Epic Games, Inc. All Rights Reserved.\r\n")
-			TEXT("/*===========================================================================\r\n")
-			TEXT("    C++ class boilerplate exported from UnrealHeaderTool.\r\n")
-			TEXT("    This is automatically generated by the tools.\r\n")
-			TEXT("    DO NOT modify this manually! Edit the corresponding .h files instead!\r\n")
-			TEXT("===========================================================================*/\r\n")
-			TEXT("#pragma once\r\n")
-			TEXT("\r\n")
-			);
+			PreHeaderText.Logf(
+				TEXT("// Copyright 1998-2014 Epic Games, Inc. All Rights Reserved.\r\n")
+				TEXT("/*===========================================================================\r\n")
+				TEXT("    C++ class boilerplate exported from UnrealHeaderTool.\r\n")
+				TEXT("    This is automatically generated by the tools.\r\n")
+				TEXT("    DO NOT modify this manually! Edit the corresponding .h files instead!\r\n")
+				TEXT("===========================================================================*/\r\n")
+				TEXT("#pragma once\r\n")
+				TEXT("\r\n")
+				);
 
-		// if a global auto-include file exists, generate a line to have that file included
-		FString GlobalAutoIncludeFilename = PackageName + ClassHeaderFilename + TEXT("GlobalIncludes.h");
-		const FString StandardHeaderFileLocation = PkgDir / TEXT("Public");
-		if ( IFileManager::Get().FileSize(*(StandardHeaderFileLocation / GlobalAutoIncludeFilename)) > 0 )
-		{
-			PreHeaderText.Logf(TEXT("#include \"%s\"\r\n\r\n"), *GlobalAutoIncludeFilename);
-		}
+			// if a global auto-include file exists, generate a line to have that file included
+			FString GlobalAutoIncludeFilename = PackageName + TEXT("GlobalIncludes.h");
+			const FString StandardHeaderFileLocation = PkgDir / TEXT("Public");
+			if (IFileManager::Get().FileSize(*(StandardHeaderFileLocation / GlobalAutoIncludeFilename)) > 0)
+			{
+				PreHeaderText.Logf(TEXT("#include \"%s\"\r\n\r\n"), *GlobalAutoIncludeFilename);
+			}
 
-		FString HeaderAPI = API;
-		if( ClassHeaderFilename.Len() > 0 )
-		{
-			HeaderAPI = HeaderAPI + TEXT("_") + ClassHeaderFilename.ToUpper();
-		}
+			FString HeaderAPI = API;
 
-		// Export an include line for each header
-		for ( int32 ClassIndex = 0; ClassIndex < Classes.Num(); ClassIndex++ )
-		{
-			UClass* Class = Classes[ClassIndex];
-			if ( GClassHeaderFilenameMap.FindRef(Class) == ClassHeaderFilename && Class->GetOuter() == Package )
+			// Export an include line for each header
+			for (UClass* Class : Classes)
 			{
-				ExportClassHeader(Class);
+				if (Class->GetOuter() == Package)
+				{
+					ExportClassHeader(Class);
+				}
 			}
-		}
-		
-		ListOfPublicClassesUObjectHeaderIncludes.Logf( TEXT("\r\n") );
-		
-		// build the full header file out of its pieces
-		const FString FullClassesHeader = FString::Printf(
-			TEXT("%s\r\n%s"),
-			*PreHeaderText, 
-			*ListOfPublicClassesUObjectHeaderIncludes
-		);
 
-		// Save the classes header if it has changed.
-		SaveHeaderIfChanged(*ClassesHeaderPath,*FullClassesHeader);
+			ListOfPublicClassesUObjectHeaderIncludes.Logf(TEXT("\r\n"));
+
+			// build the full header file out of its pieces
+			const FString FullClassesHeader = FString::Printf(
+				TEXT("%s\r\n%s"),
+				*PreHeaderText,
+				*ListOfPublicClassesUObjectHeaderIncludes
+				);
+
+			// Save the classes header if it has changed.
+			SaveHeaderIfChanged(*ClassesHeaderPath, *FullClassesHeader);
+		}
 	}
 
-	if (HeaderFilenames.Num() == 0)
+	if (ExportAndTemporaryClassCount == 0)
 	{
 		// If no headers are generated, check if there's any no export classes that need to have the inl file generated.
-		for ( int32 ClassIndex = 0; ClassIndex < Classes.Num(); ClassIndex++ )
+		for (UClass* Class : Classes)
 		{
-			UClass* Class = Classes[ClassIndex];
-			if ( Class->GetOuter() == Package && !Class->HasAnyClassFlags(CLASS_Intrinsic|CLASS_Temporary) )
+			if (Class->GetOuter() == Package && !Class->HasAnyClassFlags(CLASS_Intrinsic | CLASS_Temporary))
 			{
 				ExportClassHeader(Class);
 			}
@@ -4190,11 +3918,8 @@ FNativeClassHeaderGenerator::FNativeClassHeaderGenerator( UPackage* InPackage, F
 	// now export the names for the functions in this package
 	// notice we always export this file (as opposed to only exporting if we have any marked names)
 	// because there would be no way to know when the file was created otherwise
-	if (!bCircularDependencyDetected)
-	{
-		// Export ..generated.cpp
-		ExportGeneratedCPP();
-	}
+	// Export ..generated.cpp
+	ExportGeneratedCPP();
 
 	ExportGeneratedProto();
 
diff --git a/Engine/Source/Programs/UnrealHeaderTool/Private/NativeClassExporter.h b/Engine/Source/Programs/UnrealHeaderTool/Private/NativeClassExporter.h
--- a/Engine/Source/Programs/UnrealHeaderTool/Private/NativeClassExporter.h
+++ b/Engine/Source/Programs/UnrealHeaderTool/Private/NativeClassExporter.h
@@ -38,7 +38,6 @@ private:
 
 	UClass*				CurrentClass;
 	TArray<FClass*>		Classes;
-	FString				ClassHeaderFilename;
 	FString				API;
 	FString			ClassesHeaderPath;
 	/** the name of the cpp file where the implementation functions are generated, without the .cpp extension */
@@ -106,15 +105,6 @@ private:
 	{
 	};
 
-	/**
-	 * Sorts the list of header files being exported from a package according to their dependency on each other.
-	 *
-	 * @param	HeaderDependencyMap		A map of headers and their dependencies. Each header is represented as the actual filename string.
-	 * @param	SortedHeaderFilenames	[out] receives the sorted list of header filenames.
-	 * @return	true upon success, false if there were circular dependencies which made it impossible to sort properly.
-	 */
-	bool SortHeaderDependencyMap(const TMap<const FString*, HeaderDependents>& HeaderDependencyMap, TArray<const FString*>& SortedHeaderFilenames ) const;
-
 	/**
 	 * Finds to headers that are dependent on each other.
 	 * Wrapper for FindInterDependencyRecursive().
@@ -139,30 +129,6 @@ private:
 	 */
 	bool FindInterDependencyRecursive( TMap<const FString*, HeaderDependents>& HeaderDependencyMap, const FString* HeaderIndex, TSet<const FString*>& VisitedHeaders, const FString*& OutHeader1, const FString*& OutHeader2 );
 
-	/**
-	 * Finds a dependency chain between two class header files.
-	 * Wrapper around FindDependencyChainRecursive().
-	 *
-	 * @param	Class				A class to scan for a dependency chain between the two headers.
-	 * @param	Header1				First class header filename.
-	 * @param	Header2				Second class header filename.
-	 * @param	DependencyChain		[out] Receives dependency chain, if found.
-	 * @return	true if a dependency chain was found and filled in.
-	 */
-	bool FindDependencyChain( FClass* Class, const FString& Header1, const FString& Header2, TArray<FClass*>& DependencyChain );
-
-	/**
-	 * Finds a dependency chain between two class header files.
-	 *
-	 * @param	Class				A class to scan for a dependency chain between the two headers.
-	 * @param	Header1				First class header filename.
-	 * @param	Header2				Second class header filename.
-	 * @param	bChainStarted		Whether Header1 has been found and we've started to fill in DependencyChain. Must be false to begin with.
-	 * @param	DependencyChain		[out] Receives dependency chain, if found. Must be empty before the call.
-	 * @return	true if a dependency chain was found and filled in.
-	 */
-	bool FindDependencyChainRecursive( FClass* Class, const FString& Header1, const FString& Header2, bool bChainStarted, TArray<FClass*>& DependencyChain );
-
 	/**
 	 * Determines whether the glue version of the specified native function
 	 * should be exported
@@ -485,7 +451,6 @@ private:
 	 * Deletes all .generated.h files which do not correspond to any of the classes.
 	 */
 	void DeleteUnusedGeneratedHeaders();
-
 public:
 
 	// Constructor
diff --git a/Engine/Source/Runtime/Engine/Classes/AI/BehaviorTree/BehaviorTreeComponent.h b/Engine/Source/Runtime/Engine/Classes/AI/BehaviorTree/BehaviorTreeComponent.h
--- a/Engine/Source/Runtime/Engine/Classes/AI/BehaviorTree/BehaviorTreeComponent.h
+++ b/Engine/Source/Runtime/Engine/Classes/AI/BehaviorTree/BehaviorTreeComponent.h
@@ -26,7 +26,7 @@ struct FBTNodeExecutionInfo
 	FBTNodeExecutionInfo() : ExecuteNode(NULL) { }
 };
 
-UCLASS(HeaderGroup=Component)
+UCLASS()
 class ENGINE_API UBehaviorTreeComponent : public UBrainComponent
 {
 	GENERATED_UCLASS_BODY()
diff --git a/Engine/Source/Runtime/Engine/Classes/AI/BehaviorTree/BlackboardComponent.h b/Engine/Source/Runtime/Engine/Classes/AI/BehaviorTree/BlackboardComponent.h
--- a/Engine/Source/Runtime/Engine/Classes/AI/BehaviorTree/BlackboardComponent.h
+++ b/Engine/Source/Runtime/Engine/Classes/AI/BehaviorTree/BlackboardComponent.h
@@ -11,7 +11,10 @@
  */
 
 #pragma once
+#include "Components/ActorComponent.h"
 #include "BehaviorTreeTypes.h"
+#include "Blackboard/BlackboardKeyType.h"
+#include "BlackboardData.h"
 #include "BlackboardComponent.generated.h"
 
 namespace EBlackboardDescription
@@ -25,7 +28,7 @@ namespace EBlackboardDescription
 	};
 }
 
-UCLASS(HeaderGroup=Component)
+UCLASS()
 class ENGINE_API UBlackboardComponent : public UActorComponent
 {
 	GENERATED_UCLASS_BODY()
diff --git a/Engine/Source/Runtime/Engine/Classes/AI/BrainComponent.h b/Engine/Source/Runtime/Engine/Classes/AI/BrainComponent.h
--- a/Engine/Source/Runtime/Engine/Classes/AI/BrainComponent.h
+++ b/Engine/Source/Runtime/Engine/Classes/AI/BrainComponent.h
@@ -83,7 +83,7 @@ private:
 	TWeakObjectPtr<UBrainComponent> Owner;
 };
 
-UCLASS(HeaderGroup=Component, abstract, BlueprintType)
+UCLASS(abstract, BlueprintType)
 class ENGINE_API UBrainComponent : public UActorComponent, public IAIResourceInterface
 {
 	GENERATED_UCLASS_BODY()
diff --git a/Engine/Source/Runtime/Engine/Classes/AI/EnvironmentQuery/EQSRenderingComponent.h b/Engine/Source/Runtime/Engine/Classes/AI/EnvironmentQuery/EQSRenderingComponent.h
--- a/Engine/Source/Runtime/Engine/Classes/AI/EnvironmentQuery/EQSRenderingComponent.h
+++ b/Engine/Source/Runtime/Engine/Classes/AI/EnvironmentQuery/EQSRenderingComponent.h
@@ -2,6 +2,8 @@
 
 #pragma once
 #include "DebugRenderSceneProxy.h"
+#include "Debug/DebugDrawService.h"
+#include "EnvQueryTypes.h"
 #include "EQSRenderingComponent.generated.h"
 
 class FEQSSceneProxy : public FDebugRenderSceneProxy
@@ -28,7 +30,7 @@ private:
 	bool SafeIsActorSelected() const;
 };
 
-UCLASS(HeaderGroup=Component, hidecategories=Object)
+UCLASS(hidecategories=Object)
 class ENGINE_API UEQSRenderingComponent : public UPrimitiveComponent
 {
 	GENERATED_UCLASS_BODY()
diff --git a/Engine/Source/Runtime/Engine/Classes/AI/EnvironmentQuery/EQSTestingPawn.h b/Engine/Source/Runtime/Engine/Classes/AI/EnvironmentQuery/EQSTestingPawn.h
--- a/Engine/Source/Runtime/Engine/Classes/AI/EnvironmentQuery/EQSTestingPawn.h
+++ b/Engine/Source/Runtime/Engine/Classes/AI/EnvironmentQuery/EQSTestingPawn.h
@@ -7,7 +7,7 @@
 /** this class is abstract even though it's perfectly functional on its own.
  *	The reason is to stop it from showing as valid player pawn type when configuring 
  *	project's game mode. */
-UCLASS(abstract, HeaderGroup=Pawn, hidecategories=(Advanced, Attachment, Collision, Animation),DependsOn=(UEnvQueryTypes))
+UCLASS(abstract, hidecategories=(Advanced, Attachment, Collision, Animation),DependsOn=(UEnvQueryTypes))
 class ENGINE_API AEQSTestingPawn : public ACharacter, public IEQSQueryResultSourceInterface
 {
 	GENERATED_UCLASS_BODY()
@@ -75,4 +75,4 @@ protected:
 	void UpdateDrawing();
 
 	static void OnEditorSelectionChanged(UObject* NewSelection);
-};
\ No newline at end of file
+};
diff --git a/Engine/Source/Runtime/Engine/Classes/AI/Navigation/NavLinkRenderingComponent.h b/Engine/Source/Runtime/Engine/Classes/AI/Navigation/NavLinkRenderingComponent.h
--- a/Engine/Source/Runtime/Engine/Classes/AI/Navigation/NavLinkRenderingComponent.h
+++ b/Engine/Source/Runtime/Engine/Classes/AI/Navigation/NavLinkRenderingComponent.h
@@ -3,7 +3,7 @@
 #pragma once
 #include "NavLinkRenderingComponent.generated.h"
 
-UCLASS(HeaderGroup=Component, hidecategories=Object, editinlinenew)
+UCLASS(hidecategories=Object, editinlinenew)
 class ENGINE_API UNavLinkRenderingComponent : public UPrimitiveComponent
 {
 	GENERATED_UCLASS_BODY()
diff --git a/Engine/Source/Runtime/Engine/Classes/AI/Navigation/NavMeshRenderingComponent.h b/Engine/Source/Runtime/Engine/Classes/AI/Navigation/NavMeshRenderingComponent.h
--- a/Engine/Source/Runtime/Engine/Classes/AI/Navigation/NavMeshRenderingComponent.h
+++ b/Engine/Source/Runtime/Engine/Classes/AI/Navigation/NavMeshRenderingComponent.h
@@ -4,7 +4,7 @@
 #pragma once
 #include "NavMeshRenderingComponent.generated.h"
 
-UCLASS(HeaderGroup=Component, hidecategories=Object, editinlinenew)
+UCLASS(hidecategories=Object, editinlinenew)
 class UNavMeshRenderingComponent : public UPrimitiveComponent
 {
 	GENERATED_UCLASS_BODY()
diff --git a/Engine/Source/Runtime/Engine/Classes/AI/Navigation/NavRelevantComponent.h b/Engine/Source/Runtime/Engine/Classes/AI/Navigation/NavRelevantComponent.h
--- a/Engine/Source/Runtime/Engine/Classes/AI/Navigation/NavRelevantComponent.h
+++ b/Engine/Source/Runtime/Engine/Classes/AI/Navigation/NavRelevantComponent.h
@@ -3,7 +3,7 @@
 #pragma once
 #include "NavRelevantComponent.generated.h"
 
-UCLASS(HeaderGroup=Navigation)
+UCLASS()
 class ENGINE_API UNavRelevantComponent : public UActorComponent
 {
 	GENERATED_UCLASS_BODY()
diff --git a/Engine/Source/Runtime/Engine/Classes/AI/Navigation/NavTestRenderingComponent.h b/Engine/Source/Runtime/Engine/Classes/AI/Navigation/NavTestRenderingComponent.h
--- a/Engine/Source/Runtime/Engine/Classes/AI/Navigation/NavTestRenderingComponent.h
+++ b/Engine/Source/Runtime/Engine/Classes/AI/Navigation/NavTestRenderingComponent.h
@@ -3,7 +3,7 @@
 #pragma once
 #include "NavTestRenderingComponent.generated.h"
 
-UCLASS(HeaderGroup=Component, hidecategories=Object)
+UCLASS(hidecategories=Object)
 class UNavTestRenderingComponent: public UPrimitiveComponent
 {
 	GENERATED_UCLASS_BODY()
diff --git a/Engine/Source/Runtime/Engine/Classes/AI/Navigation/NavigationComponent.h b/Engine/Source/Runtime/Engine/Classes/AI/Navigation/NavigationComponent.h
--- a/Engine/Source/Runtime/Engine/Classes/AI/Navigation/NavigationComponent.h
+++ b/Engine/Source/Runtime/Engine/Classes/AI/Navigation/NavigationComponent.h
@@ -3,7 +3,7 @@
 #pragma once
 #include "NavigationComponent.generated.h"
 
-UCLASS(HeaderGroup=Component, dependson=(UNavAgentInterface, UNavPathObserverInterface, UNavigationSystem), config=Engine)
+UCLASS(dependson=(UNavAgentInterface, UNavPathObserverInterface, UNavigationSystem), config=Engine)
 class ENGINE_API UNavigationComponent : public UActorComponent, public INavigationPathGenerator
 {
 	GENERATED_UCLASS_BODY()
diff --git a/Engine/Source/Runtime/Engine/Classes/AI/Navigation/NavigationGraphNodeComponent.h b/Engine/Source/Runtime/Engine/Classes/AI/Navigation/NavigationGraphNodeComponent.h
--- a/Engine/Source/Runtime/Engine/Classes/AI/Navigation/NavigationGraphNodeComponent.h
+++ b/Engine/Source/Runtime/Engine/Classes/AI/Navigation/NavigationGraphNodeComponent.h
@@ -3,7 +3,7 @@
 #pragma once
 #include "NavigationGraphNodeComponent.generated.h"
 
-UCLASS(HeaderGroup=Component, config=Engine, dependson=ANavigationGraph, MinimalAPI, HideCategories=(Mobility))
+UCLASS(config=Engine, dependson=ANavigationGraph, MinimalAPI, HideCategories=(Mobility))
 class UNavigationGraphNodeComponent : public USceneComponent
 {
 	GENERATED_UCLASS_BODY()
diff --git a/Engine/Source/Runtime/Engine/Classes/AI/Navigation/NavigationTestingActor.h b/Engine/Source/Runtime/Engine/Classes/AI/Navigation/NavigationTestingActor.h
--- a/Engine/Source/Runtime/Engine/Classes/AI/Navigation/NavigationTestingActor.h
+++ b/Engine/Source/Runtime/Engine/Classes/AI/Navigation/NavigationTestingActor.h
@@ -26,7 +26,7 @@ namespace ENavCostDisplay
 	};
 }
 
-UCLASS(hidecategories=(Object, Actor, Input, Rendering), HeaderGroup=Navigation, Blueprintable)
+UCLASS(hidecategories=(Object, Actor, Input, Rendering), Blueprintable)
 class ENGINE_API ANavigationTestingActor : public AActor, public INavAgentInterface, public INavPathObserverInterface
 {
 	GENERATED_UCLASS_BODY()
diff --git a/Engine/Source/Runtime/Engine/Classes/AI/Navigation/PathFollowingComponent.h b/Engine/Source/Runtime/Engine/Classes/AI/Navigation/PathFollowingComponent.h
--- a/Engine/Source/Runtime/Engine/Classes/AI/Navigation/PathFollowingComponent.h
+++ b/Engine/Source/Runtime/Engine/Classes/AI/Navigation/PathFollowingComponent.h
@@ -58,7 +58,7 @@ namespace EPathFollowingRequestResult
 	};
 }
 
-UCLASS(HeaderGroup=Component, config=Engine)
+UCLASS(config=Engine)
 class ENGINE_API UPathFollowingComponent : public UActorComponent, public IAIResourceInterface
 {
 	GENERATED_UCLASS_BODY()
diff --git a/Engine/Source/Runtime/Engine/Classes/AI/Navigation/SmartNavLinkComponent.h b/Engine/Source/Runtime/Engine/Classes/AI/Navigation/SmartNavLinkComponent.h
--- a/Engine/Source/Runtime/Engine/Classes/AI/Navigation/SmartNavLinkComponent.h
+++ b/Engine/Source/Runtime/Engine/Classes/AI/Navigation/SmartNavLinkComponent.h
@@ -22,7 +22,7 @@ namespace ESmartNavLinkDir
 	};
 }
 
-UCLASS(HeaderGroup=Navigation)
+UCLASS()
 class ENGINE_API USmartNavLinkComponent : public UNavRelevantComponent
 {
 	GENERATED_UCLASS_BODY()
diff --git a/Engine/Source/Runtime/Engine/Classes/Animation/AnimBlueprintGeneratedClass.h b/Engine/Source/Runtime/Engine/Classes/Animation/AnimBlueprintGeneratedClass.h
--- a/Engine/Source/Runtime/Engine/Classes/Animation/AnimBlueprintGeneratedClass.h
+++ b/Engine/Source/Runtime/Engine/Classes/Animation/AnimBlueprintGeneratedClass.h
@@ -10,6 +10,7 @@ struct FBakedAnimationStateMachine;
 class  UAnimInstance;
 class  UEdGraph;
 class  UEdGraphNode;
+class  UAnimGraphNode_StateResult;
 
 // This structure represents debugging information for a single state machine
 USTRUCT()
diff --git a/Engine/Source/Runtime/Engine/Classes/Animation/AnimNode_BlendSpaceEvaluator.h b/Engine/Source/Runtime/Engine/Classes/Animation/AnimNode_BlendSpaceEvaluator.h
--- a/Engine/Source/Runtime/Engine/Classes/Animation/AnimNode_BlendSpaceEvaluator.h
+++ b/Engine/Source/Runtime/Engine/Classes/Animation/AnimNode_BlendSpaceEvaluator.h
@@ -8,7 +8,7 @@
 // Evaluates a point in a blendspace, using a specific time input rather than advancing time internally.
 // Typically the playback position of the animation for this node will represent something other than time, like jump height.
 // This node will not trigger any notifies present in the associated sequence.
-USTRUCT(HeaderGroup = AnimationNode)
+USTRUCT()
 struct ENGINE_API FAnimNode_BlendSpaceEvaluator : public FAnimNode_BlendSpacePlayer
 {
 	GENERATED_USTRUCT_BODY()
diff --git a/Engine/Source/Runtime/Engine/Classes/Animation/AnimSingleNodeInstance.h b/Engine/Source/Runtime/Engine/Classes/Animation/AnimSingleNodeInstance.h
--- a/Engine/Source/Runtime/Engine/Classes/Animation/AnimSingleNodeInstance.h
+++ b/Engine/Source/Runtime/Engine/Classes/Animation/AnimSingleNodeInstance.h
@@ -7,6 +7,7 @@
  */
 
 #pragma once
+#include "AnimInstance.h"
 #include "AnimSingleNodeInstance.generated.h"
 
 DECLARE_DYNAMIC_DELEGATE(FPostEvaluateAnimEvent);
diff --git a/Engine/Source/Runtime/Engine/Classes/Animation/SkeletalMeshActor.h b/Engine/Source/Runtime/Engine/Classes/Animation/SkeletalMeshActor.h
--- a/Engine/Source/Runtime/Engine/Classes/Animation/SkeletalMeshActor.h
+++ b/Engine/Source/Runtime/Engine/Classes/Animation/SkeletalMeshActor.h
@@ -1,9 +1,9 @@
 // Copyright 1998-2014 Epic Games, Inc. All Rights Reserved.
 #pragma once
-
+#include "Components/SkeletalMeshComponent.h"
 #include "SkeletalMeshActor.generated.h"
 
-UCLASS(HeaderGroup=Anim, ClassGroup=ISkeletalMeshes, Blueprintable, ConversionRoot, meta=(ChildCanTick))
+UCLASS(ClassGroup=ISkeletalMeshes, Blueprintable, ConversionRoot, meta=(ChildCanTick))
 class ENGINE_API ASkeletalMeshActor : public AActor, public IMatineeAnimInterface
 {
 	GENERATED_UCLASS_BODY()
diff --git a/Engine/Source/Runtime/Engine/Classes/Atmosphere/AtmosphericFogComponent.h b/Engine/Source/Runtime/Engine/Classes/Atmosphere/AtmosphericFogComponent.h
--- a/Engine/Source/Runtime/Engine/Classes/Atmosphere/AtmosphericFogComponent.h
+++ b/Engine/Source/Runtime/Engine/Classes/Atmosphere/AtmosphericFogComponent.h
@@ -109,7 +109,7 @@ public:
 /**
  *	Used to create fogging effects such as clouds.
  */
-UCLASS(HeaderGroup=Component, ClassGroup=Rendering, collapsecategories, hidecategories=(Object, Mobility, Activation, "Components|Activation"), editinlinenew, meta=(BlueprintSpawnableComponent), MinimalAPI)
+UCLASS(ClassGroup=Rendering, collapsecategories, hidecategories=(Object, Mobility, Activation, "Components|Activation"), editinlinenew, meta=(BlueprintSpawnableComponent), MinimalAPI)
 class UAtmosphericFogComponent : public USceneComponent
 {
 	GENERATED_UCLASS_BODY()
diff --git a/Engine/Source/Runtime/Engine/Classes/Camera/CameraComponent.h b/Engine/Source/Runtime/Engine/Classes/Camera/CameraComponent.h
--- a/Engine/Source/Runtime/Engine/Classes/Camera/CameraComponent.h
+++ b/Engine/Source/Runtime/Engine/Classes/Camera/CameraComponent.h
@@ -7,7 +7,7 @@
   * Represents a camera viewpoint and settings, such as projection type, field of view, and post-process overrides.
   * The default behavior for an actor used as the camera view target is to look for an attached camera component and use its location, rotation, and settings.
   */
-UCLASS(HeaderGroup=Component, HideCategories=(Mobility, Rendering, LOD), ClassGroup=Camera, meta=(BlueprintSpawnableComponent), MinimalAPI)
+UCLASS(HideCategories=(Mobility, Rendering, LOD), ClassGroup=Camera, meta=(BlueprintSpawnableComponent), MinimalAPI)
 class UCameraComponent : public USceneComponent
 {
 	GENERATED_UCLASS_BODY()
diff --git a/Engine/Source/Runtime/Engine/Classes/Components/ActorComponent.h b/Engine/Source/Runtime/Engine/Classes/Components/ActorComponent.h
--- a/Engine/Source/Runtime/Engine/Classes/Components/ActorComponent.h
+++ b/Engine/Source/Runtime/Engine/Classes/Components/ActorComponent.h
@@ -3,9 +3,10 @@
 #pragma once
 
 #include "ComponentInstanceDataCache.h"
+#include "Engine/EngineTypes.h"
 #include "ActorComponent.generated.h"
 
-UCLASS(DefaultToInstanced, HeaderGroup=Component, abstract, hidecategories=(ComponentReplication))
+UCLASS(DefaultToInstanced, abstract, hidecategories=(ComponentReplication))
 class ENGINE_API UActorComponent : public UObject
 {
 	GENERATED_UCLASS_BODY()
diff --git a/Engine/Source/Runtime/Engine/Classes/Components/ApplicationLifecycleComponent.h b/Engine/Source/Runtime/Engine/Classes/Components/ApplicationLifecycleComponent.h
--- a/Engine/Source/Runtime/Engine/Classes/Components/ApplicationLifecycleComponent.h
+++ b/Engine/Source/Runtime/Engine/Classes/Components/ApplicationLifecycleComponent.h
@@ -5,7 +5,7 @@
 
 #include "ApplicationLifecycleComponent.generated.h"
 
-UCLASS(HeaderGroup=Component, ClassGroup=Utility, HideCategories=(Activation, "Components|Activation", Collision), meta=(BlueprintSpawnableComponent))
+UCLASS(ClassGroup=Utility, HideCategories=(Activation, "Components|Activation", Collision), meta=(BlueprintSpawnableComponent))
 class ENGINE_API UApplicationLifecycleComponent : public UActorComponent
 {
 	GENERATED_UCLASS_BODY()
diff --git a/Engine/Source/Runtime/Engine/Classes/Components/ArrowComponent.h b/Engine/Source/Runtime/Engine/Classes/Components/ArrowComponent.h
--- a/Engine/Source/Runtime/Engine/Classes/Components/ArrowComponent.h
+++ b/Engine/Source/Runtime/Engine/Classes/Components/ArrowComponent.h
@@ -7,7 +7,7 @@
 /** 
  * A simple arrow rendered using lines.Useful for indicating which way an object is facing.
  */
-UCLASS(HeaderGroup=Component, ClassGroup=Shapes, hidecategories=(Object,LOD,Physics,Lighting,TextureStreaming,Activation,"Components|Activation",Collision), editinlinenew, meta=(BlueprintSpawnableComponent), MinimalAPI)
+UCLASS(ClassGroup=Shapes, hidecategories=(Object,LOD,Physics,Lighting,TextureStreaming,Activation,"Components|Activation",Collision), editinlinenew, meta=(BlueprintSpawnableComponent), MinimalAPI)
 class UArrowComponent : public UPrimitiveComponent
 {
 	GENERATED_UCLASS_BODY()
diff --git a/Engine/Source/Runtime/Engine/Classes/Components/AudioComponent.h b/Engine/Source/Runtime/Engine/Classes/Components/AudioComponent.h
--- a/Engine/Source/Runtime/Engine/Classes/Components/AudioComponent.h
+++ b/Engine/Source/Runtime/Engine/Classes/Components/AudioComponent.h
@@ -2,6 +2,7 @@
 
 
 #pragma once
+#include "Audio.h"
 #include "AudioComponent.generated.h"
 
 /** called when we finish playing audio, either because it played to completion or because a Stop() call turned it off early */
@@ -89,7 +90,7 @@ struct FSoundParseParameters
 /**
  * Used to provide objects with audio.
  */
-UCLASS(HeaderGroup=Component, ClassGroup=(Audio, Common), hidecategories=(Object, ActorComponent, Physics, Rendering, Mobility, LOD), ShowCategories=Trigger, dependson=(AReverbVolume, UEngineTypes, USoundAttenuation), meta=(BlueprintSpawnableComponent))
+UCLASS(ClassGroup=(Audio, Common), hidecategories=(Object, ActorComponent, Physics, Rendering, Mobility, LOD), ShowCategories=Trigger, dependson=(AReverbVolume, UEngineTypes, USoundAttenuation), meta=(BlueprintSpawnableComponent))
 class ENGINE_API UAudioComponent : public USceneComponent
 {
 	GENERATED_UCLASS_BODY()
diff --git a/Engine/Source/Runtime/Engine/Classes/Components/BillboardComponent.h b/Engine/Source/Runtime/Engine/Classes/Components/BillboardComponent.h
--- a/Engine/Source/Runtime/Engine/Classes/Components/BillboardComponent.h
+++ b/Engine/Source/Runtime/Engine/Classes/Components/BillboardComponent.h
@@ -7,7 +7,7 @@
 /** 
  * A 2d texture that will be rendered always facing the camera.
  */
-UCLASS(HeaderGroup=Component, ClassGroup=Rendering, collapsecategories, hidecategories=(Object,Activation,"Components|Activation",Physics,Collision,Lighting,LOD,Mesh), editinlinenew, meta=(BlueprintSpawnableComponent),MinimalAPI)
+UCLASS(ClassGroup=Rendering, collapsecategories, hidecategories=(Object,Activation,"Components|Activation",Physics,Collision,Lighting,LOD,Mesh), editinlinenew, meta=(BlueprintSpawnableComponent),MinimalAPI)
 class UBillboardComponent : public UPrimitiveComponent
 {
 	GENERATED_UCLASS_BODY()
diff --git a/Engine/Source/Runtime/Engine/Classes/Components/BoxComponent.h b/Engine/Source/Runtime/Engine/Classes/Components/BoxComponent.h
--- a/Engine/Source/Runtime/Engine/Classes/Components/BoxComponent.h
+++ b/Engine/Source/Runtime/Engine/Classes/Components/BoxComponent.h
@@ -7,7 +7,7 @@
 /** 
  * A box generally used for simple collision. Bounds are rendered as lines in the editor.
  */
-UCLASS(HeaderGroup=Component, ClassGroup=Shapes, hidecategories=(Object,LOD,Lighting,TextureStreaming), editinlinenew, meta=(BlueprintSpawnableComponent))
+UCLASS(ClassGroup=Shapes, hidecategories=(Object,LOD,Lighting,TextureStreaming), editinlinenew, meta=(BlueprintSpawnableComponent))
 class ENGINE_API UBoxComponent : public UShapeComponent
 {
 	GENERATED_UCLASS_BODY()
