Merging UE4-Pretest @ 2042161 to UE4

diff --git a/Engine/Source/Programs/UnrealBuildTool/System/ExternalExecution.cs b/Engine/Source/Programs/UnrealBuildTool/System/ExternalExecution.cs
--- a/Engine/Source/Programs/UnrealBuildTool/System/ExternalExecution.cs
+++ b/Engine/Source/Programs/UnrealBuildTool/System/ExternalExecution.cs
@@ -40,6 +40,12 @@ namespace UnrealBuildTool
 		/** Private headers with UObjects */
 		public List<FileItem> PrivateUObjectHeaders;
 
+		/** Module PCH absolute path */
+		public string PCH;
+
+		/** Base (i.e. extensionless) path+filename of the .generated files */
+		public string GeneratedCPPFilenameBase;
+
 		public override string ToString()
 		{
 			return ModuleName;
@@ -56,6 +62,8 @@ namespace UnrealBuildTool
 			public List<string> ClassesHeaders;
 			public List<string> PublicHeaders;
 			public List<string> PrivateHeaders;
+			public string       PCH;
+			public string       GeneratedCPPFilenameBase;
 			public bool         SaveExportedHeaders;
 
 			public override string ToString()
@@ -71,12 +79,14 @@ namespace UnrealBuildTool
 			RootBuildPath    = InRootBuildPath;
 
 			Modules = ModuleInfo.Select(Info => new Module{
-				Name                = Info.ModuleName,
-				BaseDirectory       = Info.ModuleDirectory,
-				OutputDirectory     = UEBuildModuleCPP.GetGeneratedCodeDirectoryForModule(Target, Info.ModuleDirectory, Info.ModuleName),
-				ClassesHeaders      = Info.PublicUObjectClassesHeaders.Select((Header) => Header.AbsolutePath).ToList(),
-				PublicHeaders       = Info.PublicUObjectHeaders       .Select((Header) => Header.AbsolutePath).ToList(),
-				PrivateHeaders      = Info.PrivateUObjectHeaders      .Select((Header) => Header.AbsolutePath).ToList(),
+				Name                     = Info.ModuleName,
+				BaseDirectory            = Info.ModuleDirectory,
+				OutputDirectory          = UEBuildModuleCPP.GetGeneratedCodeDirectoryForModule(Target, Info.ModuleDirectory, Info.ModuleName),
+				ClassesHeaders           = Info.PublicUObjectClassesHeaders.Select((Header) => Header.AbsolutePath).ToList(),
+				PublicHeaders            = Info.PublicUObjectHeaders       .Select((Header) => Header.AbsolutePath).ToList(),
+				PrivateHeaders           = Info.PrivateUObjectHeaders      .Select((Header) => Header.AbsolutePath).ToList(),
+				PCH                      = Info.PCH,
+				GeneratedCPPFilenameBase = Info.GeneratedCPPFilenameBase,
 				[REDACTED_EMAIL]: This assumes Engine/Source is a 'safe' folder name to check for
 				SaveExportedHeaders = !UnrealBuildTool.RunningRocket() || !Info.ModuleDirectory.Contains("Engine\\Source\\")
 
@@ -133,9 +143,9 @@ namespace UnrealBuildTool
 		}
 
 		/// <summary>
-		/// Gets the timestamp of CoreUObject.generated.inl file.
+		/// Gets the timestamp of CoreUObject.generated.cpp file.
 		/// </summary>
-		/// <returns>Last write time of CoreUObject.generated.inl or DateTime.MaxValue if it doesn't exist.</returns>
+		/// <returns>Last write time of CoreUObject.generated.cpp or DateTime.MaxValue if it doesn't exist.</returns>
 		private static DateTime GetCoreGeneratedTimestamp(UEBuildTarget Target)
 		{			
 			DateTime Timestamp;
@@ -147,7 +157,7 @@ namespace UnrealBuildTool
 			else
 			{
 				var CoreUObjectModule = (UEBuildModuleCPP)Target.GetModuleByName( "CoreUObject" );
-				string CoreGeneratedFilename = Path.Combine(UEBuildModuleCPP.GetGeneratedCodeDirectoryForModule(Target, CoreUObjectModule.ModuleDirectory, CoreUObjectModule.Name), CoreUObjectModule.Name + ".generated.inl");
+				string CoreGeneratedFilename = Path.Combine(UEBuildModuleCPP.GetGeneratedCodeDirectoryForModule(Target, CoreUObjectModule.ModuleDirectory, CoreUObjectModule.Name), CoreUObjectModule.Name + ".generated.cpp");
 				if (File.Exists(CoreGeneratedFilename))
 				{
 					Timestamp = new FileInfo(CoreGeneratedFilename).LastWriteTime;
@@ -180,7 +190,7 @@ namespace UnrealBuildTool
 				HeaderToolTimestamp = new FileInfo(HeaderToolPath).LastWriteTime;
 			}
 
-			// Get CoreUObject.generated.inl timestamp.  If the source files are older than the CoreUObject generated code, we'll
+			// Get CoreUObject.generated.cpp timestamp.  If the source files are older than the CoreUObject generated code, we'll
 			// need to regenerate code for the module
 			var CoreGeneratedTimestamp = GetCoreGeneratedTimestamp(Target);
 
@@ -367,11 +377,11 @@ namespace UnrealBuildTool
 		{
 			// We never want to try to execute the header tool when we're already trying to build it!
 			var bIsBuildingUHT = Target.GetTargetName().Equals( "UnrealHeaderTool", StringComparison.InvariantCultureIgnoreCase );
-			var BuildPlatform = UEBuildPlatform.GetBuildPlatform(Target.Platform);
-			var CppPlatform = BuildPlatform.GetCPPTargetPlatform(Target.Platform);
-			var ToolChain = UEToolChain.GetPlatformToolChain(CppPlatform);
-			var RootLocalPath = Path.GetFullPath(ProjectFileGenerator.RootRelativePath);
-			var Manifest = new UHTManifest(UnrealBuildTool.BuildingRocket() || UnrealBuildTool.RunningRocket(), Target, RootLocalPath, ToolChain.ConvertPath(RootLocalPath + '\\'), UObjectModules);
+			var BuildPlatform  = UEBuildPlatform.GetBuildPlatform(Target.Platform);
+			var CppPlatform    = BuildPlatform.GetCPPTargetPlatform(Target.Platform);
+			var ToolChain      = UEToolChain.GetPlatformToolChain(CppPlatform);
+			var RootLocalPath  = Path.GetFullPath(ProjectFileGenerator.RootRelativePath);
+			var Manifest       = new UHTManifest(UnrealBuildTool.BuildingRocket() || UnrealBuildTool.RunningRocket(), Target, RootLocalPath, ToolChain.ConvertPath(RootLocalPath + '\\'), UObjectModules);
 
 			// ensure the headers are up to date
 			if (!bIsBuildingUHT && 
diff --git a/Engine/Source/Programs/UnrealBuildTool/System/Unity.cs b/Engine/Source/Programs/UnrealBuildTool/System/Unity.cs
--- a/Engine/Source/Programs/UnrealBuildTool/System/Unity.cs
+++ b/Engine/Source/Programs/UnrealBuildTool/System/Unity.cs
@@ -4,12 +4,99 @@ using System;
 using System.Collections.Generic;
 using System.Text;
 using System.IO;
+using System.Linq;
 using System.Diagnostics;
 
 namespace UnrealBuildTool
 {
 	public class Unity
 	{
+		/// <summary>
+		/// A class which represents a list of files and the sum of their lengths.
+		/// </summary>
+		public class FileCollection
+		{
+			public List<FileItem> Files       { get; private set; }
+			public long           TotalLength { get; private set; }
+
+			public FileCollection()
+			{
+				Files       = new List<FileItem>();
+				TotalLength = 0;
+			}
+
+			public void AddFile(FileItem File)
+			{
+				Files.Add(File);
+				TotalLength += File.Info.Length;
+			}
+		}
+
+		/// <summary>
+		/// A class for building up a set of unity files.  You add files one-by-one using AddFile then call EndCurrentUnityFile to finish that one and
+		/// (perhaps) begin a new one.
+		/// </summary>
+		public class UnityFileBuilder
+		{
+			private List<FileCollection> UnityFiles;
+			private FileCollection       CurrentCollection;
+			private int                  SplitLength;
+
+			/// <summary>
+			/// Constructs a new UnityFileBuilder.
+			/// </summary>
+			/// <param name="InSplitLength">The accumulated length at which to automatically split a unity file, or -1 to disable automatic splitting.</param>
+			public UnityFileBuilder(int InSplitLength)
+			{
+				UnityFiles        = new List<FileCollection>();
+				CurrentCollection = new FileCollection();
+				SplitLength       = InSplitLength;
+			}
+
+			/// <summary>
+			/// Adds a file to the current unity file.  If splitting is required and the total size of the
+			/// unity file exceeds the split limit, then a new file is automatically started.
+			/// </summary>
+			/// <param name="File">The file to add.</param>
+			public void AddFile(FileItem File)
+			{
+				CurrentCollection.AddFile(File);
+				if (SplitLength != -1 && CurrentCollection.TotalLength > SplitLength)
+				{
+					EndCurrentUnityFile();
+				}
+			}
+
+			/// <summary>
+			/// Starts a new unity file.  If the current unity file contains no files, this function has no effect, i.e. you will not get an empty unity file.
+			/// </summary>
+			public void EndCurrentUnityFile()
+			{
+				if (CurrentCollection.Files.Count == 0)
+					return;
+
+				UnityFiles.Add(CurrentCollection);
+				CurrentCollection = new FileCollection();
+			}
+
+			/// <summary>
+			/// Returns the list of built unity files.  The UnityFileBuilder is unusable after this.
+			/// </summary>
+			/// <returns></returns>
+			public List<FileCollection> GetUnityFiles()
+			{
+				EndCurrentUnityFile();
+
+				var Result = UnityFiles;
+
+				// Null everything to ensure that failure will occur if you accidentally reuse this object.
+				CurrentCollection = null;
+				UnityFiles         = null;
+
+				return Result;
+			}
+		}
+
 		/// <summary>
 		/// Given a set of C++ files, generates another set of C++ files that #include all the original
 		/// files, the goal being to compile the same code in fewer translation units.
@@ -25,59 +112,40 @@ namespace UnrealBuildTool
 			string BaseName
 			)
 		{
+			var ToolChain = UEToolChain.GetPlatformToolChain(CompileEnvironment.Config.TargetPlatform);
+
 			UEBuildPlatform BuildPlatform = UEBuildPlatform.GetBuildPlatformForCPPTargetPlatform(CompileEnvironment.Config.TargetPlatform);
 
 			// Figure out size of all input files combined. We use this to determine whether to use larger unity threshold or not.
-			long TotalBytesInCPPFiles = 0;
-			foreach( FileItem CPPFile in CPPFiles )
-			{
-				TotalBytesInCPPFiles += CPPFile.Info.Length;
-			}
+			long TotalBytesInCPPFiles = CPPFiles.Sum(F => F.Info.Length);
 
 			// We have an increased threshold for unity file size if, and only if, all files fit into the same unity file. This
 			// is beneficial when dealing with PCH files. The default PCH creation limit is X unity files so if we generate < X 
 			// this could be fairly slow and we'd rather bump the limit a bit to group them all into the same unity file.
-			bool bForceIntoSingleUnityFile = BuildConfiguration.bStressTestUnity;			
-			if( (TotalBytesInCPPFiles < BuildConfiguration.NumIncludedBytesPerUnityCPP * 2)
+			//
 			// Optimization only makes sense if PCH files are enabled.
-			&&	CompileEnvironment.ShouldUsePCHs() )
-			{
-				bForceIntoSingleUnityFile = true;
-			}
+			bool bForceIntoSingleUnityFile = BuildConfiguration.bStressTestUnity || (TotalBytesInCPPFiles < BuildConfiguration.NumIncludedBytesPerUnityCPP * 2 && CompileEnvironment.ShouldUsePCHs());
 
-			// Figure out how many unity files there are going to be total.
-			int NumUnityFiles = 0;
-			int InputFileIndex = 0;
-			while (InputFileIndex < CPPFiles.Count)
+			// Build the list of unity files.
+			List<FileCollection> AllUnityFiles;
 			{
-				long NumIncludedBytesInThisOutputFile = 0;
-				//@warning: this condition is mirrored below
-				while(	InputFileIndex < CPPFiles.Count &&
-						(bForceIntoSingleUnityFile ||
-						NumIncludedBytesInThisOutputFile < BuildConfiguration.NumIncludedBytesPerUnityCPP))
+				var CPPUnityFileBuilder = new UnityFileBuilder(bForceIntoSingleUnityFile ? -1 : BuildConfiguration.NumIncludedBytesPerUnityCPP);
+				foreach( var CPPFile in CPPFiles )
 				{
-                    bool ForceAlone = CPPFiles[InputFileIndex].AbsolutePath.Contains(".GeneratedWrapper.");
-                    if (ForceAlone && !bForceIntoSingleUnityFile && NumIncludedBytesInThisOutputFile > 0)
-                    {
-                        break;
-                    }
-                    NumIncludedBytesInThisOutputFile += CPPFiles[InputFileIndex].Info.Length;
-                    InputFileIndex++;
-                    if (ForceAlone && !bForceIntoSingleUnityFile)
-                    {
-                        break;
-                    }
-                }
-				NumUnityFiles++;
+					if (!bForceIntoSingleUnityFile && CPPFile.AbsolutePath.Contains(".GeneratedWrapper."))
+					{
+						CPPUnityFileBuilder.EndCurrentUnityFile();
+						CPPUnityFileBuilder.AddFile(CPPFile);
+						CPPUnityFileBuilder.EndCurrentUnityFile();
+					}
+					else
+					{
+						CPPUnityFileBuilder.AddFile(CPPFile);
+					}
+				}
+				AllUnityFiles = CPPUnityFileBuilder.GetUnityFiles();
 			}
 
-			// Create a set of CPP files that combine smaller CPP files into larger compilation units, along with the corresponding 
-			// actions to compile them.
-			InputFileIndex = 0;
-			int CurrentUnityFileCount = 1;
-			List<FileItem> UnityCPPFiles = new List<FileItem>();
-			var ToolChain = UEToolChain.GetPlatformToolChain(CompileEnvironment.Config.TargetPlatform);
-
 			string PCHHeaderNameInCode = CPPFiles[0].PCHHeaderNameInCode;
 			if( CompileEnvironment.Config.PrecompiledHeaderIncludeFilename != null )
 			{
@@ -89,10 +157,17 @@ namespace UnrealBuildTool
 				CompileEnvironment.Config.PCHHeaderNameInCode = PCHHeaderNameInCode;
 			}
 
-			while (InputFileIndex < CPPFiles.Count)
+			// Create a set of CPP files that combine smaller CPP files into larger compilation units, along with the corresponding 
+			// actions to compile them.
+			int CurrentUnityFileCount = 0;
+			var UnityCPPFiles         = new List<FileItem>();
+			foreach( var UnityFile in AllUnityFiles )
 			{
-				StringWriter OutputUnityCPPWriter = new StringWriter();
+				++CurrentUnityFileCount;
+
+				StringWriter OutputUnityCPPWriter      = new StringWriter();
 				StringWriter OutputUnityCPPWriterExtra = null;
+
 				// add an extra file for UBT to get the #include dependencies from
 				if (BuildPlatform.RequiresExtraUnityCPPWriter() == true)
 				{
@@ -112,61 +187,39 @@ namespace UnrealBuildTool
 					}
 				}
 
-				// Add source files to the unity file until the number of included bytes crosses a threshold.
-				long NumIncludedBytesInThisOutputFile = 0;
-				var NumInputFilesInThisOutputFile = 0;
-                string FileDescription = "";
-				//@warning: this condition is mirrored above
-				while(	InputFileIndex < CPPFiles.Count &&
-						(bForceIntoSingleUnityFile ||
-						NumIncludedBytesInThisOutputFile < BuildConfiguration.NumIncludedBytesPerUnityCPP))
+				// Add source files to the unity file
+				foreach( var CPPFile in UnityFile.Files )
 				{
-                    bool ForceAlone = CPPFiles[InputFileIndex].AbsolutePath.Contains(".GeneratedWrapper.");
-                    if (ForceAlone && !bForceIntoSingleUnityFile && NumIncludedBytesInThisOutputFile > 0)
-                    {
-                        break;
-                    }
-                    FileItem CPPFile = CPPFiles[InputFileIndex];
 					OutputUnityCPPWriter.WriteLine("#include \"{0}\"", ToolChain.ConvertPath(CPPFile.AbsolutePath));
                     if (OutputUnityCPPWriterExtra != null)
                     {
                         OutputUnityCPPWriterExtra.WriteLine("#include \"{0}\"", CPPFile.AbsolutePath);
-                    }
-
-					NumIncludedBytesInThisOutputFile += CPPFile.Info.Length;
-                    InputFileIndex++;
-					NumInputFilesInThisOutputFile++;
-                    FileDescription += Path.GetFileName(CPPFile.AbsolutePath) + " + ";
-                    if (ForceAlone && !bForceIntoSingleUnityFile)
-                    {
-                        break;
                     }
                 }
-                // Remove trailing " + "
-                FileDescription = FileDescription.Remove( FileDescription.Length - 3);
 
-				// Write the unity file to the intermediate folder.
-				string ProgressString = "." + CurrentUnityFileCount.ToString() + "_of_" + NumUnityFiles.ToString();
-				string UnityCPPFilePath = Path.Combine(
-					CompileEnvironment.Config.OutputDirectory,
-					string.Format("Module.{0}{1}.cpp",
-						BaseName,
-						NumUnityFiles > 1 ? ProgressString : ""
-						)
-					);
-				CurrentUnityFileCount++;
+				// Determine unity file path name
+				string UnityCPPFilePath;
+				if (AllUnityFiles.Count > 1)
+				{
+					UnityCPPFilePath = string.Format("Module.{0}.{1}_of_{2}.cpp", BaseName, CurrentUnityFileCount, AllUnityFiles.Count);
+				}
+				else
+				{
+					UnityCPPFilePath = string.Format("Module.{0}.cpp", BaseName);
+				}
+				UnityCPPFilePath = Path.Combine(CompileEnvironment.Config.OutputDirectory, UnityCPPFilePath);
 
+				// Write the unity file to the intermediate folder.
 				FileItem UnityCPPFile = FileItem.CreateIntermediateTextFile(UnityCPPFilePath, OutputUnityCPPWriter.ToString());
-				UnityCPPFile.RelativeCost = NumIncludedBytesInThisOutputFile;
-				UnityCPPFile.Description = FileDescription;
-				UnityCPPFile.PCHHeaderNameInCode = PCHHeaderNameInCode;
-                UnityCPPFiles.Add(UnityCPPFile);
-
-				// write out the extra file
 				if (OutputUnityCPPWriterExtra != null)
 				{
 					FileItem.CreateIntermediateTextFile(UnityCPPFilePath + ".ex", OutputUnityCPPWriterExtra.ToString());
 				}
+
+				UnityCPPFile.RelativeCost        = UnityFile.TotalLength;
+				UnityCPPFile.Description         = string.Join(" + ", UnityFile.Files.Select(F => Path.GetFileName(F.AbsolutePath)));
+				UnityCPPFile.PCHHeaderNameInCode = PCHHeaderNameInCode;
+                UnityCPPFiles.Add(UnityCPPFile);
 			}
 
 			return UnityCPPFiles;
diff --git a/Engine/Source/Programs/UnrealBuildTool/UnrealBuildTool.csproj b/Engine/Source/Programs/UnrealBuildTool/UnrealBuildTool.csproj
--- a/Engine/Source/Programs/UnrealBuildTool/UnrealBuildTool.csproj
+++ b/Engine/Source/Programs/UnrealBuildTool/UnrealBuildTool.csproj
@@ -249,6 +249,7 @@
     <Compile Include="ToolChain\UEToolChain.cs" />
     <Compile Include="System\XGE.cs" />
     <Compile Include="Utilities\CopyrightVerify.cs" />
+    <Compile Include="Utilities\DictionaryExtensions.cs" />
     <Compile Include="Utilities\GraphVisualization.cs">
       <SubType>Code</SubType>
     </Compile>
diff --git a/Engine/Source/Programs/UnrealBuildTool/UnrealBuildTool_Mono.csproj b/Engine/Source/Programs/UnrealBuildTool/UnrealBuildTool_Mono.csproj
--- a/Engine/Source/Programs/UnrealBuildTool/UnrealBuildTool_Mono.csproj
+++ b/Engine/Source/Programs/UnrealBuildTool/UnrealBuildTool_Mono.csproj
@@ -156,6 +156,7 @@
     <Compile Include="ToolChain\UEToolChain.cs" />
     <Compile Include="System\XGE.cs" />
     <Compile Include="Utilities\CopyrightVerify.cs" />
+    <Compile Include="Utilities\DictionaryExtensions.cs" />
     <Compile Include="Utilities\XmlConfigLoader.cs" />
     <Compile Include="Utilities\GraphVisualization.cs">
       <SubType>Code</SubType>
diff --git a/Engine/Source/Programs/UnrealBuildTool/Utilities/DictionaryExtensions.cs b/Engine/Source/Programs/UnrealBuildTool/Utilities/DictionaryExtensions.cs
new file mode 100644
--- /dev/null
+++ b/Engine/Source/Programs/UnrealBuildTool/Utilities/DictionaryExtensions.cs
@@ -0,0 +1,31 @@
+// Copyright 1998-2014 Epic Games, Inc. All Rights Reserved.
+
+using System.Collections.Generic;
+
+namespace UnrealBuildTool
+{
+	public static class DictionaryExtensions
+	{
+		public static Value GetOrAddDefault<Key, Value>(this IDictionary<Key, Value> Dict, Key InKey)
+		{
+			Value Result;
+			if( !Dict.TryGetValue( InKey, out Result ) )
+			{
+				Result = default(Value);
+				Dict.Add( InKey, Result );
+			}
+			return Result;
+		}
+
+		public static Value GetOrAddNew<Key, Value>(this IDictionary<Key, Value> Dict, Key InKey) where Value : new()
+		{
+			Value Result;
+			if( !Dict.TryGetValue( InKey, out Result ) )
+			{
+				Result = new Value();
+				Dict.Add( InKey, Result );
+			}
+			return Result;
+		}
+	}
+}
\ No newline at end of file
diff --git a/Engine/Source/Programs/UnrealBuildTool/WinRT/UEBuildWinRT.cs b/Engine/Source/Programs/UnrealBuildTool/WinRT/UEBuildWinRT.cs
--- a/Engine/Source/Programs/UnrealBuildTool/WinRT/UEBuildWinRT.cs
+++ b/Engine/Source/Programs/UnrealBuildTool/WinRT/UEBuildWinRT.cs
@@ -17,7 +17,7 @@ namespace UnrealBuildTool
 
 		public static bool IsVisualStudioInstalled()
 		{
-			string BaseVSToolPath = WindowsPlatform.Compiler == WindowsCompiler.VisualStudio2012 ? Environment.GetEnvironmentVariable("VS110COMNTOOLS") : Environment.GetEnvironmentVariable("VS120COMNTOOLS");
+			string BaseVSToolPath = WindowsPlatform.GetVSComnToolsPath();
 			if (string.IsNullOrEmpty(BaseVSToolPath) == false)
 			{
 				return true;
diff --git a/Engine/Source/Programs/UnrealBuildTool/WinRT/WinRTToolChain.cs b/Engine/Source/Programs/UnrealBuildTool/WinRT/WinRTToolChain.cs
--- a/Engine/Source/Programs/UnrealBuildTool/WinRT/WinRTToolChain.cs
+++ b/Engine/Source/Programs/UnrealBuildTool/WinRT/WinRTToolChain.cs
@@ -876,7 +876,7 @@ namespace UnrealBuildTool
 				else
 				{
 					// Grab path to Visual Studio binaries from the system environment
-					string BaseVSToolPath = Environment.GetEnvironmentVariable("VS110COMNTOOLS");
+					string BaseVSToolPath = WindowsPlatform.GetVSComnToolsPath(WindowsCompiler.VisualStudio2012);
 					if (string.IsNullOrEmpty(BaseVSToolPath))
 					{
 						throw new BuildException("Visual Studio 2011 must be installed in order to build this target.");
@@ -929,7 +929,7 @@ namespace UnrealBuildTool
 				string VCVarsBatchFile = "";
 
 				// Grab path to Visual Studio binaries from the system environment
-				string BaseVSToolPath = Environment.GetEnvironmentVariable("VS110COMNTOOLS");
+				string BaseVSToolPath = WindowsPlatform.GetVSComnToolsPath(WindowsCompiler.VisualStudio2012);
 				if (string.IsNullOrEmpty(BaseVSToolPath))
 				{
 					BaseVSToolPath = "C:/Program Files (x86)/Microsoft Visual Studio 11.0/Common7/Tools/";
diff --git a/Engine/Source/Programs/UnrealBuildTool/Windows/VCToolChain.cs b/Engine/Source/Programs/UnrealBuildTool/Windows/VCToolChain.cs
--- a/Engine/Source/Programs/UnrealBuildTool/Windows/VCToolChain.cs
+++ b/Engine/Source/Programs/UnrealBuildTool/Windows/VCToolChain.cs
@@ -1499,24 +1499,15 @@ namespace UnrealBuildTool
 			}
 
 			// Grab path to Visual Studio binaries from the system environment
-			var VS2012Path = Environment.GetEnvironmentVariable( "VS110COMNTOOLS" );
-			var VS2013Path = Environment.GetEnvironmentVariable( "VS120COMNTOOLS" );
-			if( WindowsPlatform.Compiler == WindowsCompiler.VisualStudio2012 )
-			{
-				BaseVSToolPath = VS2012Path;
-			}
-			else if( WindowsPlatform.Compiler == WindowsCompiler.VisualStudio2013 )
-			{
-				BaseVSToolPath = VS2013Path;
-			}
+			BaseVSToolPath = WindowsPlatform.GetVSComnToolsPath();
 
 			if (string.IsNullOrEmpty(BaseVSToolPath))
 			{
-				if (WindowsPlatform.Compiler == WindowsCompiler.VisualStudio2012 && !String.IsNullOrEmpty( VS2013Path ) )
+				if (WindowsPlatform.Compiler == WindowsCompiler.VisualStudio2012 && !String.IsNullOrEmpty(WindowsPlatform.GetVSComnToolsPath(WindowsCompiler.VisualStudio2013)))
 				{
 					throw new BuildException("Visual Studio 2012 must be installed in order to build this target.  However, Visual Studio 2013 was detected.  To compile with Visual Studio 2013, change the 'WindowsPlatform.Compiler' variable in UnrealBuildTool");
 				}
-				else if (WindowsPlatform.Compiler == WindowsCompiler.VisualStudio2013 && !String.IsNullOrEmpty( VS2012Path ) )
+				else if (WindowsPlatform.Compiler == WindowsCompiler.VisualStudio2013 && !String.IsNullOrEmpty(WindowsPlatform.GetVSComnToolsPath(WindowsCompiler.VisualStudio2012)))
 				{
 					throw new BuildException("Visual Studio 2013 must be installed in order to build this target.  However, Visual Studio 2012 was detected.  To compile with Visual Studio 2012, change the 'WindowsPlatform.Compiler' variable in UnrealBuildTool");
 				}
diff --git a/Engine/Source/Programs/UnrealHeaderTool/Private/ClassMaps.h b/Engine/Source/Programs/UnrealHeaderTool/Private/ClassMaps.h
--- a/Engine/Source/Programs/UnrealHeaderTool/Private/ClassMaps.h
+++ b/Engine/Source/Programs/UnrealHeaderTool/Private/ClassMaps.h
@@ -29,6 +29,8 @@ private:
 	TUniquePtr<T> Obj_;
 };
 
+struct FManifestModule;
+
 extern TMap<UClass*, FString>                   GClassStrippedHeaderTextMap;
 extern TMap<UClass*, FString>                   GClassSourceFileMap;
 extern TMap<UClass*, TUniqueObj<TArray<FName>>> GClassDependentOnMap;
@@ -39,3 +41,4 @@ extern TMap<UClass*, FString>                   GClassIncludePathMap;
 extern TSet<UClass*>                            GExportedClasses;
 extern TSet<UClass*>                            GPublicClassSet;
 extern TMap<UProperty*, FString>                GArrayDimensions;
+extern TMap<UPackage*,  const FManifestModule*> GPackageToManifestModuleMap;
diff --git a/Engine/Source/Programs/UnrealHeaderTool/Private/CodeGenerator.cpp b/Engine/Source/Programs/UnrealHeaderTool/Private/CodeGenerator.cpp
--- a/Engine/Source/Programs/UnrealHeaderTool/Private/CodeGenerator.cpp
+++ b/Engine/Source/Programs/UnrealHeaderTool/Private/CodeGenerator.cpp
@@ -176,7 +176,7 @@ bool MakeCommandlet_FindPackageLocation(const TCHAR* InPackage, FString& OutLoca
 {
 	// Mapping of processed packages to their locations
 	// An empty location string means it was processed but not found
-	static TMap<FString, FManifest::FModule*> CheckedPackageList;
+	static TMap<FString, FManifestModule*> CheckedPackageList;
 
 	FString CheckPackage(InPackage);
 
@@ -184,7 +184,7 @@ bool MakeCommandlet_FindPackageLocation(const TCHAR* InPackage, FString& OutLoca
 
 	if (!ModuleInfoPtr)
 	{
-		auto* ModuleInfoPtr2 = GManifest.Modules.FindByPredicate([&](FManifest::FModule& Module) { return Module.Name == CheckPackage; });
+		auto* ModuleInfoPtr2 = GManifest.Modules.FindByPredicate([&](FManifestModule& Module) { return Module.Name == CheckPackage; });
 		if (ModuleInfoPtr2 && IFileManager::Get().DirectoryExists(*ModuleInfoPtr2->BaseDirectory))
 		{
 			ModuleInfoPtr = ModuleInfoPtr2;
@@ -1773,8 +1773,8 @@ int32 FindByIndirectName( const TArray<T>& Array, const FName& Name )
 void FNativeClassHeaderGenerator::ExportClassHeader( UClass* Class )
 {
 	TArray<UClass*> DependencyChain;
-	VisitedMap.Empty( Classes.Num() );
-	ExportClassHeaderRecursive( Class, DependencyChain, false );
+	TSet<const UClass*>	VisitedSet;
+	ExportClassHeaderRecursive(Class, DependencyChain, VisitedSet, false);
 }
 
 
@@ -2021,20 +2021,23 @@ void FNativeClassHeaderGenerator::ExportClassHeaderWrapper( UClass* Class, bool
 			NewFileName = GManifest.RootBuildPath / NewFileName;
 		}
 
-		// Keep track of all of the UObject headers for this module, in the same order that we digest them in
-		// @todo uht: We're wrapping these includes in checks for header guards, ONLY because most existing UObject headers are missing '#pragma once'
-		ListOfAllUObjectHeaderIncludes.Logf(
+		auto IncludeStr = FString::Printf(
 			TEXT("#ifndef %s_%s_generated_h")													LINE_TERMINATOR
 			TEXT("    #include \"%s\"")															LINE_TERMINATOR
 			TEXT("#endif")																		LINE_TERMINATOR,
 			*API, *CurrentClass->GetName(), *NewFileName );
 
+		// Keep track of all of the UObject headers for this module, in the same order that we digest them in
+		// @todo uht: We're wrapping these includes in checks for header guards, ONLY because most existing UObject headers are missing '#pragma once'
+		ListOfAllUObjectHeaderIncludes                .Logf(TEXT("%s"), *IncludeStr);
+		ListOfPublicClassesUObjectHeaderModuleIncludes.Logf(TEXT("%s"), *IncludeStr);
+
 		// For the 'Classes' mega-header, only include legacy UObjects classes from the 'Classes' folder.  Private and Public classes are
 		// not included here -- you should include those yourself if you need them!
 		const bool bIsPublicClassesHeader = GPublicClassSet.Contains( Class );
 		if( bIsPublicClassesHeader )
 		{
-			ListOfPublicClassesUObjectHeaderIncludes.Logf(TEXT("#include \"%s\"\r\n"), *NewFileName);
+			ListOfPublicClassesUObjectHeaderGroupIncludes.Logf(TEXT("#include \"%s\"\r\n"), *NewFileName);
 		}
 	}
 
@@ -2042,14 +2045,7 @@ void FNativeClassHeaderGenerator::ExportClassHeaderWrapper( UClass* Class, bool
 	EnumForeachText.Empty();
 }
 
-/**
- * Appends the header definition for an inheritance hierarchy of classes to the header.
- *
- * @param	Class					The class to be exported.
- * @param	DependencyChain			Used for finding errors. Must be empty before the first call.
- * @param	bCheckDependenciesOnly	Whether we should just keep checking for dependency errors, without exporting anything.
- */
-void FNativeClassHeaderGenerator::ExportClassHeaderRecursive( UClass* Class, TArray<UClass*>& DependencyChain, bool bCheckDependenciesOnly )
+void FNativeClassHeaderGenerator::ExportClassHeaderRecursive( UClass* Class, TArray<UClass*>& DependencyChain, TSet<const UClass*>& VisitedSet, bool bCheckDependenciesOnly )
 {
 	bool bIsExportClass   = GClassStrippedHeaderTextMap.Contains(Class) && Class->HasAnyClassFlags(CLASS_Native|CLASS_Temporary) && !Class->HasAnyClassFlags(CLASS_NoExport);
 	bool bIsCorrectHeader = GClassHeaderFilenameMap.FindRef(Class) == ClassHeaderFilename && Class->GetOuter() == Package;
@@ -2057,25 +2053,22 @@ void FNativeClassHeaderGenerator::ExportClassHeaderRecursive( UClass* Class, TAr
 	// Check for circular header dependencies between export classes.
 	if ( bIsExportClass )
 	{
-		if ( bIsCorrectHeader == false )
+		if ( !bIsCorrectHeader )
 		{
 			if ( DependencyChain.Num() == 0 )
 			{
 				// The first export class we found doesn't belong in this header: No need to keep exporting along this dependency path.
 				return;
 			}
-			else
-			{
-				// From now on, we're not going to export anything. Instead, we're going to check that no deeper dependency tries to export to this header file.
-				bCheckDependenciesOnly = true;
-			}
+
+			// From now on, we're not going to export anything. Instead, we're going to check that no deeper dependency tries to export to this header file.
+			bCheckDependenciesOnly = true;
 		}
 		else if ( bCheckDependenciesOnly && !ClassHeaderFilename.IsEmpty() )
 		{
 			FString DependencyChainString;
-			for ( int32 DependencyIndex=0; DependencyIndex < DependencyChain.Num(); DependencyIndex++ )
+			for (UClass* DependencyClass : DependencyChain)
 			{
-				UClass* DependencyClass = DependencyChain[DependencyIndex];
 				DependencyChainString += DependencyClass->GetName() + TEXT(" -> ");
 			}
 			DependencyChainString += Class->GetName();
@@ -2090,14 +2083,14 @@ void FNativeClassHeaderGenerator::ExportClassHeaderRecursive( UClass* Class, TAr
 	}
 
 	// Check for circular dependencies.
-	bool* bVisited = VisitedMap.Find( Class );
-	if ( bVisited && *bVisited == true )
+	if (VisitedSet.Contains(Class))
 	{
 		UE_LOG(LogCompile, Error, TEXT("Circular dependency detected for class %s!"), *Class->GetName() );
 		return;
 	}
+
 	// Temporarily mark the Class as VISITED. Make sure to clear this flag before returning!
-	VisitedMap.Add( Class, true );
+	VisitedSet.Add(Class);
 
 	if ( bIsExportClass )
 	{
@@ -2107,38 +2100,33 @@ void FNativeClassHeaderGenerator::ExportClassHeaderRecursive( UClass* Class, TAr
 	// Export the super class first.
 	if (UClass* SuperClass = Class->GetSuperClass())
 	{
-		ExportClassHeaderRecursive(SuperClass, DependencyChain, bCheckDependenciesOnly);
+		ExportClassHeaderRecursive(SuperClass, DependencyChain, VisitedSet, bCheckDependenciesOnly);
 	}
 
 	// Export all classes we depend on.
-	for( auto It = GClassDependentOnMap.FindOrAdd(Class)->CreateConstIterator(); It; ++It )
+	for (const FName& DependencyClassName : *GClassDependentOnMap.FindOrAdd(Class))
 	{
-		const FName& DependencyClassName = *It;
-
 		// By this point, all classes should have been validated so we can strip the given class name when searching.		
-		FString DependencyClassNameStripped(DependencyClassName.ToString());
+		FString DependencyClassNameStripped = DependencyClassName.ToString();
 
 		// Handle #include directives as 'dependson' (allowing them to fail).
 		const bool bClassNameFromHeader = FHeaderParser::DependentClassNameFromHeader(*DependencyClassNameStripped, DependencyClassNameStripped);
 		DependencyClassNameStripped = GetClassNameWithPrefixRemoved( DependencyClassNameStripped );
 
-		int32 ClassIndex = FindByIndirectName( Classes, *DependencyClassNameStripped );
-		UClass* FoundClass = FindObject<UClass>(ANY_PACKAGE, *DependencyClassNameStripped);
 		// Only export the class if it's in Classes array (it may be in another package).
-		UClass* DependsOnClass = (ClassIndex != INDEX_NONE) ? FoundClass : NULL;
+		UClass* FoundClass = FindObject<UClass>(ANY_PACKAGE, *DependencyClassNameStripped);
 		if (FoundClass == NULL)
 		{
 			// Try again with temporary class name. This may be struct only header.
 			DependencyClassNameStripped = GetClassNameWithoutPrefix(DependencyClassName.ToString());
 			DependencyClassNameStripped = FHeaderParser::GenerateTemporaryClassName(*DependencyClassNameStripped);
-			ClassIndex = FindByIndirectName( Classes, *DependencyClassNameStripped );
-			FoundClass = FindObject<UClass>(ANY_PACKAGE, *DependencyClassNameStripped);
-			DependsOnClass = (ClassIndex != INDEX_NONE) ? FoundClass : NULL;
+			FoundClass                  = FindObject<UClass>(ANY_PACKAGE, *DependencyClassNameStripped);
 		}
 
-		if (DependsOnClass)
+		int32 ClassIndex = FindByIndirectName( Classes, *DependencyClassNameStripped );
+		if (UClass* DependsOnClass = (ClassIndex != INDEX_NONE) ? FoundClass : NULL)
 		{
-			ExportClassHeaderRecursive(DependsOnClass, DependencyChain, bCheckDependenciesOnly);
+			ExportClassHeaderRecursive(DependsOnClass, DependencyChain, VisitedSet, bCheckDependenciesOnly);
 		}
 		else if ( !FoundClass && !bClassNameFromHeader )
 		{
@@ -2158,10 +2146,10 @@ void FNativeClassHeaderGenerator::ExportClassHeaderRecursive( UClass* Class, TAr
 	}
 
 	// We're done visiting this Class.
-	VisitedMap.Add( Class, false );
+	VisitedSet.Remove( Class );
 	if ( bIsExportClass )
 	{
-		DependencyChain.RemoveAtSwap( DependencyChain.Num() - 1 );
+		DependencyChain.Pop();
 	}
 }
 
@@ -3926,10 +3914,8 @@ FNativeClassHeaderGenerator::FNativeClassHeaderGenerator( UPackage* InPackage, F
 	public:
 		FString* AddUnique(const FString& Str)
 		{
-			for (auto It = Data.CreateConstIterator(); It; ++It)
+			for (auto& Ptr : Data)
 			{
-				auto& Ptr = *It;
-
 				if (*Ptr == Str)
 					return &*Ptr;
 			}
@@ -4063,18 +4049,23 @@ FNativeClassHeaderGenerator::FNativeClassHeaderGenerator( UPackage* InPackage, F
 	bool bHasNamesForExport = false;
 	TempHeaderPaths.Empty();
 	PackageHeaderPaths.Empty();
+
 	// Reset header generation output strings
-	GeneratedPackageCPP.Empty();
-	GeneratedProtoText.Empty();
-	GeneratedMCPText.Empty();
-	CrossModuleGeneratedFunctionDeclarations.Empty();
-	UniqueCrossModuleReferences.Empty();
-	GeneratedFunctionDeclarations.Empty();
-	GeneratedFunctionBodyTextSplit.Empty();
-	
-	for (auto It = HeaderFilenames.CreateConstIterator(); !bCircularDependencyDetected && It; ++It)
+	GeneratedPackageCPP                           .Empty();
+	GeneratedProtoText                            .Empty();
+	GeneratedMCPText                              .Empty();
+	CrossModuleGeneratedFunctionDeclarations      .Empty();
+	UniqueCrossModuleReferences                   .Empty();
+	GeneratedFunctionDeclarations                 .Empty();
+	GeneratedFunctionBodyTextSplit                .Empty();
+	ListOfPublicClassesUObjectHeaderModuleIncludes.Empty();
+
+	for (auto HeaderFilename : HeaderFilenames)
 	{
-		ClassHeaderFilename = **It;
+		if (bCircularDependencyDetected)
+			break;
+
+		ClassHeaderFilename = *HeaderFilename;
 
 		FString PkgDir;
 		FString GeneratedIncludeDirectory;
@@ -4082,13 +4073,12 @@ FNativeClassHeaderGenerator::FNativeClassHeaderGenerator( UPackage* InPackage, F
 		{
 			UE_LOG(LogCompile, Error, TEXT("Failed to find path for package %s"), *PackageName);
 		}
-		ClassesHeaderPath = GeneratedIncludeDirectory / PackageName + ClassHeaderFilename + TEXT("Classes.h");
+		FString ClassesHeaderName = PackageName + ClassHeaderFilename + TEXT("Classes.h");
+		ClassesHeaderPath = GeneratedIncludeDirectory / ClassesHeaderName;
 
 		int32 ClassCount = 0;
-		for( int32 ClassIndex = 0; ClassIndex < Classes.Num(); ClassIndex++ )
+		for (UClass* Class : Classes)
 		{
-			UClass* Class = Classes[ClassIndex];
-
 			if( Class->GetOuter()==Package && Class->HasAnyClassFlags(CLASS_Native|CLASS_Temporary) && GClassHeaderFilenameMap.FindRef(Class) == ClassHeaderFilename )
 			{
 				if (GClassStrippedHeaderTextMap.Contains(Class) && !Class->HasAnyClassFlags(CLASS_NoExport))
@@ -4113,8 +4103,10 @@ FNativeClassHeaderGenerator::FNativeClassHeaderGenerator( UPackage* InPackage, F
 			continue;
 		}
 
+		ClassesHeaders.Add(ClassesHeaderName);
+
 		CurrentClass = NULL;
-		ListOfPublicClassesUObjectHeaderIncludes.Empty();
+		ListOfPublicClassesUObjectHeaderGroupIncludes.Empty();
 		ListOfAllUObjectHeaderIncludes.Empty();
 		OriginalHeader.Empty();
 		PreHeaderText.Empty();
@@ -4152,22 +4144,21 @@ FNativeClassHeaderGenerator::FNativeClassHeaderGenerator( UPackage* InPackage, F
 		}
 
 		// Export an include line for each header
-		for ( int32 ClassIndex = 0; ClassIndex < Classes.Num(); ClassIndex++ )
+		for (UClass* Class : Classes)
 		{
-			UClass* Class = Classes[ClassIndex];
 			if ( GClassHeaderFilenameMap.FindRef(Class) == ClassHeaderFilename && Class->GetOuter() == Package )
 			{
 				ExportClassHeader(Class);
 			}
 		}
 		
-		ListOfPublicClassesUObjectHeaderIncludes.Logf( TEXT("\r\n") );
+		ListOfPublicClassesUObjectHeaderGroupIncludes.Logf( TEXT("\r\n") );
 		
 		// build the full header file out of its pieces
 		const FString FullClassesHeader = FString::Printf(
 			TEXT("%s\r\n%s"),
 			*PreHeaderText, 
-			*ListOfPublicClassesUObjectHeaderIncludes
+			*ListOfPublicClassesUObjectHeaderGroupIncludes
 		);
 
 		// Save the classes header if it has changed.
@@ -4177,9 +4168,8 @@ FNativeClassHeaderGenerator::FNativeClassHeaderGenerator( UPackage* InPackage, F
 	if (HeaderFilenames.Num() == 0)
 	{
 		// If no headers are generated, check if there's any no export classes that need to have the inl file generated.
-		for ( int32 ClassIndex = 0; ClassIndex < Classes.Num(); ClassIndex++ )
+		for (UClass* Class : Classes)
 		{
-			UClass* Class = Classes[ClassIndex];
 			if ( Class->GetOuter() == Package && !Class->HasAnyClassFlags(CLASS_Intrinsic|CLASS_Temporary) )
 			{
 				ExportClassHeader(Class);
@@ -4192,7 +4182,7 @@ FNativeClassHeaderGenerator::FNativeClassHeaderGenerator( UPackage* InPackage, F
 	// because there would be no way to know when the file was created otherwise
 	if (!bCircularDependencyDetected)
 	{
-		// Export ..generated.cpp
+		// Export .generated.cpp
 		ExportGeneratedCPP();
 	}
 
@@ -4211,28 +4201,24 @@ void FNativeClassHeaderGenerator::DeleteUnusedGeneratedHeaders()
 {
 	TArray<FString> AllIntermediateFolders;
 
-	for (auto It = PackageHeaderPaths.CreateConstIterator(); It; ++It)
+	for (const auto& PackageHeader : PackageHeaderPaths)
 	{
-		const FString IntermediatePath = FPaths::GetPath(*It);
+		const FString IntermediatePath = FPaths::GetPath(PackageHeader);
 
 		if (AllIntermediateFolders.Contains(IntermediatePath))
-		{
 			continue;
-		}
 
 		AllIntermediateFolders.Add( IntermediatePath );
 
 		TArray<FString> AllHeaders;
 		IFileManager::Get().FindFiles( AllHeaders, *(IntermediatePath / TEXT("*.generated.h")), true, false );
 
-		for (auto It = AllHeaders.CreateConstIterator(); It; ++It)
+		for (const auto& Header : AllHeaders)
 		{
-			const FString HeaderPath = IntermediatePath / *It;
+			const FString HeaderPath = IntermediatePath / Header;
 
 			if (PackageHeaderPaths.Contains(HeaderPath))
-			{
 				continue;
-			}
 
 			// Check intrinsic classes. Get the class name from file name by removing .generated.h.
 			const FString HeaderFilename = FPaths::GetBaseFilename(HeaderPath);
@@ -4287,7 +4273,7 @@ bool FNativeClassHeaderGenerator::SaveHeaderIfChanged(const TCHAR* HeaderPath, c
 
 	if (bWriteContents || bVerifyContents)
 	{
-		FString Ref = FString(FPaths::GameSavedDir()) / TEXT("ReferenceGeneratedCode") / FPaths::GetCleanFilename(HeaderPath);
+		FString Ref    = FString(FPaths::GameSavedDir()) / TEXT("ReferenceGeneratedCode") / FPaths::GetCleanFilename(HeaderPath);
 		FString Verify = FString(FPaths::GameSavedDir()) / TEXT("VerifyGeneratedCode") / FPaths::GetCleanFilename(HeaderPath);
 
 		if (bWriteContents)
@@ -4485,12 +4471,13 @@ void FNativeClassHeaderGenerator::ExportGeneratedMCP()
 
 /**
  * Exports C++ definitions for boilerplate that was generated for a package.
- * They are exported to a file using the name <PackageName>.generated.inl
+ * They are exported to a file using the name <PackageName>.generated.cpp
  * @param ReferencedNames list of function names to export.
  */
 void FNativeClassHeaderGenerator::ExportGeneratedCPP()
 {
 	FStringOutputDevice GeneratedCPPPreamble;
+	FStringOutputDevice GeneratedCPPClassesIncludes;
 	FStringOutputDevice GeneratedCPPEpilogue;
 	FStringOutputDevice GeneratedCPPText;
 
@@ -4508,8 +4495,18 @@ void FNativeClassHeaderGenerator::ExportGeneratedCPP()
 																								LINE_TERMINATOR
 		);
 
-	// Include all the UObject headers for this module first, in the same order that we digested them
-	GeneratedCPPText.Logf( *ListOfAllUObjectHeaderIncludes );
+	const auto* ModuleInfo = GPackageToManifestModuleMap.FindChecked(Package);
+	if (ModuleInfo->PCH.Len())
+	{
+		GeneratedCPPClassesIncludes.Logf(TEXT("#include \"%s\"") LINE_TERMINATOR, *ModuleInfo->PCH);
+	}
+
+	// Write out the ordered class dependencies into a single header that we can easily include
+	FString DepHeaderPathname = ModuleInfo->GeneratedCPPFilenameBase + TEXT(".dep.h");
+	SaveHeaderIfChanged(*DepHeaderPathname, *(GeneratedCPPPreamble + ListOfPublicClassesUObjectHeaderModuleIncludes));
+
+	// Write out our include to the .dep.h file
+	GeneratedCPPClassesIncludes.Logf(TEXT("#include \"%s\"") LINE_TERMINATOR, *FPaths::GetCleanFilename(DepHeaderPathname));
 
 	GeneratedCPPText.Log(*GeneratedPackageCPP);
 	// Add all names marked for export to a list (for sorting)
@@ -4532,11 +4529,6 @@ void FNativeClassHeaderGenerator::ExportGeneratedCPP()
 	
 	FString PkgName = FPackageName::GetShortName(Package);
 	FString PkgDir;
-	FString GeneratedIncludeDirectory;
-	if (MakeCommandlet_FindPackageLocation(*PkgName, PkgDir, GeneratedIncludeDirectory) == false)
-	{
-		UE_LOG(LogCompile, Error, TEXT("Failed to find path for package %s"), *PkgName);
-	}
 	if (GeneratedFunctionDeclarations.Len() || CrossModuleGeneratedFunctionDeclarations.Len())
 	{
 		ExportGeneratedPackageInitCode(Package);
@@ -4544,7 +4536,7 @@ void FNativeClassHeaderGenerator::ExportGeneratedCPP()
 
 	TArray<FString> NumberedHeaderNames;
 
-	// Generate each of the .generated.inl files
+	// Generate each of the .generated.cpp files
 	for( int32 FileIdx=0;FileIdx<GeneratedFunctionBodyTextSplit.Num();FileIdx++ )
 	{
 		FStringOutputDevice FileText;
@@ -4568,14 +4560,13 @@ void FNativeClassHeaderGenerator::ExportGeneratedCPP()
 			FileText.Logf(TEXT("#endif\r\n"));
 		}
 
-		FString HeaderName = GeneratedCPPFilenameBase + (GeneratedFunctionBodyTextSplit.Num() > 1 ? *FString::Printf(TEXT(".%d.inl"),FileIdx+1) : TEXT(".inl"));
-		FString HeaderPath = GeneratedIncludeDirectory / HeaderName;
+		FString HeaderPath = ModuleInfo->GeneratedCPPFilenameBase + (GeneratedFunctionBodyTextSplit.Num() > 1 ? *FString::Printf(TEXT(".%d.cpp"), FileIdx + 1) : TEXT(".cpp"));
 		
-		SaveHeaderIfChanged(*HeaderPath, *(GeneratedCPPPreamble + FileText + GeneratedCPPEpilogue));
+		SaveHeaderIfChanged(*HeaderPath, *(GeneratedCPPPreamble + GeneratedCPPClassesIncludes + FileText + GeneratedCPPEpilogue));
 
 		if (GeneratedFunctionBodyTextSplit.Num() > 1)
 		{
-			NumberedHeaderNames.Add(HeaderName);
+			NumberedHeaderNames.Add(FPaths::GetCleanFilename(HeaderPath));
 		}
 	}
 
@@ -4588,7 +4579,7 @@ void FNativeClassHeaderGenerator::ExportGeneratedCPP()
 			FileText.Logf(TEXT("#include \"%s\"") LINE_TERMINATOR, *NumberedHeaderNames[i]);
 		}
 
-		FString HeaderPath = GeneratedIncludeDirectory / GeneratedCPPFilenameBase + TEXT(".inl");
+		FString HeaderPath = ModuleInfo->GeneratedCPPFilenameBase + TEXT(".cpp");
 		SaveHeaderIfChanged(*HeaderPath, *(GeneratedCPPPreamble + FileText + GeneratedCPPEpilogue));
 	}
 
@@ -4669,112 +4660,101 @@ ECompilationResult::Type UnrealHeaderTool_Main(const FString& ModuleInfoFilename
 			}
 
 			// We'll make an ordered list of all UObject headers we care about.
-			TArray<FString> UObjectHeaders;
-			if( CurrentlyProcessing == PublicClassesHeaders )
-			{
-				UObjectHeaders = Module.PublicUObjectClassesHeaders;
-			}
 			// @todo uht: Ideally 'dependson' would not be allowed from public -> private, or NOT at all for new style headers
-			else if( CurrentlyProcessing == PublicHeaders )
-			{
-				// Append the public headers
-				UObjectHeaders = Module.PublicUObjectHeaders;
-			}
-			else
+			const TArray<FString>& UObjectHeaders =
+				(CurrentlyProcessing == PublicClassesHeaders) ? Module.PublicUObjectClassesHeaders :
+				(CurrentlyProcessing == PublicHeaders       ) ? Module.PublicUObjectHeaders        :
+				                                                Module.PrivateUObjectHeaders;
+			if (!UObjectHeaders.Num())
+				continue;
+
+			UPackage* Package = Cast<UPackage>( StaticFindObjectFast( UPackage::StaticClass(), NULL, FName(*Module.LongPackageName), false, false ) );
+			if( Package == NULL )
 			{
-				// Append the private headers
-				UObjectHeaders = Module.PrivateUObjectHeaders;
+				Package = CreatePackage( NULL, *Module.LongPackageName );
 			}
-			if( UObjectHeaders.Num() > 0 )
+			GPackageToManifestModuleMap.Add(Package, &Module);
+
+			// Set some package flags for indicating that this package contains script
+			// NOTE: We do this even if we didn't have to create the package, because CoreUObject is compiled into UnrealHeaderTool and we still
+			//       want to make sure our flags get set
+			Package->PackageFlags |= PKG_ContainsScript;
+			Package->PackageFlags &= ~(PKG_ClientOptional|PKG_ServerSideOnly);
+			Package->PackageFlags |= PKG_Compiling;
+
+			for (const FString& Filename : UObjectHeaders)
 			{
-				UPackage* Package = Cast<UPackage>( StaticFindObjectFast( UPackage::StaticClass(), NULL, FName(*Module.LongPackageName), false, false ) );
-				if( Package == NULL )
+				// Best faith effort at a useful line number for errors occurring during this initial pre-parsing
+				int32 ClassDeclLine = -1;
+
+			#if !PLATFORM_EXCEPTIONS_DISABLED
+				try
+			#endif
 				{
-					Package = CreatePackage( NULL, *Module.LongPackageName );
-				}
+					// Import class.
+					const FString ClassName      = FPaths::GetBaseFilename(Filename);
+					const FString FullModulePath = FPaths::ConvertRelativePathToFull(ModuleInfoPath, Filename);
 
-				// Set some package flags for indicating that this package contains script
-				// NOTE: We do this even if we didn't have to create the package, because CoreUObject is compiled into UnrealHeaderTool and we still
-				//       want to make sure our flags get set
-				Package->PackageFlags |= PKG_ContainsScript;
-				Package->PackageFlags &= ~(PKG_ClientOptional|PKG_ServerSideOnly);
-				Package->PackageFlags |= PKG_Compiling;
+					FString HeaderFile;
+					if (!FFileHelper::LoadFileToString(HeaderFile, *FullModulePath))
+						FError::Throwf(TEXT( "UnrealHeaderTool was unable to load source file '%s'"), *FullModulePath);
 
-				for (const FString& Filename : UObjectHeaders)
-				{
-					// Best faith effort at a useful line number for errors occurring during this initial pre-parsing
-					int32 ClassDeclLine = -1;
+					UClass* ResultClass = GenerateCodeForHeader(Package, *ClassName, RF_Public|RF_Standalone, *HeaderFile, ClassDeclLine);
+					GClassSourceFileMap.Add(ResultClass, Filename);
 
-				#if !PLATFORM_EXCEPTIONS_DISABLED
-					try
-				#endif
+					if( CurrentlyProcessing == PublicClassesHeaders )
 					{
-						// Import class.
-						const FString ClassName      = FPaths::GetBaseFilename(Filename);
-						const FString FullModulePath = FPaths::ConvertRelativePathToFull(ModuleInfoPath, Filename);
-
-						FString HeaderFile;
-						if (!FFileHelper::LoadFileToString(HeaderFile, *FullModulePath))
-							FError::Throwf(TEXT( "UnrealHeaderTool was unable to load source file '%s'"), *FullModulePath);
+						GPublicClassSet.Add(ResultClass);
+					}
 
-						UClass* ResultClass = GenerateCodeForHeader(Package, *ClassName, RF_Public|RF_Standalone, *HeaderFile, ClassDeclLine);
-						GClassSourceFileMap.Add(ResultClass, Filename);
+					// Save metadata for the class path, both for it's include path and relative to the module base directory
+					if(FullModulePath.StartsWith(Module.BaseDirectory))
+					{
+						// Get the path relative to the module directory
+						const TCHAR *ModuleRelativePath = *FullModulePath + Module.BaseDirectory.Len();
+						GClassModuleRelativePathMap.Add(ResultClass, ModuleRelativePath);
 
-						if( CurrentlyProcessing == PublicClassesHeaders )
+						// Add the include path
+						const TCHAR *IncludePath = ModuleRelativePath;
+						if(*IncludePath == '/')
 						{
-							GPublicClassSet.Add(ResultClass);
+							IncludePath++;
 						}
-
-						// Save metadata for the class path, both for it's include path and relative to the module base directory
-						if(FullModulePath.StartsWith(Module.BaseDirectory))
+						while(*IncludePath != 0 && *IncludePath != '/')
 						{
-							// Get the path relative to the module directory
-							const TCHAR *ModuleRelativePath = *FullModulePath + Module.BaseDirectory.Len();
-							GClassModuleRelativePathMap.Add(ResultClass, ModuleRelativePath);
-
-							// Add the include path
-							const TCHAR *IncludePath = ModuleRelativePath;
-							if(*IncludePath == '/')
-							{
-								IncludePath++;
-							}
-							while(*IncludePath != 0 && *IncludePath != '/')
-							{
-								IncludePath++;
-							}
-							if(*IncludePath == '/' && *(IncludePath + 1) != 0)
-							{
-								GClassIncludePathMap.Add(ResultClass, IncludePath + 1);
-							}
+							IncludePath++;
 						}
-					}
-				#if !PLATFORM_EXCEPTIONS_DISABLED
-					catch( TCHAR* ErrorMsg )
-					{
-						TGuardValue<ELogTimes::Type> DisableLogTimes(GPrintLogTimes, ELogTimes::None);
-
-						FString Prefix;
-						if (ClassDeclLine != -1)
+						if(*IncludePath == '/' && *(IncludePath + 1) != 0)
 						{
-							const FString AbsFilename = IFileManager::Get().ConvertToAbsolutePathForExternalAppForRead(*Filename);
-							Prefix = FString::Printf(TEXT("%s(%i): "), *AbsFilename, ClassDeclLine);
+							GClassIncludePathMap.Add(ResultClass, IncludePath + 1);
 						}
+					}
+				}
+			#if !PLATFORM_EXCEPTIONS_DISABLED
+				catch( TCHAR* ErrorMsg )
+				{
+					TGuardValue<ELogTimes::Type> DisableLogTimes(GPrintLogTimes, ELogTimes::None);
 
-						FString FormattedErrorMessage = FString::Printf(TEXT("%sError: %s\r\n"), *Prefix, ErrorMsg);
-						Result = GCompilationResult;
+					FString Prefix;
+					if (ClassDeclLine != -1)
+					{
+						const FString AbsFilename = IFileManager::Get().ConvertToAbsolutePathForExternalAppForRead(*Filename);
+						Prefix = FString::Printf(TEXT("%s(%i): "), *AbsFilename, ClassDeclLine);
+					}
 
-						UE_LOG(LogCompile, Log, TEXT("%s"), *FormattedErrorMessage);
-						GWarn->Log(ELogVerbosity::Error, FormattedErrorMessage);
+					FString FormattedErrorMessage = FString::Printf(TEXT("%sError: %s\r\n"), *Prefix, ErrorMsg);
+					Result = GCompilationResult;
 
-						++NumFailures;
-					}
-				#endif
-				}
+					UE_LOG(LogCompile, Log, TEXT("%s"), *FormattedErrorMessage);
+					GWarn->Log(ELogVerbosity::Error, FormattedErrorMessage);
 
-				if(Result == ECompilationResult::Succeeded && NumFailures != 0)
-				{
-					Result = ECompilationResult::OtherCompilationError;
+					++NumFailures;
 				}
+			#endif
+			}
+			if(Result == ECompilationResult::Succeeded && NumFailures != 0)
+			{
+				Result = ECompilationResult::OtherCompilationError;
 			}
 		}
 	}
@@ -4784,9 +4764,8 @@ ECompilationResult::Type UnrealHeaderTool_Main(const FString& ModuleInfoFilename
 	if ( Result == ECompilationResult::Succeeded )
 	{
 		// Verify that all script declared superclasses exist.
-		for( TObjectIterator<UClass> ItC; ItC; ++ItC )
+		for (const UClass* ScriptClass : TObjectRange<UClass>())
 		{
-			const UClass* ScriptClass      = *ItC;
 			const UClass* ScriptSuperClass = ScriptClass->GetSuperClass();
 
 			if (ScriptSuperClass && !ScriptSuperClass->HasAnyClassFlags(CLASS_Intrinsic) && GClassStrippedHeaderTextMap.Contains(ScriptClass) && !GClassStrippedHeaderTextMap.Contains(ScriptSuperClass))
diff --git a/Engine/Source/Programs/UnrealHeaderTool/Private/HeaderParser.cpp b/Engine/Source/Programs/UnrealHeaderTool/Private/HeaderParser.cpp
--- a/Engine/Source/Programs/UnrealHeaderTool/Private/HeaderParser.cpp
+++ b/Engine/Source/Programs/UnrealHeaderTool/Private/HeaderParser.cpp
@@ -3466,6 +3466,18 @@ UProperty* FHeaderParser::GetVarNameAndDim
 		AddFormattedPrevCommentAsTooltipMetaData(VarProperty.MetaData);
 	}
 
+	// validate UFunction parameters
+	if (IsFunction)
+	{
+		// UFunctions with a smart pointer as input parameter wont compile anyway, because of missing P_GET_... macro.
+		// UFunctions with a smart pointer as return type will crash when called via blueprint, because they are not supported in VM.
+		// WeakPointer is supported by VM as return type (see UObject::execLetWeakObjPtr), but there is no P_GET_... macro for WeakPointer.
+		if ((VarProperty.Type == CPT_LazyObjectReference) || (VarProperty.Type == CPT_AssetObjectReference))
+		{
+			FError::Throwf(TEXT("UFunctions cannot take a smart pointer (LazyPtr, AssetPtr, etc) as a parameter."));
+		}
+	}
+
 	// If this is the first time seeing the property name, then flag it for replace instead of add
 	const EFindName FindFlag = VarProperty.PropertyFlags & CPF_Config ? GetFindFlagForPropertyName(VarProperty.Identifier) : FNAME_Add;
 	// create the FName for the property, splitting (ie Unnamed_3 -> Unnamed,3)
@@ -5758,6 +5770,36 @@ void FHeaderParser::ValidatePropertyIsDeprecatedIfNecessary(FPropertyBase& VarPr
 	}
 }
 
+struct FExposeOnSpawnValidator
+{
+	// Keep this function synced with UEdGraphSchema_K2::FindSetVariableByNameFunction
+	static bool IsSupported(const FPropertyBase& Property)
+	{
+		bool ProperNativeType = false;
+		switch (Property.Type)
+		{
+		case CPT_Int:
+		case CPT_Byte:
+		case CPT_Float:
+		case CPT_Bool:
+		case CPT_ObjectReference:
+		case CPT_String:
+		case CPT_Name:
+		case CPT_Vector:
+		case CPT_Rotation:
+			ProperNativeType = true;
+		}
+
+		if (!ProperNativeType && (CPT_Struct == Property.Type) && Property.Struct)
+		{
+			static const FName BlueprintTypeName(TEXT("BlueprintType"));
+			ProperNativeType |= Property.Struct->GetBoolMetaData(BlueprintTypeName);
+		}
+
+		return ProperNativeType;
+	}
+};
+
 void FHeaderParser::CompileVariableDeclaration(FClasses& AllClasses, UStruct* Struct, EPropertyDeclarationStyle::Type PropertyDeclarationStyle)
 {
 	uint64 DisallowFlags = CPF_ParmFlags;
@@ -5812,6 +5854,14 @@ void FHeaderParser::CompileVariableDeclaration(FClasses& AllClasses, UStruct* St
 		FError::Throwf(TEXT("Illegal type modifiers in member variable declaration") );
 	}
 
+	if (FString* ExposeOnSpawnValue = OriginalProperty.MetaData.Find(TEXT("ExposeOnSpawn")))
+	{
+		if ((*ExposeOnSpawnValue == TEXT("true")) && !FExposeOnSpawnValidator::IsSupported(OriginalProperty))
+		{
+			FError::Throwf(TEXT("ExposeOnSpawn - Property cannoty be exposed"));
+		}
+	}
+
 	// Process all variables of this type.
 	TArray<UProperty*> NewProperties;
 	do
diff --git a/Engine/Source/Programs/UnrealHeaderTool/Private/Manifest.cpp b/Engine/Source/Programs/UnrealHeaderTool/Private/Manifest.cpp
--- a/Engine/Source/Programs/UnrealHeaderTool/Private/Manifest.cpp
+++ b/Engine/Source/Programs/UnrealHeaderTool/Private/Manifest.cpp
@@ -87,19 +87,22 @@ FManifest FManifest::LoadFromFile(const FString& Filename)
 
 		FString Outer = FString::Printf(TEXT("Modules[%d]"), ModuleIndex);
 
-		GetJsonFieldValue(KnownModule.Name,                      ModuleObj, TEXT("Name"),                *Outer);
-		GetJsonFieldValue(KnownModule.BaseDirectory,             ModuleObj, TEXT("BaseDirectory"),       *Outer);
-		GetJsonFieldValue(KnownModule.GeneratedIncludeDirectory, ModuleObj, TEXT("OutputDirectory"),     *Outer);
-		GetJsonFieldValue(KnownModule.SaveExportedHeaders,       ModuleObj, TEXT("SaveExportedHeaders"), *Outer);
-		GetJsonFieldValue(ClassesHeaders,                        ModuleObj, TEXT("ClassesHeaders"),      *Outer);
-		GetJsonFieldValue(PublicHeaders,                         ModuleObj, TEXT("PublicHeaders"),       *Outer);
-		GetJsonFieldValue(PrivateHeaders,                        ModuleObj, TEXT("PrivateHeaders"),      *Outer);
+		GetJsonFieldValue(KnownModule.Name,                      ModuleObj, TEXT("Name"),                     *Outer);
+		GetJsonFieldValue(KnownModule.BaseDirectory,             ModuleObj, TEXT("BaseDirectory"),            *Outer);
+		GetJsonFieldValue(KnownModule.GeneratedIncludeDirectory, ModuleObj, TEXT("OutputDirectory"),          *Outer);
+		GetJsonFieldValue(KnownModule.SaveExportedHeaders,       ModuleObj, TEXT("SaveExportedHeaders"),      *Outer);
+		GetJsonFieldValue(ClassesHeaders,                        ModuleObj, TEXT("ClassesHeaders"),           *Outer);
+		GetJsonFieldValue(PublicHeaders,                         ModuleObj, TEXT("PublicHeaders"),            *Outer);
+		GetJsonFieldValue(PrivateHeaders,                        ModuleObj, TEXT("PrivateHeaders"),           *Outer);
+		GetJsonFieldValue(KnownModule.PCH,                       ModuleObj, TEXT("PCH"),                      *Outer);
+		GetJsonFieldValue(KnownModule.GeneratedCPPFilenameBase,  ModuleObj, TEXT("GeneratedCPPFilenameBase"), *Outer);
 
 		KnownModule.LongPackageName = FPackageName::ConvertToLongScriptPackageName(*KnownModule.Name);
 
 		// Convert relative paths
 		KnownModule.BaseDirectory             = FPaths::ConvertRelativePathToFull(FilenamePath, KnownModule.BaseDirectory);
 		KnownModule.GeneratedIncludeDirectory = FPaths::ConvertRelativePathToFull(FilenamePath, KnownModule.GeneratedIncludeDirectory);
+		KnownModule.GeneratedCPPFilenameBase  = FPaths::ConvertRelativePathToFull(FilenamePath, KnownModule.GeneratedCPPFilenameBase);
 
 		// Ensure directories end with a slash, because this aids their use with FPaths::MakePathRelativeTo.
 		if (!KnownModule.BaseDirectory            .EndsWith(TEXT("/"))) { KnownModule.BaseDirectory            .AppendChar(TEXT('/')); }
diff --git a/Engine/Source/Programs/UnrealHeaderTool/Private/Manifest.h b/Engine/Source/Programs/UnrealHeaderTool/Private/Manifest.h
--- a/Engine/Source/Programs/UnrealHeaderTool/Private/Manifest.h
+++ b/Engine/Source/Programs/UnrealHeaderTool/Private/Manifest.h
@@ -7,42 +7,48 @@
 
 #pragma once
 
-struct FManifest
+struct FManifestModule
 {
-	struct FModule
-	{
-		/** The name of the module */
-		FString Name;
+	/** The name of the module */
+	FString Name;
+
+	/** Long package name for this module's UObject class */
+	FString LongPackageName;
+
+	/** Base directory of this module on disk */
+	FString BaseDirectory;
 
-		/** Long package name for this module's UObject class */
-		FString LongPackageName;
+	/** Directory where generated include files should go */
+	FString GeneratedIncludeDirectory;
 
-		/** Base directory of this module on disk */
-		FString BaseDirectory;
+	/** List of C++ public 'Classes' header files with UObjects in them (legacy) */
+	TArray<FString> PublicUObjectClassesHeaders;
 
-		/** Directory where generated include files should go */
-		FString GeneratedIncludeDirectory;
+	/** List of C++ public header files with UObjects in them */
+	TArray<FString> PublicUObjectHeaders;
 
-		/** List of C++ public 'Classes' header files with UObjects in them (legacy) */
-		TArray<FString> PublicUObjectClassesHeaders;
+	/** List of C++ private header files with UObjects in them */
+	TArray<FString> PrivateUObjectHeaders;
 
-		/** List of C++ public header files with UObjects in them */
-		TArray<FString> PublicUObjectHeaders;
+	/** Absolute path to the module's PCH */
+	FString PCH;
 
-		/** List of C++ private header files with UObjects in them */
-		TArray<FString> PrivateUObjectHeaders;
+	/** Base (i.e. extensionless) path+filename of where to write out the module's .generated.* files */
+	FString GeneratedCPPFilenameBase;
 
-		/** Whether or not to write out headers that have changed */
-		bool SaveExportedHeaders;
-	};
+	/** Whether or not to write out headers that have changed */
+	bool SaveExportedHeaders;
+};
 
+struct FManifest
+{
 	bool    UseRelativePaths;
 	FString RootLocalPath;
 	FString RootBuildPath;
 
 	/** Ordered list of modules that define UObjects or UStructs, which we may need to generate
 	    code for.  The list is in module dependency order, such that most dependent modules appear first. */
-	TArray<FModule> Modules;
+	TArray<FManifestModule> Modules;
 
 	/**
 	 * Loads an UnrealHeaderTool.manifest from the specified filename.
diff --git a/Engine/Source/Programs/UnrealHeaderTool/Private/NativeClassExporter.h b/Engine/Source/Programs/UnrealHeaderTool/Private/NativeClassExporter.h
--- a/Engine/Source/Programs/UnrealHeaderTool/Private/NativeClassExporter.h
+++ b/Engine/Source/Programs/UnrealHeaderTool/Private/NativeClassExporter.h
@@ -50,7 +50,8 @@ private:
 	UPackage*			Package;
 	FStringOutputDevice	PreHeaderText;
 	FStringOutputDevice	EnumForeachText;
-	FStringOutputDevice ListOfPublicClassesUObjectHeaderIncludes;
+	FStringOutputDevice ListOfPublicClassesUObjectHeaderGroupIncludes;  // This is built up each time from scratch each time for each header groups
+	FStringOutputDevice ListOfPublicClassesUObjectHeaderModuleIncludes; // This is built up for the entire module, and for all processed headers
 	FStringOutputDevice ListOfAllUObjectHeaderIncludes;
 	FStringOutputDevice FriendText;
 	FStringOutputDevice	GeneratedPackageCPP;
@@ -80,15 +81,15 @@ private:
 	/** the existing disk version of this header */
 	FString				OriginalHeader;
 
-	/** Used by certain traversal algorithms to flag if a class has been visited or not. */
-	TMap<const UClass*, bool>	VisitedMap;
-
 	/** Array of temp filenames that for files to overwrite headers */
 	TArray<FString>		TempHeaderPaths;
 
 	/** Array of all header filenames from the current package. */
 	TArray<FString>		PackageHeaderPaths;
 
+	/** Array of all generated Classes headers for the current package */
+	TArray<FString>		ClassesHeaders;
+	
 	/** if we are exporting a struct for offset determination only, replace some types with simpler ones (delegates) */
 	bool bIsExportingForOffsetDeterminationOnly;
 
@@ -231,9 +232,10 @@ private:
 	 *
 	 * @param	Class					The class to be exported.
 	 * @param	DependencyChain			Used for finding errors. Must be empty before the first call.
+	 * @param	VisitedSet				The set of classes visited so far. Must be empty before the first call.
 	 * @param	bCheckDependenciesOnly	Whether we should just keep checking for dependency errors, without exporting anything.
 	 */
-	void ExportClassHeaderRecursive( UClass* Class, TArray<UClass*>& DependencyChain, bool bCheckDependenciesOnly );
+	void ExportClassHeaderRecursive( UClass* Class, TArray<UClass*>& DependencyChain, TSet<const UClass*>& VisitedSet, bool bCheckDependenciesOnly );
 
 	/**
 	 * Returns a string in the format CLASS_Something|CLASS_Something which represents all class flags that are set for the specified
diff --git a/Engine/Source/Programs/UnrealHeaderTool/Resources/UHTDebugging.manifest b/Engine/Source/Programs/UnrealHeaderTool/Resources/UHTDebugging.manifest
--- a/Engine/Source/Programs/UnrealHeaderTool/Resources/UHTDebugging.manifest
+++ b/Engine/Source/Programs/UnrealHeaderTool/Resources/UHTDebugging.manifest
@@ -5,7 +5,7 @@
 	"Modules":[{
 			"Name": "CoreUObject",
 			"BaseDirectory": "..\\..\\..\\Runtime\\CoreUObject",
-			"OutputDirectory": "..\\..\\..\\..\\Intermediate/Build/Include\\CoreUObject\\",
+			"OutputDirectory": "..\\..\\..\\..\\Intermediate/Build/Inc\\CoreUObject\\",
 			"ClassesHeaders": [
 				"..\\..\\..\\Runtime\\CoreUObject\\Classes\\Object.h"
 			],
@@ -13,11 +13,13 @@
 			],
 			"PrivateHeaders": [
 			],
+			"PCH": "",
+			"GeneratedCPPFilenameBase": "..\\..\\..\\..\\Intermediate/Build/Inc\\CoreUObject\\CoreUObject.generated",
 			"SaveExportedHeaders": true
 		}, {
 			"Name": "UHTDebugging",
 			"BaseDirectory": ".\\UHTDebugging",
-			"OutputDirectory": "..\\..\\..\\..\\Intermediate/Build/Include\\UHTDebugging\\",
+			"OutputDirectory": "..\\..\\..\\..\\Intermediate/Build/Inc\\UHTDebugging\\",
 			"ClassesHeaders": [
 				".\\UHTDebugging\\TestObject.h",
 				".\\UHTDebugging\\TestInterface.h",
@@ -29,6 +31,8 @@
 			],
 			"PrivateHeaders": [
 			],
+			"PCH": "",
+			"GeneratedCPPFilenameBase": "..\\..\\..\\..\\Intermediate/Build/Inc\\UHTDebugging\\UHTDebugging.generated",
 			"SaveExportedHeaders": true
 		}
 	]
diff --git a/Engine/Source/Runtime/Android/AndroidRuntimeSettings/Private/AndroidRuntimeSettingsModule.cpp b/Engine/Source/Runtime/Android/AndroidRuntimeSettings/Private/AndroidRuntimeSettingsModule.cpp
--- a/Engine/Source/Runtime/Android/AndroidRuntimeSettings/Private/AndroidRuntimeSettingsModule.cpp
+++ b/Engine/Source/Runtime/Android/AndroidRuntimeSettings/Private/AndroidRuntimeSettingsModule.cpp
@@ -1,6 +1,6 @@
 // Copyright 1998-2014 Epic Games, Inc. All Rights Reserved.
 
 #include "AndroidRuntimeSettingsPrivatePCH.h"
-#include "AndroidRuntimeSettings.generated.inl"
+
 
 IMPLEMENT_MODULE(FDefaultModuleImpl, AndroidRuntimeSettings);
diff --git a/Engine/Source/Runtime/AutomationMessages/Private/AutomationMessagesModule.cpp b/Engine/Source/Runtime/AutomationMessages/Private/AutomationMessagesModule.cpp
--- a/Engine/Source/Runtime/AutomationMessages/Private/AutomationMessagesModule.cpp
+++ b/Engine/Source/Runtime/AutomationMessages/Private/AutomationMessagesModule.cpp
@@ -5,7 +5,7 @@
 =============================================================================*/
 
 #include "AutomationMessagesPrivatePCH.h"
-#include "AutomationMessages.generated.inl"
+
 
 
 /**
diff --git a/Engine/Source/Runtime/Core/Private/Containers/String.cpp b/Engine/Source/Runtime/Core/Private/Containers/String.cpp
--- a/Engine/Source/Runtime/Core/Private/Containers/String.cpp
+++ b/Engine/Source/Runtime/Core/Private/Containers/String.cpp
@@ -811,50 +811,33 @@ int32 FString::ParseIntoArrayWS( TArray<FString>* InArray, const TCHAR* pchExtra
 
 FString FString::Replace(const TCHAR* From, const TCHAR* To, ESearchCase::Type SearchCase) const
 {
-	if (IsEmpty())
+	if (IsEmpty() || !From || !*From)
 	{
 		return *this;
 	}
 
-	FString Result;
-
 	// get a pointer into the character data
-	TCHAR* Travel = (TCHAR*)Data.GetData();
+	const TCHAR* Travel = Data.GetData();
 
-	// precalc the length of the From string
+	// precalc the lengths of the replacement strings
 	int32 FromLength = FCString::Strlen(From);
+	int32 ToLength   = FCString::Strlen(To);
 
-	// FCString::Strstr will not behave like we want on empty From string
-	if (FromLength == 0)
-	{
-		return *this;
-	}
-
+	FString Result;
 	while (true)
 	{
 		// look for From in the remaining string
-		TCHAR* FromLocation = SearchCase == ESearchCase::IgnoreCase ? FCString::Stristr(Travel, From) : FCString::Strstr(Travel, From);
-		if (FromLocation)
-		{
-			// replace the first letter of the From with 0 so we can do a strcpy (FString +=)
-			TCHAR C = *FromLocation;
-			*FromLocation = 0;
-			
-			// copy everything up to the From
-			Result += Travel;
+		const TCHAR* FromLocation = SearchCase == ESearchCase::IgnoreCase ? FCString::Stristr(Travel, From) : FCString::Strstr(Travel, From);
+		if (!FromLocation)
+			break;
 
-			// copy over the To
-			Result += To;
+		// copy everything up to FromLocation
+		Result.AppendChars(Travel, FromLocation - Travel);
 
-			// retore the letter, just so we don't have 0's in the string
-			*FromLocation = *From;
+		// copy over the To
+		Result.AppendChars(To, ToLength);
 
-			Travel = FromLocation + FromLength;
-		}
-		else
-		{
-			break;
-		}
+		Travel = FromLocation + FromLength;
 	}
 
 	// copy anything left over
diff --git a/Engine/Source/Runtime/Core/Private/Modules/ModuleManager.cpp b/Engine/Source/Runtime/Core/Private/Modules/ModuleManager.cpp
--- a/Engine/Source/Runtime/Core/Private/Modules/ModuleManager.cpp
+++ b/Engine/Source/Runtime/Core/Private/Modules/ModuleManager.cpp
@@ -1610,19 +1610,14 @@ bool FModuleManager::BuildUnrealBuildTool(FOutputDevice &Ar)
 
 	// First determine the appropriate vcvars batch file to launch
 	FString VCVarsBat;
-	{
+
 #if _MSC_VER >= 1800
- 		TCHAR VS12Path[MAX_PATH];
- 		FPlatformMisc::GetEnvironmentVariable(TEXT("VS120COMNTOOLS"), VS12Path, ARRAY_COUNT(VS12Path));
- 		VCVarsBat = VS12Path;
- 		VCVarsBat += TEXT("/../../VC/bin/x86_amd64/vcvarsx86_amd64.bat");
+	FPlatformMisc::GetVSComnTools(12, VCVarsBat);
 #else
-		TCHAR VS11Path[MAX_PATH];
-		FPlatformMisc::GetEnvironmentVariable(TEXT("VS110COMNTOOLS"), VS11Path, ARRAY_COUNT(VS11Path));
-		VCVarsBat = VS11Path;
-		VCVarsBat += TEXT("/../../VC/bin/x86_amd64/vcvarsx86_amd64.bat");
+	FPlatformMisc::GetVSComnTools(11, VCVarsBat);
 #endif
-	}
+
+	VCVarsBat += FPaths::Combine(*VCVarsBat, L"../../VC/bin/x86_amd64/vcvarsx86_amd64.bat");
 
 	// Check to make sure we found one.
 	if ( VCVarsBat.IsEmpty() || !FPaths::FileExists(VCVarsBat) )
diff --git a/Engine/Source/Runtime/Core/Private/Serialization/Archive.cpp b/Engine/Source/Runtime/Core/Private/Serialization/Archive.cpp
--- a/Engine/Source/Runtime/Core/Private/Serialization/Archive.cpp
+++ b/Engine/Source/Runtime/Core/Private/Serialization/Archive.cpp
@@ -79,7 +79,6 @@ void FArchive::Reset()
 	ArNoDelta							= false;
 	ArIgnoreOuterRef					= false;
 	ArIgnoreClassRef					= false;
-	ArAllowEliminatingReferences		= true;
 	ArAllowLazyLoading					= false;
 	ArIsObjectReferenceCollector		= false;
 	ArIsModifyingWeakAndStrongReferences= false;
@@ -118,7 +117,6 @@ void FArchive::CopyTrivialFArchiveStatusMembers(const FArchive& ArchiveToCopy)
 	ArNoDelta                            = ArchiveToCopy.ArNoDelta;
 	ArIgnoreOuterRef                     = ArchiveToCopy.ArIgnoreOuterRef;
 	ArIgnoreClassRef                     = ArchiveToCopy.ArIgnoreClassRef;
-	ArAllowEliminatingReferences         = ArchiveToCopy.ArAllowEliminatingReferences;
 	ArAllowLazyLoading                   = ArchiveToCopy.ArAllowLazyLoading;
 	ArIsObjectReferenceCollector         = ArchiveToCopy.ArIsObjectReferenceCollector;
 	ArIsModifyingWeakAndStrongReferences = ArchiveToCopy.ArIsModifyingWeakAndStrongReferences;
diff --git a/Engine/Source/Runtime/Core/Private/Windows/WindowsPlatformMisc.cpp b/Engine/Source/Runtime/Core/Private/Windows/WindowsPlatformMisc.cpp
--- a/Engine/Source/Runtime/Core/Private/Windows/WindowsPlatformMisc.cpp
+++ b/Engine/Source/Runtime/Core/Private/Windows/WindowsPlatformMisc.cpp
@@ -1940,6 +1940,38 @@ bool FWindowsPlatformMisc::QueryRegKey( const HKEY InKey, const TCHAR* InSubKey,
 	return bSuccess;
 }
 
+bool FWindowsPlatformMisc::GetVSComnTools(int32 Version, FString& OutData)
+{
+	checkf(11 <= Version && Version <= 12, L"Not supported Visual Studio version.");
+
+	const TCHAR* PossibleRegPaths[] = {
+		L"Wow6432Node\\Microsoft\\VisualStudio",	// Non-express VS2013 on 64-bit machine.
+		L"Microsoft\\VisualStudio",					// Non-express VS2013 on 32-bit machine.
+		L"Wow6432Node\\Microsoft\\WDExpress",		// Express VS2013 on 64-bit machine.
+		L"Microsoft\\WDExpress"						// Express VS2013 on 32-bit machine.
+	};
+
+	bool bResult = false;
+	FString IDEPath;
+
+	for (int32 Index = 0; Index < sizeof(PossibleRegPaths) / sizeof(const TCHAR *); ++Index)
+	{
+		bResult = QueryRegKey(HKEY_LOCAL_MACHINE, *FString::Printf(L"SOFTWARE\\%s\\%d.0", PossibleRegPaths[Index], Version), L"InstallDir", IDEPath);
+
+		if (bResult)
+		{
+			break;
+		}
+	}
+	
+	if(bResult)
+	{
+		OutData = FPaths::ConvertRelativePathToFull(FPaths::Combine(*IDEPath, L"..", L"Tools"));
+	}
+
+	return bResult;
+}
+
 const TCHAR* FWindowsPlatformMisc::GetDefaultPathSeparator()
 {
 	return TEXT( "\\" );
diff --git a/Engine/Source/Runtime/Core/Public/Containers/Array.h b/Engine/Source/Runtime/Core/Public/Containers/Array.h
--- a/Engine/Source/Runtime/Core/Public/Containers/Array.h
+++ b/Engine/Source/Runtime/Core/Public/Containers/Array.h
@@ -374,11 +374,16 @@ public:
 		CopyToEmpty(Other);
 	}
 
-	TArray(const TArray<ElementType,Allocator>& Other)
+	TArray(const TArray& Other)
 	{
 		CopyToEmpty(Other);
 	}
 
+	TArray(const TArray& Other, int32 ExtraSlack)
+	{
+		CopyToEmpty(Other, ExtraSlack);
+	}
+
 	template<typename OtherAllocator>
 	TArray& operator=( const TArray<ElementType,OtherAllocator>& Other )
 	{
@@ -412,9 +417,17 @@ private:
 	}
 
 	template <typename ArrayType>
-	static FORCEINLINE typename TEnableIf<!TContainerTraits<ArrayType>::MoveWillEmptyContainer>::Type MoveOrCopy(ArrayType& ToArray, ArrayType& FromArray)
+	static FORCEINLINE typename TEnableIf<TContainerTraits<ArrayType>::MoveWillEmptyContainer>::Type MoveOrCopy(ArrayType& ToArray, ArrayType& FromArray, int32 ExtraSlack)
+	{
+		MoveOrCopy(ToArray, FromArray);
+
+		ToArray.Reserve(ToArray.ArrayNum + ExtraSlack);
+	}
+
+	template <typename ArrayType>
+	static FORCEINLINE typename TEnableIf<!TContainerTraits<ArrayType>::MoveWillEmptyContainer>::Type MoveOrCopy(ArrayType& ToArray, ArrayType& FromArray, int32 ExtraSlack = 0)
 	{
-		ToArray.CopyToEmpty(FromArray);
+		ToArray.CopyToEmpty(FromArray, ExtraSlack);
 	}
 
 public:
@@ -427,6 +440,11 @@ public:
 		MoveOrCopy(*this, Other);
 	}
 
+	TArray(TArray<ElementType,Allocator>&& Other, int32 ExtraSlack)
+	{
+		MoveOrCopy(*this, Other, ExtraSlack);
+	}
+
 	TArray& operator=(TArray<ElementType,Allocator>&& Other)
 	{
 		if (this != &Other)
@@ -1490,15 +1508,17 @@ private:
 	 * @param Source the source array to copy
 	 */
 	template<typename OtherAllocator>
-	void CopyToEmpty(const TArray<ElementType,OtherAllocator>& Source)
+	void CopyToEmpty(const TArray<ElementType,OtherAllocator>& Source, int32 ExtraSlack = 0)
 	{
+		check(ExtraSlack >= 0);
+
 		int32 SourceCount = Source.Num();
-		AllocatorInstance.ResizeAllocation(0, SourceCount, sizeof(ElementType));
+		AllocatorInstance.ResizeAllocation(0, SourceCount + ExtraSlack, sizeof(ElementType));
 
 		CopyConstructItems(GetTypedData(), Source.GetTypedData(), SourceCount);
 
 		ArrayNum = SourceCount;
-		ArrayMax = SourceCount;
+		ArrayMax = SourceCount + ExtraSlack;
 	}
 
 protected:
diff --git a/Engine/Source/Runtime/Core/Public/Containers/UnrealString.h b/Engine/Source/Runtime/Core/Public/Containers/UnrealString.h
--- a/Engine/Source/Runtime/Core/Public/Containers/UnrealString.h
+++ b/Engine/Source/Runtime/Core/Public/Containers/UnrealString.h
@@ -45,15 +45,10 @@ public:
 #if PLATFORM_COMPILER_HAS_DEFAULTED_FUNCTIONS
 
 	FString() = default;
-	FString( const FString& ) = default;
-	FString& operator=( const FString& ) = default;
-
-	#if PLATFORM_COMPILER_HAS_RVALUE_REFERENCES
-
-		FString( FString&& ) = default;
-		FString& operator=( FString&& ) = default;
-
-	#endif
+	FString(FString&&) = default;
+	FString(const FString&) = default;
+	FString& operator=(FString&&) = default;
+	FString& operator=(const FString&) = default;
 
 #else
 
@@ -61,34 +56,27 @@ public:
 	{
 	}
 
-	FORCEINLINE FString( const FString& Other )
+	FORCEINLINE FString(FString&& Other)
+		: Data(MoveTemp(Other.Data))
 	{
-		Data = Other.Data;
 	}
 
-	FORCEINLINE FString& operator=( const FString& Other )
+	FORCEINLINE FString(const FString& Other)
+		: Data(Other.Data)
 	{
-		if (this != &Other)
-		{
-			Data = Other.Data;
-		}
-		return *this;
 	}
 
-	#if PLATFORM_COMPILER_HAS_RVALUE_REFERENCES
-
-		FORCEINLINE FString(FString&& Other)
-			: Data(MoveTemp(Other.Data))
-		{
-		}
-
-		FORCEINLINE FString& operator=(FString&& Other)
-		{
-			Data = MoveTemp(Other.Data);
-			return *this;
-		}
+	FORCEINLINE FString& operator=(FString&& Other)
+	{
+		Data = MoveTemp(Other.Data);
+		return *this;
+	}
 
-	#endif
+	FORCEINLINE FString& operator=(const FString& Other)
+	{
+		Data = Other.Data;
+		return *this;
+	}
 
 #endif
 
@@ -98,15 +86,20 @@ public:
 	 * @param Other the other string to create a new copy from
 	 * @param ExtraSlack number of extra characters to add to the end of the other string in this string
 	 */
-	FORCEINLINE FString( const FString& Other, int32 ExtraSlack)
+	FORCEINLINE FString(const FString& Other, int32 ExtraSlack)
+		: Data(Other.Data, ExtraSlack + ((Other.Data.Num() || !ExtraSlack) ? 0 : 1)) // Add 1 if the source string array is empty and we want some slack, because we'll need to include a null terminator which is currently missing
 	{
-		Data.Empty(Other.Data.Num() + ExtraSlack);
+	}
 
-		if( Other.Data.Num() )
-		{
-			Data.AddUninitialized(Other.Data.Num());
-			FMemory::Memcpy( Data.GetData(), Other.Data.GetData(), Data.Num() * sizeof(TCHAR) );
-		}
+	/**
+	 * Create a copy of the Other string with extra space for characters at the end of the string
+	 *
+	 * @param Other the other string to create a new copy from
+	 * @param ExtraSlack number of extra characters to add to the end of the other string in this string
+	 */
+	FORCEINLINE FString(FString&& Other, int32 ExtraSlack)
+		: Data(MoveTemp(Other.Data), ExtraSlack + ((Other.Data.Num() || !ExtraSlack) ? 0 : 1)) // Add 1 if the source string array is empty and we want some slack, because we'll need to include a null terminator which is currently missing
+	{
 	}
 
 	/**
@@ -248,8 +241,9 @@ public:
 	 */
 	FORCEINLINE void CheckInvariants() const
 	{
-		checkSlow(Data.Num() >= 0);
-		checkSlow(!Data.Num() || !(Data[Data.Num()-1]));
+		int32 Num = Data.Num();
+		checkSlow(Num >= 0);
+		checkSlow(!Num || !Data.GetTypedData()[Num - 1]);
 		checkSlow(Data.GetSlack() >= 0);
 	}
 
@@ -322,6 +316,35 @@ public:
 	}
 #endif
 
+	/**
+	 * Appends an array of characters to the string.  
+	 *
+	 * @param Array A pointer to the start of an array of characters to append.  This array need not be null-terminated, and null characters are not treated specially.
+	 * @param Count The number of characters to copy from Array.
+	 */
+	FORCEINLINE void AppendChars(const TCHAR* Array, int32 Count)
+	{
+		check(Count >= 0);
+
+		if (!Count)
+			return;
+
+		checkSlow(Array);
+
+		int32 Index = Data.Num();
+
+		// Reserve enough space - including an extra gap for a null terminator if we don't already have a string allocated
+		Data.AddUninitialized(Count + (Index ? 0 : 1));
+
+		TCHAR* EndPtr = Data.GetTypedData() + Index - (Index ? 1 : 0);
+
+		// Copy characters to end of string, overwriting null terminator if we already have one
+		CopyAssignItems(EndPtr, Array, Count);
+
+		// (Re-)establish the null terminator
+		*(EndPtr + Count) = 0;
+	}
+
 	/**
 	 * Concatenate this with given string
 	 * 
@@ -333,20 +356,8 @@ public:
 		checkSlow(Str);
 		CheckInvariants();
 
-		if( Str && *Str )
-		{
-			int32 Index = Data.Num();
-			int32 NumAdd = FCString::Strlen(Str)+1;
-			int32 NumCopy = NumAdd;
+		AppendChars(Str, FCString::Strlen(Str));
 
-			if (Data.Num())
-			{
-				NumAdd--;
-				Index--;
-			}
-			Data.AddUninitialized( NumAdd );
-			FMemory::Memcpy( &Data[Index], Str, NumCopy * sizeof(TCHAR) );
-		}
 		return *this;
 	}
 
@@ -402,19 +413,19 @@ public:
 	}
 
 	/**
-	* Removes the text from the start of the string if it exists.
-	*
-	* @param InPrefix the prefix to search for at the start of the string to remove.
-	* @return true if the prefix was removed, otherwise false.
-	*/
+	 * Removes the text from the start of the string if it exists.
+	 *
+	 * @param InPrefix the prefix to search for at the start of the string to remove.
+	 * @return true if the prefix was removed, otherwise false.
+	 */
 	bool RemoveFromStart( const FString& InPrefix, ESearchCase::Type SearchCase = ESearchCase::IgnoreCase );
 
 	/**
-	* Removes the text from the end of the string if it exists.
-	*
-	* @param InSuffix the suffix to search for at the end of the string to remove.
-	* @return true if the suffix was removed, otherwise false.
-	*/
+	 * Removes the text from the end of the string if it exists.
+	 *
+	 * @param InSuffix the suffix to search for at the end of the string to remove.
+	 * @return true if the suffix was removed, otherwise false.
+	 */
 	bool RemoveFromEnd( const FString& InSuffix, ESearchCase::Type SearchCase = ESearchCase::IgnoreCase );
 
 	/**
@@ -428,95 +439,211 @@ public:
 		CheckInvariants();
 		Str.CheckInvariants();
 
-		if( Str.Len() )
-		{
-			int32 Index = Data.Num();
-			int32 NumAdd = Str.Data.Num();
-			
-			if (Data.Num())
-			{
-				NumAdd--;
-				Index--;
-			}
-			
-			Data.AddUninitialized( NumAdd );
+		AppendChars(Str.Data.GetTypedData(), Str.Len());
 
-			FMemory::Memcpy( &Data[Index], Str.Data.GetData(), Str.Data.Num() * sizeof(TCHAR) );
-		}
 		return *this;
 	}
 
 	/**
-	 * Concatenate this with given char
+	 * Concatenates an FString with a TCHAR.
 	 * 
-	 * @param inChar other Char to be concatenated onto the end of this string
-	 * @return reference to this
+	 * @param Lhs The FString on the left-hand-side of the expression.
+	 * @param Rhs The char on the right-hand-side of the expression.
+	 *
+	 * @return The concatenated string.
 	 */
-	FORCEINLINE FString operator+( const TCHAR InChar ) const
+	FORCEINLINE friend FString operator+(const FString& Lhs, TCHAR Rhs)
 	{
-		CheckInvariants();
-		return FString(*this, 2) += InChar; // may have an extra character of "slack"
+		Lhs.CheckInvariants();
+
+		FString Result(Lhs, 1);
+		Result += Rhs;
+
+		return Result;
+	}
+
+	/**
+	 * Concatenates an FString with a TCHAR.
+	 * 
+	 * @param Lhs The FString on the left-hand-side of the expression.
+	 * @param Rhs The char on the right-hand-side of the expression.
+	 *
+	 * @return The concatenated string.
+	 */
+	FORCEINLINE friend FString operator+(FString&& Lhs, TCHAR Rhs)
+	{
+		Lhs.CheckInvariants();
+
+		FString Result(MoveTemp(Lhs), 1);
+		Result += Rhs;
+
+		return Result;
+	}
+
+private:
+	template <typename LhsType, typename RhsType>
+	FORCEINLINE static FString ConcatFStrings(typename TIdentity<LhsType>::Type Lhs, typename TIdentity<RhsType>::Type Rhs)
+	{
+		Lhs.CheckInvariants();
+		Rhs.CheckInvariants();
+
+		if (Lhs.IsEmpty())
+			return MoveTemp(Rhs);
+
+		int32 RhsLen = Rhs.Len();
+
+		FString Result(MoveTemp(Lhs), RhsLen);
+		Result.AppendChars(Rhs.Data.GetTypedData(), RhsLen);
+		
+		return Result;
+	}
+
+	template <typename RhsType>
+	FORCEINLINE static FString ConcatTCHARsToFString(const TCHAR* Lhs, typename TIdentity<RhsType>::Type Rhs)
+	{
+		checkSlow(Lhs);
+		Rhs.CheckInvariants();
+
+		if (!Lhs || !*Lhs)
+			return MoveTemp(Rhs);
+
+		int32 LhsLen = FCString::Strlen(Lhs);
+		int32 RhsLen = Rhs.Len();
+
+		// This is not entirely optimal, as if the Rhs is an rvalue and has enough slack space to hold Lhs, then
+		// the memory could be reused here without constructing a new object.  However, until there is proof otherwise,
+		// I believe this will be relatively rare and isn't worth making the code a lot more complex right now.
+		FString Result;
+		Result.Data.AddUninitialized(LhsLen + RhsLen + 1);
+
+		TCHAR* ResultData = Result.Data.GetTypedData();
+		CopyAssignItems(ResultData, Lhs, LhsLen);
+		CopyAssignItems(ResultData + LhsLen, Rhs.Data.GetTypedData(), RhsLen);
+		*(ResultData + LhsLen + RhsLen) = 0;
+		
+		return Result;
 	}
 
+	template <typename LhsType>
+	FORCEINLINE static FString ConcatFStringToTCHARs(typename TIdentity<LhsType>::Type Lhs, const TCHAR* Rhs)
+	{
+		Lhs.CheckInvariants();
+		checkSlow(Rhs);
+
+		if (!Rhs || !*Rhs)
+			return MoveTemp(Lhs);
+
+		int32 RhsLen = FCString::Strlen(Rhs);
+
+		FString Result(MoveTemp(Lhs), RhsLen);
+		Result.AppendChars(Rhs, RhsLen);
+		
+		return Result;
+	}
 
+public:
 	/**
-	 * Concatenate this with given string
+	 * Concatenate two FStrings.
 	 * 
-	 * @param Str array of TCHAR to be concatenated onto the end of this
-	 * @return new FString
+	 * @param Lhs The FString on the left-hand-side of the expression.
+	 * @param Rhs The FString on the right-hand-side of the expression.
+	 *
+	 * @return The concatenated string.
 	 */
-	FORCEINLINE FString operator+( const TCHAR* Str ) const
+	FORCEINLINE friend FString operator+(const FString& Lhs, const FString& Rhs)
 	{
-		checkSlow(Str);
-		CheckInvariants();
+		return ConcatFStrings<const FString&, const FString&>(Lhs, Rhs);
+	}
 
-		if( Str && *Str )
-		{
-			int32 Index = Data.Num();
-			int32 NumAdd = FCString::Strlen(Str)+1;
-			int32 NumCopy = NumAdd;
-			if (Data.Num())
-			{
-				NumAdd--;
-				Index--;
-			}
-			
-			FString Ret( *this, NumAdd ) ; 
-			Ret.Data.AddUninitialized( NumAdd );
-			FMemory::Memcpy( &(Ret)[Index], Str, NumCopy * sizeof(TCHAR) );
-			return Ret;
-		}
-		return *this;
+	/**
+	 * Concatenate two FStrings.
+	 * 
+	 * @param Lhs The FString on the left-hand-side of the expression.
+	 * @param Rhs The FString on the right-hand-side of the expression.
+	 *
+	 * @return The concatenated string.
+	 */
+	FORCEINLINE friend FString operator+(FString&& Lhs, const FString& Rhs)
+	{
+		return ConcatFStrings<FString&&, const FString&>(MoveTemp(Lhs), Rhs);
 	}
 
+	/**
+	 * Concatenate two FStrings.
+	 * 
+	 * @param Lhs The FString on the left-hand-side of the expression.
+	 * @param Rhs The FString on the right-hand-side of the expression.
+	 *
+	 * @return The concatenated string.
+	 */
+	FORCEINLINE friend FString operator+(const FString& Lhs, FString&& Rhs)
+	{
+		return ConcatFStrings<const FString&, FString&&>(Lhs, MoveTemp(Rhs));
+	}
 
 	/**
-	 * Concatenate this with given string
+	 * Concatenate two FStrings.
 	 * 
-	 * @param Str other FString to be concatenated onto the end of this
-	 * @return new FString
+	 * @param Lhs The FString on the left-hand-side of the expression.
+	 * @param Rhs The FString on the right-hand-side of the expression.
+	 *
+	 * @return The concatenated string.
 	 */
-	FORCEINLINE FString operator+( const FString& Str ) const
+	FORCEINLINE friend FString operator+(FString&& Lhs, FString&& Rhs)
 	{
-		CheckInvariants();
-		Str.CheckInvariants();
-		
-		if( Str.Data.Num() )
-		{
-			int32 Index = Data.Num();
-			int32 NumAdd = Str.Data.Num();
-			if (Data.Num())
-			{
-				NumAdd--;
-				Index--;
-			}
-			
-			FString ret( *this, NumAdd ) ; 
-			ret.Data.AddUninitialized( NumAdd );
-			FMemory::Memcpy( &(ret)[Index], Str.Data.GetData(), Str.Data.Num() * sizeof(TCHAR) );
-			return ret;
-		}
-		return *this;
+		return ConcatFStrings<FString&&, FString&&>(MoveTemp(Lhs), MoveTemp(Rhs));
+	}
+
+	/**
+	 * Concatenates a TCHAR array to an FString.
+	 * 
+	 * @param Lhs The TCHAR array on the left-hand-side of the expression.
+	 * @param Rhs The FString on the right-hand-side of the expression.
+	 *
+	 * @return The concatenated string.
+	 */
+	FORCEINLINE friend FString operator+(const TCHAR* Lhs, const FString& Rhs)
+	{
+		return ConcatTCHARsToFString<const FString&>(Lhs, Rhs);
+	}
+
+	/**
+	 * Concatenates a TCHAR array to an FString.
+	 * 
+	 * @param Lhs The TCHAR array on the left-hand-side of the expression.
+	 * @param Rhs The FString on the right-hand-side of the expression.
+	 *
+	 * @return The concatenated string.
+	 */
+	FORCEINLINE friend FString operator+(const TCHAR* Lhs, FString&& Rhs)
+	{
+		return ConcatTCHARsToFString<FString&&>(Lhs, MoveTemp(Rhs));
+	}
+
+	/**
+	 * Concatenates an FString to a TCHAR array.
+	 * 
+	 * @param Lhs The FString on the left-hand-side of the expression.
+	 * @param Rhs The TCHAR array on the right-hand-side of the expression.
+	 *
+	 * @return The concatenated string.
+	 */
+	FORCEINLINE friend FString operator+(const FString& Lhs, const TCHAR* Rhs)
+	{
+		return ConcatFStringToTCHARs<const FString&>(Lhs, Rhs);
+	}
+
+	/**
+	 * Concatenates an FString to a TCHAR array.
+	 * 
+	 * @param Lhs The FString on the left-hand-side of the expression.
+	 * @param Rhs The TCHAR array on the right-hand-side of the expression.
+	 *
+	 * @return The concatenated string.
+	 */
+	FORCEINLINE friend FString operator+(FString&& Lhs, const TCHAR* Rhs)
+	{
+		return ConcatFStringToTCHARs<FString&&>(MoveTemp(Lhs), Rhs);
 	}
 
 	/**
diff --git a/Engine/Source/Runtime/Core/Public/Serialization/ArchiveBase.h b/Engine/Source/Runtime/Core/Public/Serialization/ArchiveBase.h
--- a/Engine/Source/Runtime/Core/Public/Serialization/ArchiveBase.h
+++ b/Engine/Source/Runtime/Core/Public/Serialization/ArchiveBase.h
@@ -536,21 +536,6 @@ public:
 
 	virtual void IndicateSerializationMismatch() { }
 
-	/**
-	 * Whether to allow serialization code to potentially remove references. Note that this
-	 * is increasing/ decreasing an internal counter and hence supporting nesting so elimination
-	 * might not be enabled after calling it. 
-	 *
-	 * @param	Allow Whether to increase or decrease counter determining whether to allow elimination
-	 * @return	true if reference elimination is enabled, false otherwise
-	 */
-	bool AllowEliminatingReferences(bool Allow)
-	{
-		ArAllowEliminatingReferences += Allow ? 1 : -1;
-
-		return IsAllowingReferenceElimination();
-	}
-
 	// Logf implementation for convenience.
 	VARARG_DECL(void, void, {}, Logf, VARARG_NONE, const TCHAR*, VARARG_NONE, VARARG_NONE);
 
@@ -689,11 +674,6 @@ public:
 		return ArIgnoreClassRef;
 	}
 
-	FORCEINLINE bool IsAllowingReferenceElimination() const
-	{
-		return (ArAllowEliminatingReferences > 0) ? true : false;
-	}
-
 	FORCEINLINE bool IsAllowingLazyLoading() const
 	{
 		return ArAllowLazyLoading;
@@ -1025,9 +1005,6 @@ public:
 	/** Whether we are currently serializing defaults. > 0 means yes, <= 0 means no. */
 	int32 ArSerializingDefaults;
 
-	/** Whether we are allowed to potentially NULL references to UObjects, > 0 means yes, <= 0 means no. */
-	int32 ArAllowEliminatingReferences;
-
 	/** Modifier flags that be used when serializing UProperties */
 	uint32 ArPortFlags;
 			
diff --git a/Engine/Source/Runtime/Core/Public/Templates/UnrealTemplate.h b/Engine/Source/Runtime/Core/Public/Templates/UnrealTemplate.h
--- a/Engine/Source/Runtime/Core/Public/Templates/UnrealTemplate.h
+++ b/Engine/Source/Runtime/Core/Public/Templates/UnrealTemplate.h
@@ -478,3 +478,18 @@ template <typename From, typename To> struct TCopyQualifiersFromTo
 template <typename From, typename To> struct TCopyQualifiersFromTo<const          From, To> { typedef const          To Type; };
 template <typename From, typename To> struct TCopyQualifiersFromTo<      volatile From, To> { typedef       volatile To Type; };
 template <typename From, typename To> struct TCopyQualifiersFromTo<const volatile From, To> { typedef const volatile To Type; };
+
+/**
+ * Returns the same type passed to it.  This is useful in a few cases, but mainly for inhibiting template argument deduction in function arguments, e.g.:
+ *
+ * template <typename T>
+ * void Func1(T Val); // Can be called like Func(123) or Func<int>(123);
+ *
+ * template <typename T>
+ * void Func2(typename TIdentity<T>::Type Val); // Must be called like Func<int>(123)
+ */
+template <typename T>
+struct TIdentity
+{
+	typedef T Type;
+};
diff --git a/Engine/Source/Runtime/Core/Public/UObject/ObjectVersion.h b/Engine/Source/Runtime/Core/Public/UObject/ObjectVersion.h
--- a/Engine/Source/Runtime/Core/Public/UObject/ObjectVersion.h
+++ b/Engine/Source/Runtime/Core/Public/UObject/ObjectVersion.h
@@ -520,7 +520,8 @@ enum EUnrealEngineObjectUE4Version
 	VER_UE4_ADD_CUSTOMPROFILENAME_CHANGE, 
 	// Permanently flip and scale material expression coordinates
 	VER_UE4_FLIP_MATERIAL_COORDS,
-
+	// PinSubCategoryMemberReference added to FEdGraphPinType
+	VER_UE4_MEMBERREFERENCE_IN_PINTYPE,
 	// -----<new versions can be added before this line>-------------------------------------------------
 	// - this needs to be the last line (see note below)
 	VER_UE4_AUTOMATIC_VERSION_PLUS_ONE,
diff --git a/Engine/Source/Runtime/Core/Public/Windows/WindowsPlatformMisc.h b/Engine/Source/Runtime/Core/Public/Windows/WindowsPlatformMisc.h
--- a/Engine/Source/Runtime/Core/Public/Windows/WindowsPlatformMisc.h
+++ b/Engine/Source/Runtime/Core/Public/Windows/WindowsPlatformMisc.h
@@ -175,6 +175,16 @@ struct CORE_API FWindowsPlatformMisc : public FGenericPlatformMisc
 	 */
 	static bool QueryRegKey( const HKEY InKey, const TCHAR* InSubKey, const TCHAR* InValueName, FString& OutData );
 
+	/**
+	 * Gets Visual Studio common tools path.
+	 *
+	 * @param Version Version of VS to get (11 - 2012, 12 - 2013).
+	 * @param OutData Output parameter with common tools path.
+	 *
+	 * @return Returns if succeeded.
+	 */
+	static bool GetVSComnTools(int32 Version, FString& OutData);
+
 	/**
 	 * Returns the size of the cache line in bytes.
 	 *
diff --git a/Engine/Source/Runtime/CoreUObject/Private/UObject/Class.cpp b/Engine/Source/Runtime/CoreUObject/Private/UObject/Class.cpp
--- a/Engine/Source/Runtime/CoreUObject/Private/UObject/Class.cpp
+++ b/Engine/Source/Runtime/CoreUObject/Private/UObject/Class.cpp
@@ -439,7 +439,8 @@ void UStruct::Link(FArchive& Ar, bool bRelinkExistingProperties)
 				// If we don't have the editor, make sure we aren't trying to link properties that are editor only.
 				check( !Property->IsEditorOnlyProperty() );
 #endif // WITH_EDITORONLY_DATA
-
+				ensureMsgf(Property->GetOuter() == this, TEXT("Linking '%s'. Property '%s' has outer '%s'"),
+					*GetFullName(), *Property->GetName(), *Property->GetOuter()->GetFullName());
 				PropertiesSize = Property->Link(Ar);
 				MinAlignment = FMath::Max( MinAlignment, Property->GetMinAlignment() );
 			}
@@ -2136,6 +2137,8 @@ void UScriptStruct::DestroyScriptStruct(void* Dest, int32 ArrayDim)
 	}
 }
 
+void UScriptStruct::RecursivelyPreload() {}
+
 IMPLEMENT_CORE_INTRINSIC_CLASS(UScriptStruct, UStruct,
 	{
 	}
diff --git a/Engine/Source/Runtime/CoreUObject/Private/UObject/CoreNative.cpp b/Engine/Source/Runtime/CoreUObject/Private/UObject/CoreNative.cpp
--- a/Engine/Source/Runtime/CoreUObject/Private/UObject/CoreNative.cpp
+++ b/Engine/Source/Runtime/CoreUObject/Private/UObject/CoreNative.cpp
@@ -4,9 +4,6 @@
 #include "ModuleManager.h"
 #include "StringAssetReference.h"
 #include "StringClassReference.h"
-#if USE_COMPILED_IN_NATIVES
-#include "CoreUObject.generated.inl"
-#endif
 
 // CoreUObject module. Handles UObject system pre-init (registers init function with Core callbacks).
 class FCoreUObjectModule : public FDefaultModuleImpl
diff --git a/Engine/Source/Runtime/CoreUObject/Private/UObject/Obj.cpp b/Engine/Source/Runtime/CoreUObject/Private/UObject/Obj.cpp
--- a/Engine/Source/Runtime/CoreUObject/Private/UObject/Obj.cpp
+++ b/Engine/Source/Runtime/CoreUObject/Private/UObject/Obj.cpp
@@ -740,13 +740,10 @@ void UObject::Serialize( FArchive& Ar )
 	if( (!Ar.IsLoading() && !Ar.IsSaving() && !Ar.IsObjectReferenceCollector()) )
 	{
 		Ar << LoadName;
-		// We don't want to have the following potentially be clobbered by GC code.
-		Ar.AllowEliminatingReferences( false );
 		if(!Ar.IsIgnoringOuterRef())
 		{
 			Ar << LoadOuter;
 		}
-		Ar.AllowEliminatingReferences( true );
 		if ( !Ar.IsIgnoringClassRef() )
 		{
 			Ar << Class;
diff --git a/Engine/Source/Runtime/CoreUObject/Private/UObject/PropertyStruct.cpp b/Engine/Source/Runtime/CoreUObject/Private/UObject/PropertyStruct.cpp
--- a/Engine/Source/Runtime/CoreUObject/Private/UObject/PropertyStruct.cpp
+++ b/Engine/Source/Runtime/CoreUObject/Private/UObject/PropertyStruct.cpp
@@ -38,6 +38,8 @@ void UStructProperty::LinkInternal(FArchive& Ar)
 	// Preload is required here in order to load the value of Struct->PropertiesSize
 	Ar.Preload(Struct);
 	checkSlow(Struct);
+	Struct->RecursivelyPreload();
+	
 	ElementSize = Align(Struct->PropertiesSize, Struct->GetMinAlignment());
 	if (Struct->StructFlags & STRUCT_IsPlainOldData) // if there is nothing to construct or the struct is known to be memcpy-able, then allow memcpy
 	{
@@ -153,11 +155,19 @@ bool UStructProperty::NetSerializeItem( FArchive& Ar, UPackageMap* Map, void* Da
 void UStructProperty::Serialize( FArchive& Ar )
 {
 	Super::Serialize( Ar );
+
+	static UScriptStruct* FallbackStruct = GetFallbackStruct();
+	
+	if (Ar.IsPersistent() && Ar.GetLinker() && Ar.IsLoading() && !Struct)
+	{
+		// It's necessary to solve circular dependency problems, when serializing the Struct causes linking of the Property.
+		Struct = FallbackStruct;
+	}
+
 	Ar << Struct;
 #if WITH_EDITOR
 	if (Ar.IsPersistent() && Ar.GetLinker())
 	{
-		UScriptStruct* FallbackStruct = GetFallbackStruct();
 		if (!Struct && Ar.IsLoading())
 		{
 			UE_LOG(LogProperty, Error, TEXT("UStructProperty::Serialize Loading: Property '%s'. Unknown structure."), *GetFullName());
@@ -169,6 +179,14 @@ void UStructProperty::Serialize( FArchive& Ar )
 		}
 	}
 #endif // WITH_EDITOR
+	if (Struct)
+	{
+		Struct->RecursivelyPreload();
+	}
+	else
+	{
+		ensure(true);
+	}
 }
 void UStructProperty::AddReferencedObjects(UObject* InThis, FReferenceCollector& Collector)
 {
diff --git a/Engine/Source/Runtime/CoreUObject/Private/UObject/UTextProperty.cpp b/Engine/Source/Runtime/CoreUObject/Private/UObject/UTextProperty.cpp
--- a/Engine/Source/Runtime/CoreUObject/Private/UObject/UTextProperty.cpp
+++ b/Engine/Source/Runtime/CoreUObject/Private/UObject/UTextProperty.cpp
@@ -21,7 +21,7 @@ bool UTextProperty::Identical( const void* A, const void* B, uint32 PortFlags )
 		{
 			//Culture invariant text don't have a namespace/key so we compare the source string
 			//Transient text don't have a namespace/key or source so we compare the display string
-			return *FTextInspector::GetDisplayString(ValueA) == *FTextInspector::GetDisplayString(ValueB);
+			return FTextInspector::GetDisplayString(ValueA) == FTextInspector::GetDisplayString(ValueB);
 		}
 
 		if (GIsEditor)
diff --git a/Engine/Source/Runtime/CoreUObject/Public/UObject/Class.h b/Engine/Source/Runtime/CoreUObject/Public/UObject/Class.h
--- a/Engine/Source/Runtime/CoreUObject/Public/UObject/Class.h
+++ b/Engine/Source/Runtime/CoreUObject/Public/UObject/Class.h
@@ -271,7 +271,7 @@ public:
 	 */
 	void SerializeBinEx( FArchive& Ar, void* Data, void const* DefaultData, UStruct* DefaultStruct ) const;
 
-	void SerializeTaggedProperties( FArchive& Ar, uint8* Data, UStruct* DefaultsStruct, uint8* Defaults ) const;
+	virtual void SerializeTaggedProperties( FArchive& Ar, uint8* Data, UStruct* DefaultsStruct, uint8* Defaults ) const;
 
 	virtual EExprToken SerializeExpr(int32& iCode, FArchive& Ar);
 	virtual void TagSubobjects(EObjectFlags NewFlags) OVERRIDE;
@@ -1014,7 +1014,7 @@ public:
 	 * @param	ArrayDim	Number of elements in the array
 	 * @param	Stride		Stride of the array, If this default (0), then we will pull the size from the struct
 	 */
-	void CopyScriptStruct(void* Dest, void const* Src, int32 ArrayDim = 1);
+	COREUOBJECT_API void CopyScriptStruct(void* Dest, void const* Src, int32 ArrayDim = 1);
 	/**
 	 * Initialize a struct over uninitialized memory. This may be done by calling the native constructor or individually initializing properties
 	 *
@@ -1040,6 +1040,7 @@ public:
 	 */
 	COREUOBJECT_API void DestroyScriptStruct(void* Dest, int32 ArrayDim = 1);
 
+	virtual COREUOBJECT_API void RecursivelyPreload();
 };
 
 
