Fixed wrong header parsing order in UnrealHeaderTool

diff --git a/Engine/Source/Programs/UnrealHeaderTool/Private/CodeGenerator.cpp b/Engine/Source/Programs/UnrealHeaderTool/Private/CodeGenerator.cpp
--- a/Engine/Source/Programs/UnrealHeaderTool/Private/CodeGenerator.cpp
+++ b/Engine/Source/Programs/UnrealHeaderTool/Private/CodeGenerator.cpp
@@ -4628,38 +4628,39 @@ ECompilationResult::Type UnrealHeaderTool_Main(const FString& ModuleInfoFilename
 	int32 NumFailures = 0;
 
 	// Three passes.  1) Public 'Classes' headers (legacy)  2) Public headers   3) Private headers
-	for( int32 PassIndex = 0; PassIndex < 3; ++PassIndex )
+	enum EHeaderFolderTypes
 	{
-		enum
-		{
-			PublicClassesHeaders,
-			PublicHeaders,
-			PrivateHeaders,
-		} CurrentlyProcessing;
-		{
-			switch( PassIndex )
-			{
-				case 0:
-				default:
-					CurrentlyProcessing = PublicClassesHeaders;
-					break;
+		PublicClassesHeaders = 0,
+		PublicHeaders = 1,
+		PrivateHeaders,
 
-				case 1:
-					CurrentlyProcessing = PublicHeaders;
-					break;
+		FolderType_Count
+	};
 
-				case 2:
-					CurrentlyProcessing = PrivateHeaders;
-					break;
-			};
+	for (const auto& Module : GManifest.Modules)
+	{
+		if (Result != ECompilationResult::Succeeded)
+		{
+			break;
 		}
 
-		for (const auto& Module : GManifest.Modules)
+		UPackage* Package = Cast<UPackage>(StaticFindObjectFast(UPackage::StaticClass(), NULL, FName(*Module.LongPackageName), false, false));
+		if (Package == NULL)
 		{
-			if (Result != ECompilationResult::Succeeded)
-			{
-				break;
-			}
+			Package = CreatePackage(NULL, *Module.LongPackageName);
+		}
+		// Set some package flags for indicating that this package contains script
+		// NOTE: We do this even if we didn't have to create the package, because CoreUObject is compiled into UnrealHeaderTool and we still
+		//       want to make sure our flags get set
+		Package->PackageFlags |= PKG_ContainsScript;
+		Package->PackageFlags &= ~(PKG_ClientOptional | PKG_ServerSideOnly);
+		Package->PackageFlags |= PKG_Compiling;
+
+		GPackageToManifestModuleMap.Add(Package, &Module);
+
+		for (int32 PassIndex = 0; PassIndex < FolderType_Count && Result == ECompilationResult::Succeeded; ++PassIndex)
+		{
+			EHeaderFolderTypes CurrentlyProcessing = (EHeaderFolderTypes)PassIndex;
 
 			// We'll make an ordered list of all UObject headers we care about.
 			// @todo uht: Ideally 'dependson' would not be allowed from public -> private, or NOT at all for new style headers
@@ -4670,20 +4671,6 @@ ECompilationResult::Type UnrealHeaderTool_Main(const FString& ModuleInfoFilename
 			if (!UObjectHeaders.Num())
 				continue;
 
-			UPackage* Package = Cast<UPackage>( StaticFindObjectFast( UPackage::StaticClass(), NULL, FName(*Module.LongPackageName), false, false ) );
-			if( Package == NULL )
-			{
-				Package = CreatePackage( NULL, *Module.LongPackageName );
-			}
-			GPackageToManifestModuleMap.Add(Package, &Module);
-
-			// Set some package flags for indicating that this package contains script
-			// NOTE: We do this even if we didn't have to create the package, because CoreUObject is compiled into UnrealHeaderTool and we still
-			//       want to make sure our flags get set
-			Package->PackageFlags |= PKG_ContainsScript;
-			Package->PackageFlags &= ~(PKG_ClientOptional|PKG_ServerSideOnly);
-			Package->PackageFlags |= PKG_Compiling;
-
 			for (const FString& Filename : UObjectHeaders)
 			{
 				// Best faith effort at a useful line number for errors occurring during this initial pre-parsing
@@ -4754,7 +4741,7 @@ ECompilationResult::Type UnrealHeaderTool_Main(const FString& ModuleInfoFilename
 				}
 			#endif
 			}
-			if(Result == ECompilationResult::Succeeded && NumFailures != 0)
+			if (Result == ECompilationResult::Succeeded && NumFailures != 0)
 			{
 				Result = ECompilationResult::OtherCompilationError;
 			}
