Remove HeaderGroup support from UHT.

diff --git a/Engine/Source/Programs/UnrealHeaderTool/Private/ClassMaps.cpp b/Engine/Source/Programs/UnrealHeaderTool/Private/ClassMaps.cpp
--- a/Engine/Source/Programs/UnrealHeaderTool/Private/ClassMaps.cpp
+++ b/Engine/Source/Programs/UnrealHeaderTool/Private/ClassMaps.cpp
@@ -7,7 +7,6 @@
 TMap<UClass*, FString>                   GClassStrippedHeaderTextMap;
 TMap<UClass*, FString>                   GClassSourceFileMap;
 TMap<UClass*, TUniqueObj<TArray<FName>>> GClassDependentOnMap;
-TMap<UClass*, FString>                   GClassHeaderFilenameMap;
 TMap<UClass*, FString>                   GClassHeaderNameWithNoPathMap;
 TMap<UClass*, FString>                   GClassModuleRelativePathMap;
 TMap<UClass*, FString>                   GClassIncludePathMap;
diff --git a/Engine/Source/Programs/UnrealHeaderTool/Private/ClassMaps.h b/Engine/Source/Programs/UnrealHeaderTool/Private/ClassMaps.h
--- a/Engine/Source/Programs/UnrealHeaderTool/Private/ClassMaps.h
+++ b/Engine/Source/Programs/UnrealHeaderTool/Private/ClassMaps.h
@@ -34,7 +34,6 @@ struct FManifestModule;
 extern TMap<UClass*, FString>                   GClassStrippedHeaderTextMap;
 extern TMap<UClass*, FString>                   GClassSourceFileMap;
 extern TMap<UClass*, TUniqueObj<TArray<FName>>> GClassDependentOnMap;
-extern TMap<UClass*, FString>                   GClassHeaderFilenameMap;
 extern TMap<UClass*, FString>                   GClassHeaderNameWithNoPathMap;
 extern TMap<UClass*, FString>                   GClassModuleRelativePathMap;
 extern TMap<UClass*, FString>                   GClassIncludePathMap;
diff --git a/Engine/Source/Programs/UnrealHeaderTool/Private/CodeGenerator.cpp b/Engine/Source/Programs/UnrealHeaderTool/Private/CodeGenerator.cpp
--- a/Engine/Source/Programs/UnrealHeaderTool/Private/CodeGenerator.cpp
+++ b/Engine/Source/Programs/UnrealHeaderTool/Private/CodeGenerator.cpp
@@ -2048,7 +2048,7 @@ void FNativeClassHeaderGenerator::ExportClassHeaderWrapper( UClass* Class, bool
 void FNativeClassHeaderGenerator::ExportClassHeaderRecursive( UClass* Class, TArray<UClass*>& DependencyChain, TSet<const UClass*>& VisitedSet, bool bCheckDependenciesOnly )
 {
 	bool bIsExportClass   = GClassStrippedHeaderTextMap.Contains(Class) && Class->HasAnyClassFlags(CLASS_Native|CLASS_Temporary) && !Class->HasAnyClassFlags(CLASS_NoExport);
-	bool bIsCorrectHeader = GClassHeaderFilenameMap.FindRef(Class) == ClassHeaderFilename && Class->GetOuter() == Package;
+	bool bIsCorrectHeader = Class->GetOuter() == Package;
 
 	// Check for circular header dependencies between export classes.
 	if ( bIsExportClass )
@@ -2064,16 +2064,6 @@ void FNativeClassHeaderGenerator::ExportClassHeaderRecursive( UClass* Class, TAr
 			// From now on, we're not going to export anything. Instead, we're going to check that no deeper dependency tries to export to this header file.
 			bCheckDependenciesOnly = true;
 		}
-		else if ( bCheckDependenciesOnly && !ClassHeaderFilename.IsEmpty() )
-		{
-			FString DependencyChainString;
-			for (UClass* DependencyClass : DependencyChain)
-			{
-				DependencyChainString += DependencyClass->GetName() + TEXT(" -> ");
-			}
-			DependencyChainString += Class->GetName();
-			UE_LOG(LogCompile, Warning, TEXT("Circular header dependency detected (%s), while exporting %s!"), *DependencyChainString, *(FPackageName::GetShortName(Package) + ClassHeaderFilename + TEXT("Classes.h")) );
-		}
 	}
 
 	// Check if the Class has already been exported, after we've checked for circular header dependencies.
@@ -3816,70 +3806,12 @@ bool FNativeClassHeaderGenerator::FindInterDependencyRecursive( TMap<const FStri
 	return false;
 }
 
-/**
- * Finds a dependency chain between two class header files.
- * Wrapper around FindDependencyChainRecursive().
- *
- * @param	Class				A class to scan for a dependency chain between the two headers.
- * @param	Header1				First class header filename.
- * @param	Header2				Second class header filename.
- * @param	DependencyChain		[out] Receives dependency chain, if found.
- * @return	true if a dependency chain was found and filled in.
- */
-bool FNativeClassHeaderGenerator::FindDependencyChain( FClass* Class, const FString& Header1, const FString& Header2, TArray<FClass*>& DependencyChain )
+bool IsExportOrTemporaryClass(FClass* Class)
 {
-	DependencyChain.Empty();
-	return FindDependencyChainRecursive( Class, Header1, Header2, false, DependencyChain );
-}
+	bool bIsExportClass = Class->HasAnyClassFlags(CLASS_Native) && !Class->HasAnyClassFlags(CLASS_NoExport | CLASS_Intrinsic | CLASS_Temporary);
+	bool bIsTemporaryClass = Class->HasAnyClassFlags(CLASS_Temporary);
 
-/**
- * Finds a dependency chain between two class header files.
- *
- * @param	Class				A class to scan for a dependency chain between the two headers.
- * @param	Header1				First class header filename.
- * @param	Header2				Second class header filename.
- * @param	bChainStarted		Whether Header1 has been found and we've started to fill in DependencyChain. Must be false to begin with.
- * @param	DependencyChain		[out] Receives dependency chain, if found. Must be empty before the call.
- * @return	true if a dependency chain was found and filled in.
- */
-bool FNativeClassHeaderGenerator::FindDependencyChainRecursive( FClass* Class, const FString& Header1, const FString& Header2, bool bChainStarted, TArray<FClass*>& DependencyChain )
-{
-	bool bIsExportClass = Class->HasAnyClassFlags(CLASS_Native) && !Class->HasAnyClassFlags(CLASS_NoExport|CLASS_Intrinsic|CLASS_Temporary);
-	if ( bIsExportClass )
-	{
-		auto HeaderFilename = GClassHeaderFilenameMap.FindRef(Class);
-
-		bChainStarted = bChainStarted || HeaderFilename == Header1;
-
-		if ( bChainStarted )
-		{
-			DependencyChain.Add( Class );
-			if ( HeaderFilename == Header2 )
-			{
-				return true;
-			}
-		}
-	}
-
-	FClass* SuperClass = Class->GetSuperClass();
-	if ( SuperClass && SuperClass->GetOuter() == Class->GetOuter() && FindDependencyChainRecursive( SuperClass, Header1, Header2, bChainStarted, DependencyChain ) )
-	{
-		return true;
-	}
-
-	for (FClass* InterfaceClass : Class->GetInterfaceTypes())
-	{
-		if ( InterfaceClass->GetOuter() == Class->GetOuter() && FindDependencyChainRecursive( InterfaceClass, Header1, Header2, bChainStarted, DependencyChain ) )
-		{
-			return true;
-		}
-	}
-
-	if ( bIsExportClass && bChainStarted )
-	{
-		DependencyChain.RemoveAtSwap( DependencyChain.Num() - 1 );
-	}
-	return false;
+	return bIsExportClass || bIsTemporaryClass;
 }
 
 // Constructor.
@@ -3909,143 +3841,7 @@ FNativeClassHeaderGenerator::FNativeClassHeaderGenerator( UPackage* InPackage, F
 		ClassNameMap.Add(Classes[ClassIndex]->GetFName(), Classes[ClassIndex]);
 	}
 
-	class UnsortedHeaderFilenamesType
-	{
-	public:
-		FString* AddUnique(const FString& Str)
-		{
-			for (auto& Ptr : Data)
-			{
-				if (*Ptr == Str)
-					return &*Ptr;
-			}
-
-			return Data[Data.Add(MakeUnique<FString>(Str))].Get();
-		}
-
-		TArray<TUniquePtr<FString>> Data;
-	};
-
-	UnsortedHeaderFilenamesType UnsortedHeaderFilenames;
-	TArray<const FString*>      HeaderFilenames; // These reference the strings stored in UnsortedHeaderFilenames
-
-	bool bCircularDependencyDetected = false;
- 	{
-	 	TMap<const FString*, HeaderDependents> HeaderDependencyMap;
-		TArray<FClass*> ClassesInPackage = AllClasses.GetClassesInPackage(Package);
-		for (FClass* Cls : ClassesInPackage)
-		{
-			auto ClsHeaderFilename = GClassHeaderFilenameMap.FindRef(Cls);
-
-			bool bIsExportClass = Cls->HasAnyClassFlags(CLASS_Native) && !Cls->HasAnyClassFlags(CLASS_NoExport|CLASS_Intrinsic|CLASS_Temporary);
-			if ( bIsExportClass )
-			{
-				auto* HeaderFilename = UnsortedHeaderFilenames.AddUnique(ClsHeaderFilename);
-				auto& Dependencies   = HeaderDependencyMap.FindOrAdd(HeaderFilename);
-
-				// Add the super class' header as a dependency if it's different.
-				UClass* SuperClass = Cls->GetSuperClass();
-				bool bIsDependentExportClass = SuperClass && SuperClass->HasAnyClassFlags(CLASS_Native) && !SuperClass->HasAnyClassFlags(CLASS_NoExport|CLASS_Intrinsic|CLASS_Temporary);
-				if ( bIsDependentExportClass && SuperClass->GetOuter() == Cls->GetOuter() )
-				{
-					auto SuperClassHeaderFilename = GClassHeaderFilenameMap.FindRef(SuperClass);
-					if ( SuperClassHeaderFilename != ClsHeaderFilename )
-					{
-						auto* DependentHeaderFilename = UnsortedHeaderFilenames.AddUnique(SuperClassHeaderFilename);
-						Dependencies.AddUnique(DependentHeaderFilename);
-					}
-				}
-
-				// Add base interface headers as dependencies, if they're different.
-				for (TArray<FImplementedInterface>::TIterator It(Cls->Interfaces); It; ++It)
-				{
-					UClass* InterfaceClass = It->Class;
-					bool bIsDependentExportClass = InterfaceClass->HasAnyClassFlags(CLASS_Native) && !InterfaceClass->HasAnyClassFlags(CLASS_NoExport|CLASS_Intrinsic|CLASS_Temporary);
-					if ( bIsDependentExportClass && InterfaceClass->GetOuter() == Cls->GetOuter() )
-					{
-						auto InterfaceClassHeaderFilename = GClassHeaderFilenameMap.FindRef(InterfaceClass);
-						if ( InterfaceClassHeaderFilename != ClsHeaderFilename )
-						{
-							auto* DependentHeaderFilename = UnsortedHeaderFilenames.AddUnique(InterfaceClassHeaderFilename);
-							Dependencies.AddUnique(DependentHeaderFilename);
-						}
-					}
-				}
-			}
-			else if (Cls->HasAnyClassFlags(CLASS_Temporary))
-			{
-				// Struct only headers can generate header groups too so make sure their added too.
-				auto* HeaderFilename = UnsortedHeaderFilenames.AddUnique(ClsHeaderFilename);
-				HeaderDependencyMap.FindOrAdd(HeaderFilename);
-			}
-		}
-		bCircularDependencyDetected = !SortHeaderDependencyMap(HeaderDependencyMap, HeaderFilenames);
-		if (bCircularDependencyDetected)
-		{
-			// Find one circular path (though there may be multiple).
-			for (auto HeaderIt = HeaderDependencyMap.CreateConstIterator(); HeaderIt; ++HeaderIt)
-			{
-				const FString* Header = HeaderIt->Key;
-
-				const FString* ClassHeaderFilename1;
-				const FString* ClassHeaderFilename2;
-				if (!FindInterDependency(HeaderDependencyMap, Header, ClassHeaderFilename1, ClassHeaderFilename2))
-					continue;
-
-				TArray<FClass*> DependencyChain1;
-				for (FClass* Class : ClassesInPackage)
-				{
-					FindDependencyChain( Class, *ClassHeaderFilename1, *ClassHeaderFilename2, DependencyChain1 );
-
-					if (DependencyChain1.Num())
-						break;
-				}
-
-				TArray<FClass*> DependencyChain2;
-				for (FClass* Class : ClassesInPackage)
-				{
-					FindDependencyChain( Class, *ClassHeaderFilename2, *ClassHeaderFilename1, DependencyChain2 );
-
-					if (DependencyChain2.Num())
-						break;
-				}
-
-				if (!DependencyChain1.Num() || !DependencyChain2.Num())
-					continue;
-
-				FString DependencyChainString1;
-				for (FClass* Dependency : DependencyChain1)
-				{
-					if (!DependencyChainString1.IsEmpty())
-					{
-						DependencyChainString1 += TEXT(" -> ");
-					}
-
-					DependencyChainString1 += Dependency->GetName();
-				}
-
-				FString DependencyChainString2;
-				for (FClass* Dependency : DependencyChain2)
-				{
-					if (!DependencyChainString1.IsEmpty())
-					{
-						DependencyChainString2 += TEXT(" -> ");
-					}
-
-					DependencyChainString2 += Dependency->GetName();
-				}
-
-				UE_LOG(LogCompile, Error, TEXT("Header interdependency: %s <-> %s (%s and %s)."),
-					*(PackageName + *ClassHeaderFilename1 + TEXT("Classes.h")),
-					*(PackageName + *ClassHeaderFilename2 + TEXT("Classes.h")),
-					*DependencyChainString1,
-					*DependencyChainString2 );
-				break;
-			}
-			UE_LOG(LogCompile, Error, TEXT("Interdependent headers detected - aborting!") );
-		}
-	}
-
+	bool bPackageHasAnyExportOrTemporaryClasses = AllClasses.GetClassesInPackage(Package).ContainsByPredicate(IsExportOrTemporaryClass);
 	bool bHasNamesForExport = false;
 	TempHeaderPaths.Empty();
 	PackageHeaderPaths.Empty();
@@ -4060,26 +3856,21 @@ FNativeClassHeaderGenerator::FNativeClassHeaderGenerator( UPackage* InPackage, F
 	GeneratedFunctionBodyTextSplit                .Empty();
 	ListOfPublicClassesUObjectHeaderModuleIncludes.Empty();
 
-	for (auto HeaderFilename : HeaderFilenames)
+	if (bPackageHasAnyExportOrTemporaryClasses)
 	{
-		if (bCircularDependencyDetected)
-			break;
-
-		ClassHeaderFilename = *HeaderFilename;
-
 		FString PkgDir;
 		FString GeneratedIncludeDirectory;
 		if (MakeCommandlet_FindPackageLocation(*PackageName, PkgDir, GeneratedIncludeDirectory) == false)
 		{
 			UE_LOG(LogCompile, Error, TEXT("Failed to find path for package %s"), *PackageName);
 		}
-		FString ClassesHeaderName = PackageName + ClassHeaderFilename + TEXT("Classes.h");
+		FString ClassesHeaderName = PackageName + TEXT("Classes.h");
 		ClassesHeaderPath = GeneratedIncludeDirectory / ClassesHeaderName;
 
 		int32 ClassCount = 0;
 		for (UClass* Class : Classes)
 		{
-			if( Class->GetOuter()==Package && Class->HasAnyClassFlags(CLASS_Native|CLASS_Temporary) && GClassHeaderFilenameMap.FindRef(Class) == ClassHeaderFilename )
+			if( Class->GetOuter()==Package && Class->HasAnyClassFlags(CLASS_Native|CLASS_Temporary))
 			{
 				if (GClassStrippedHeaderTextMap.Contains(Class) && !Class->HasAnyClassFlags(CLASS_NoExport))
 				{
@@ -4098,79 +3889,69 @@ FNativeClassHeaderGenerator::FNativeClassHeaderGenerator( UPackage* InPackage, F
 			}
 		}
 
-		if( ClassCount == 0 )
+		if (ClassCount != 0)
 		{
-			continue;
-		}
-
-		ClassesHeaders.Add(ClassesHeaderName);
-
-		CurrentClass = NULL;
-		ListOfPublicClassesUObjectHeaderGroupIncludes.Empty();
-		ListOfAllUObjectHeaderIncludes.Empty();
-		OriginalHeader.Empty();
-		PreHeaderText.Empty();
+			ClassesHeaders.Add(ClassesHeaderName);
 
-		// Load the original header file into memory
-		FFileHelper::LoadFileToString(OriginalHeader,*ClassesHeaderPath);
-		
-		UE_LOG(LogCompile, Log,  TEXT("Autogenerating C++ header: %s"), *ClassHeaderFilename );
+			CurrentClass = NULL;
+			ListOfPublicClassesUObjectHeaderGroupIncludes.Empty();
+			ListOfAllUObjectHeaderIncludes.Empty();
+			OriginalHeader.Empty();
+			PreHeaderText.Empty();
 
-		// Write the classes and enums header prefixes.
+			// Load the original header file into memory
+			FFileHelper::LoadFileToString(OriginalHeader, *ClassesHeaderPath);
 
-		PreHeaderText.Logf(
-			TEXT("// Copyright 1998-2014 Epic Games, Inc. All Rights Reserved.\r\n")
-			TEXT("/*===========================================================================\r\n")
-			TEXT("    C++ class boilerplate exported from UnrealHeaderTool.\r\n")
-			TEXT("    This is automatically generated by the tools.\r\n")
-			TEXT("    DO NOT modify this manually! Edit the corresponding .h files instead!\r\n")
-			TEXT("===========================================================================*/\r\n")
-			TEXT("#pragma once\r\n")
-			TEXT("\r\n")
-			);
+			// Write the classes and enums header prefixes.
 
-		// if a global auto-include file exists, generate a line to have that file included
-		FString GlobalAutoIncludeFilename = PackageName + ClassHeaderFilename + TEXT("GlobalIncludes.h");
-		const FString StandardHeaderFileLocation = PkgDir / TEXT("Public");
-		if ( IFileManager::Get().FileSize(*(StandardHeaderFileLocation / GlobalAutoIncludeFilename)) > 0 )
-		{
-			PreHeaderText.Logf(TEXT("#include \"%s\"\r\n\r\n"), *GlobalAutoIncludeFilename);
-		}
+			PreHeaderText.Logf(
+				TEXT("// Copyright 1998-2014 Epic Games, Inc. All Rights Reserved.\r\n")
+				TEXT("/*===========================================================================\r\n")
+				TEXT("    C++ class boilerplate exported from UnrealHeaderTool.\r\n")
+				TEXT("    This is automatically generated by the tools.\r\n")
+				TEXT("    DO NOT modify this manually! Edit the corresponding .h files instead!\r\n")
+				TEXT("===========================================================================*/\r\n")
+				TEXT("#pragma once\r\n")
+				TEXT("\r\n")
+				);
 
-		FString HeaderAPI = API;
-		if( ClassHeaderFilename.Len() > 0 )
-		{
-			HeaderAPI = HeaderAPI + TEXT("_") + ClassHeaderFilename.ToUpper();
-		}
+			// if a global auto-include file exists, generate a line to have that file included
+			FString GlobalAutoIncludeFilename = PackageName + TEXT("GlobalIncludes.h");
+			const FString StandardHeaderFileLocation = PkgDir / TEXT("Public");
+			if (IFileManager::Get().FileSize(*(StandardHeaderFileLocation / GlobalAutoIncludeFilename)) > 0)
+			{
+				PreHeaderText.Logf(TEXT("#include \"%s\"\r\n\r\n"), *GlobalAutoIncludeFilename);
+			}
 
-		// Export an include line for each header
-		for (UClass* Class : Classes)
-		{
-			if ( GClassHeaderFilenameMap.FindRef(Class) == ClassHeaderFilename && Class->GetOuter() == Package )
+			// Export an include line for each header
+			for (UClass* Class : Classes)
 			{
-				ExportClassHeader(Class);
+				if (Class->GetOuter() == Package)
+				{
+					ExportClassHeader(Class);
+				}
 			}
-		}
-		
-		ListOfPublicClassesUObjectHeaderGroupIncludes.Logf( TEXT("\r\n") );
-		
-		// build the full header file out of its pieces
-		const FString FullClassesHeader = FString::Printf(
-			TEXT("%s\r\n%s"),
-			*PreHeaderText, 
-			*ListOfPublicClassesUObjectHeaderGroupIncludes
-		);
 
-		// Save the classes header if it has changed.
-		SaveHeaderIfChanged(*ClassesHeaderPath,*FullClassesHeader);
+			ListOfPublicClassesUObjectHeaderGroupIncludes.Logf(TEXT("\r\n"));
+
+			// build the full header file out of its pieces
+			const FString FullClassesHeader = FString::Printf(
+				TEXT("%s\r\n%s"),
+				*PreHeaderText,
+				*ListOfPublicClassesUObjectHeaderGroupIncludes
+				);
+
+			// Save the classes header if it has changed.
+			SaveHeaderIfChanged(*ClassesHeaderPath, *FullClassesHeader);
+		}
 	}
 
-	if (HeaderFilenames.Num() == 0)
+	if (!bPackageHasAnyExportOrTemporaryClasses)
 	{
 		// If no headers are generated, check if there's any no export classes that need to have the inl file generated.
 		for (UClass* Class : Classes)
 		{
-			if ( Class->GetOuter() == Package && !Class->HasAnyClassFlags(CLASS_Intrinsic|CLASS_Temporary) )
+			if (Class->GetOuter() == Package && !Class->HasAnyClassFlags(CLASS_Intrinsic | CLASS_Temporary))
 			{
 				ExportClassHeader(Class);
 			}
@@ -4180,11 +3961,8 @@ FNativeClassHeaderGenerator::FNativeClassHeaderGenerator( UPackage* InPackage, F
 	// now export the names for the functions in this package
 	// notice we always export this file (as opposed to only exporting if we have any marked names)
 	// because there would be no way to know when the file was created otherwise
-	if (!bCircularDependencyDetected)
-	{
-		// Export .generated.cpp
-		ExportGeneratedCPP();
-	}
+	// Export .generated.cpp
+	ExportGeneratedCPP();
 
 	ExportGeneratedProto();
 
diff --git a/Engine/Source/Programs/UnrealHeaderTool/Private/HeaderParser.cpp b/Engine/Source/Programs/UnrealHeaderTool/Private/HeaderParser.cpp
--- a/Engine/Source/Programs/UnrealHeaderTool/Private/HeaderParser.cpp
+++ b/Engine/Source/Programs/UnrealHeaderTool/Private/HeaderParser.cpp
@@ -50,21 +50,6 @@ namespace
 		return false;
 	}
 
-	void AddHeaderGroup(UClass* Class, const FString& HeaderGroup, const TCHAR* ScopeForError)
-	{
-		if (FString* HeaderFilename = GClassHeaderFilenameMap.Find(Class))
-		{
-			if (*HeaderFilename != HeaderGroup)
-			{
-				FError::Throwf(TEXT("Found HeaderGroup (%s) in a %s declaration but the current header already had different HeaderGroup declared (%s)."), *HeaderGroup, ScopeForError, **HeaderFilename);
-			}
-		}
-		else
-		{
-			GClassHeaderFilenameMap.Add(Class, HeaderGroup);
-		}
-	}
-
 	bool ProbablyAMacro(const TCHAR* Identifier)
 	{
 		// Test for known delegate and event macros.
@@ -1115,17 +1100,6 @@ UScriptStruct* FHeaderParser::CompileStructDeclaration(FClasses& AllClasses, FCl
 			//UE_LOG(LogCompile, Warning, TEXT("Struct named %s in %s is still marked noexport"), *EffectiveStructName, *(Class->GetName()));//@TODO: UCREMOVAL: Debug printing
 			StructFlags &= ~STRUCT_Native;
 		}
-		else if (Specifier == TEXT("HeaderGroup"))
-		{
-			if (!Class)
-			{
-				// Can this ever happen?
-				FError::Throwf(TEXT("Unable to set HeaderGroup because there's no class present in the header."));
-			}
-
-			// Parse a header group filename.
-			AddHeaderGroup(Class, RequireExactlyOneSpecifierValue(*SpecifierIt), TEXT("struct"));
-		}
 		else if (Specifier == TEXT("Atomic"))
 		{
 			StructFlags |= STRUCT_Atomic;
@@ -4196,12 +4170,7 @@ void FHeaderParser::CompileClassDeclaration(FClasses& AllClasses)
 	{
 		const FString& Specifier = PropSpecifier.Key;
 
-		if (Specifier == TEXT("HeaderGroup"))
-		{
-			// Parse a header group filename.
-			AddHeaderGroup(Class, RequireExactlyOneSpecifierValue(PropSpecifier), TEXT("class"));
-		}
-		else if (Specifier == TEXT("noexport"))
+		if (Specifier == TEXT("noexport"))
 		{
 			// Don't export to C++ header.
 			Class->ClassFlags |= CLASS_NoExport;
@@ -4685,12 +4654,7 @@ void FHeaderParser::CompileInterfaceDeclaration(FClasses& AllClasses)
 	{
 		const FString& Specifier = SpecifierIt->Key;
 
-		if (Specifier == TEXT("HeaderGroup"))
-		{
-			// Parse a header group filename.
-			AddHeaderGroup(Class, RequireExactlyOneSpecifierValue(*SpecifierIt), TEXT("interface"));
-		}
-		else if (Specifier == TEXT("DependsOn"))
+		if (Specifier == TEXT("DependsOn"))
 		{
 			// Make sure the syntax matches but don't do anything with it; that's handled in MakeCommandlet.cpp
 			RequireSpecifierValue(*SpecifierIt);
diff --git a/Engine/Source/Programs/UnrealHeaderTool/Private/NativeClassExporter.h b/Engine/Source/Programs/UnrealHeaderTool/Private/NativeClassExporter.h
--- a/Engine/Source/Programs/UnrealHeaderTool/Private/NativeClassExporter.h
+++ b/Engine/Source/Programs/UnrealHeaderTool/Private/NativeClassExporter.h
@@ -38,7 +38,6 @@ private:
 
 	UClass*				CurrentClass;
 	TArray<FClass*>		Classes;
-	FString				ClassHeaderFilename;
 	FString				API;
 	FString			ClassesHeaderPath;
 	/** the name of the cpp file where the implementation functions are generated, without the .cpp extension */
@@ -140,30 +139,6 @@ private:
 	 */
 	bool FindInterDependencyRecursive( TMap<const FString*, HeaderDependents>& HeaderDependencyMap, const FString* HeaderIndex, TSet<const FString*>& VisitedHeaders, const FString*& OutHeader1, const FString*& OutHeader2 );
 
-	/**
-	 * Finds a dependency chain between two class header files.
-	 * Wrapper around FindDependencyChainRecursive().
-	 *
-	 * @param	Class				A class to scan for a dependency chain between the two headers.
-	 * @param	Header1				First class header filename.
-	 * @param	Header2				Second class header filename.
-	 * @param	DependencyChain		[out] Receives dependency chain, if found.
-	 * @return	true if a dependency chain was found and filled in.
-	 */
-	bool FindDependencyChain( FClass* Class, const FString& Header1, const FString& Header2, TArray<FClass*>& DependencyChain );
-
-	/**
-	 * Finds a dependency chain between two class header files.
-	 *
-	 * @param	Class				A class to scan for a dependency chain between the two headers.
-	 * @param	Header1				First class header filename.
-	 * @param	Header2				Second class header filename.
-	 * @param	bChainStarted		Whether Header1 has been found and we've started to fill in DependencyChain. Must be false to begin with.
-	 * @param	DependencyChain		[out] Receives dependency chain, if found. Must be empty before the call.
-	 * @return	true if a dependency chain was found and filled in.
-	 */
-	bool FindDependencyChainRecursive( FClass* Class, const FString& Header1, const FString& Header2, bool bChainStarted, TArray<FClass*>& DependencyChain );
-
 	/**
 	 * Determines whether the glue version of the specified native function
 	 * should be exported
