UHT: Make script parsing errors not-errors for now.

diff --git a/Engine/Source/Programs/UnrealHeaderTool/Private/CodeGenerator.cpp b/Engine/Source/Programs/UnrealHeaderTool/Private/CodeGenerator.cpp
--- a/Engine/Source/Programs/UnrealHeaderTool/Private/CodeGenerator.cpp
+++ b/Engine/Source/Programs/UnrealHeaderTool/Private/CodeGenerator.cpp
@@ -4582,6 +4582,7 @@ ECompilationResult::Type UnrealHeaderTool_Main(const FString& ModuleInfoFilename
 			TArray<IScriptGeneratorPluginInterface*> ScriptPlugins;
 			GetScriptPlugins(ScriptPlugins);
 
+			bool bCanGenerateScripts = true;
 			for (auto ScriptGenerator : ScriptPlugins)
 			{
 				// Find the right output direcotry for this plugin base on its target (Engine-side) plugin name.
@@ -4600,13 +4601,12 @@ ECompilationResult::Type UnrealHeaderTool_Main(const FString& ModuleInfoFilename
 				}
 				else
 				{
-					UE_LOG(LogCompile, Error, TEXT("Unable to determine output directory for %s. Cannot export script glue."), *GeneratedCodeModule);
-					Result = ECompilationResult::OtherCompilationError;
-					++NumFailures;
+					UE_LOG(LogCompile, Log, TEXT("Unable to determine output directory for %s. Cannot export script glue."), *GeneratedCodeModule);
+					bCanGenerateScripts = false;
 				}
 			}
 
-			if (Result == ECompilationResult::Succeeded)
+			if (bCanGenerateScripts)
 			{
 				for (const auto& Module : GManifest.Modules)
 				{
