Script plugins: more CIS fixes

diff --git a/Engine/Plugins/Script/ScriptGeneratorPlugin/Source/ScriptGeneratorPlugin/Private/ScriptCodeGeneratorBase.h b/Engine/Plugins/Script/ScriptGeneratorPlugin/Source/ScriptGeneratorPlugin/Private/ScriptCodeGeneratorBase.h
--- a/Engine/Plugins/Script/ScriptGeneratorPlugin/Source/ScriptGeneratorPlugin/Private/ScriptCodeGeneratorBase.h
+++ b/Engine/Plugins/Script/ScriptGeneratorPlugin/Source/ScriptGeneratorPlugin/Private/ScriptCodeGeneratorBase.h
@@ -41,10 +41,12 @@ protected:
 	/** Returns true if the specified property can be exported */
 	virtual bool CanExportProperty(const FString& ClassNameCPP, UClass* Class, UProperty* Property);
 	/** Generates a lookup table for exported proeprties. */
-	virtual int32 FScriptCodeGeneratorBase::GenerateExportedPropertyTable(const FString& ClassNameCPP, UClass* Class, FString& OutFunction, TArray<UProperty*>& OutExportedProperties);
+	virtual int32 GenerateExportedPropertyTable(const FString& ClassNameCPP, UClass* Class, FString& OutFunction, TArray<UProperty*>& OutExportedProperties);
 
 public:
 
+	virtual ~FScriptCodeGeneratorBase() {}
+
 	// IScriptGeneratorPlugin interface
 	virtual void ExportClass(UClass* Class, const FString& SourceHeaderFilename, const FString& GeneratedHeaderFilename, bool bHasChanged) = 0;
 	virtual void FinishExport() = 0;
diff --git a/Engine/Source/Programs/UnrealHeaderTool/Private/CodeGenerator.cpp b/Engine/Source/Programs/UnrealHeaderTool/Private/CodeGenerator.cpp
--- a/Engine/Source/Programs/UnrealHeaderTool/Private/CodeGenerator.cpp
+++ b/Engine/Source/Programs/UnrealHeaderTool/Private/CodeGenerator.cpp
@@ -4605,22 +4605,26 @@ ECompilationResult::Type UnrealHeaderTool_Main(const FString& ModuleInfoFilename
 					bCanGenerateScripts = false;
 				}
 			}
+			if (!bCanGenerateScripts)
+			{
+				ScriptPlugins.Empty();
+			}
 
-			if (bCanGenerateScripts)
+			for (const auto& Module : GManifest.Modules)
 			{
-				for (const auto& Module : GManifest.Modules)
+				if (UPackage* Package = Cast<UPackage>(StaticFindObjectFast(UPackage::StaticClass(), NULL, FName(*Module.LongPackageName), false, false)))
 				{
-					if (UPackage* Package = Cast<UPackage>(StaticFindObjectFast(UPackage::StaticClass(), NULL, FName(*Module.LongPackageName), false, false)))
+					Result = FHeaderParser::ParseAllHeadersInside(GWarn, Package, Module.SaveExportedHeaders, GManifest.UseRelativePaths, ScriptPlugins);
+					if (Result != ECompilationResult::Succeeded)
 					{
-						Result = FHeaderParser::ParseAllHeadersInside(GWarn, Package, Module.SaveExportedHeaders, GManifest.UseRelativePaths, ScriptPlugins);
-						if (Result != ECompilationResult::Succeeded)
-						{
-							++NumFailures;
-							break;
-						}
+						++NumFailures;
+						break;
 					}
 				}
+			}
 
+			if (bCanGenerateScripts)
+			{
 				for (auto ScriptGenerator : ScriptPlugins)
 				{
 					ScriptGenerator->FinishExport();
