Script Plugins: better handling of non-game targets

diff --git a/Engine/Plugins/Script/ScriptGeneratorPlugin/Source/ScriptGeneratorPlugin/Private/ScriptGeneratorPlugin.cpp b/Engine/Plugins/Script/ScriptGeneratorPlugin/Source/ScriptGeneratorPlugin/Private/ScriptGeneratorPlugin.cpp
--- a/Engine/Plugins/Script/ScriptGeneratorPlugin/Source/ScriptGeneratorPlugin/Private/ScriptGeneratorPlugin.cpp
+++ b/Engine/Plugins/Script/ScriptGeneratorPlugin/Source/ScriptGeneratorPlugin/Private/ScriptGeneratorPlugin.cpp
@@ -18,6 +18,7 @@ class FScriptGeneratorPlugin : public IScriptGeneratorPlugin
 
 	/** IScriptGeneratorPlugin interface */
 	virtual FString GetGeneratedCodeModuleName() const OVERRIDE { return TEXT("ScriptPlugin"); }
+	virtual bool SupportsTarget(const FString& TargetName) const OVERRIDE { return true; }
 	virtual void Initialize(const FString& RootLocalPath, const FString& RootBuildPath, const FString& OutputDirectory) OVERRIDE;
 	virtual void ExportClass(UClass* Class, const FString& SourceHeaderFilename, const FString& GeneratedHeaderFilename, bool bHasChanged) OVERRIDE;
 	virtual void FinishExport() OVERRIDE;
diff --git a/Engine/Source/Programs/UnrealBuildTool/System/ExternalExecution.cs b/Engine/Source/Programs/UnrealBuildTool/System/ExternalExecution.cs
--- a/Engine/Source/Programs/UnrealBuildTool/System/ExternalExecution.cs
+++ b/Engine/Source/Programs/UnrealBuildTool/System/ExternalExecution.cs
@@ -75,8 +75,10 @@ namespace UnrealBuildTool
 		public UHTManifest(bool InUseRelativePaths, UEBuildTarget Target, string InRootLocalPath, string InRootBuildPath, IEnumerable<UHTModuleInfo> ModuleInfo)
 		{
 			UseRelativePaths = InUseRelativePaths;
+			IsGameTarget = TargetRules.IsGameType(Target.Rules.Type);
 			RootLocalPath    = InRootLocalPath;
 			RootBuildPath    = InRootBuildPath;
+			TargetName = Target.GetTargetName();			
 
 			Modules = ModuleInfo.Select(Info => new Module{
 				Name                     = Info.ModuleName,
@@ -94,8 +96,10 @@ namespace UnrealBuildTool
 		}
 
 		public bool         UseRelativePaths; // Generate relative paths or absolute paths
+		public bool         IsGameTarget;     // True if the current target is a game target
 		public string       RootLocalPath;    // The engine path on the local machine
 		public string       RootBuildPath;    // The engine path on the build machine, if different (e.g. Mac/iOS builds)
+		public string		TargetName;       // Name of the target currently being compiled		
 		public List<Module> Modules;
 	}
 
diff --git a/Engine/Source/Programs/UnrealHeaderTool/Private/CodeGenerator.cpp b/Engine/Source/Programs/UnrealHeaderTool/Private/CodeGenerator.cpp
--- a/Engine/Source/Programs/UnrealHeaderTool/Private/CodeGenerator.cpp
+++ b/Engine/Source/Programs/UnrealHeaderTool/Private/CodeGenerator.cpp
@@ -4403,6 +4403,40 @@ void GetScriptPlugins(TArray<IScriptGeneratorPluginInterface*>& ScriptPlugins)
 			ScriptPlugins.Add(GeneratorInterface);
 		}
 	}
+
+	// Check if we can use these plugins and initialize them
+	for (int32 PluginIndex = ScriptPlugins.Num() - 1; PluginIndex >= 0; --PluginIndex)
+	{
+		auto ScriptGenerator = ScriptPlugins[PluginIndex];
+		bool bSupportedPlugin = ScriptGenerator->SupportsTarget(GManifest.TargetName);
+		if (bSupportedPlugin)
+		{
+			// Find the right output direcotry for this plugin base on its target (Engine-side) plugin name.
+			FString GeneratedCodeModule = ScriptGenerator->GetGeneratedCodeModuleName();
+			FString OutputDirectory;
+			for (const auto& Module : GManifest.Modules)
+			{
+				if (Module.Name == GeneratedCodeModule)
+				{
+					OutputDirectory = Module.GeneratedIncludeDirectory;
+				}
+			}
+			if (!OutputDirectory.IsEmpty())
+			{
+				ScriptGenerator->Initialize(GManifest.RootLocalPath, GManifest.RootBuildPath, OutputDirectory);
+			}
+			else
+			{
+				// Can't use this plugin
+				UE_LOG(LogCompile, Log, TEXT("Unable to determine output directory for %s. Cannot export script glue."), *GeneratedCodeModule);
+				bSupportedPlugin = false;				
+			}
+		}
+		if (!bSupportedPlugin)
+		{
+			ScriptPlugins.RemoveAt(PluginIndex);
+		}
+	}
 }
 
 ECompilationResult::Type UnrealHeaderTool_Main(const FString& ModuleInfoFilename)
@@ -4580,34 +4614,10 @@ ECompilationResult::Type UnrealHeaderTool_Main(const FString& ModuleInfoFilename
 		if (Result == ECompilationResult::Succeeded)
 		{
 			TArray<IScriptGeneratorPluginInterface*> ScriptPlugins;
-			GetScriptPlugins(ScriptPlugins);
-
-			bool bCanGenerateScripts = true;
-			for (auto ScriptGenerator : ScriptPlugins)
-			{
-				// Find the right output direcotry for this plugin base on its target (Engine-side) plugin name.
-				FString GeneratedCodeModule = ScriptGenerator->GetGeneratedCodeModuleName();
-				FString OutputDirectory;
-				for (const auto& Module : GManifest.Modules)
-				{
-					if (Module.Name == GeneratedCodeModule)
-					{
-						OutputDirectory = Module.GeneratedIncludeDirectory;
-					}
-				}
-				if (!OutputDirectory.IsEmpty())
-				{
-					ScriptGenerator->Initialize(GManifest.RootLocalPath, GManifest.RootBuildPath, OutputDirectory);
-				}
-				else
-				{
-					UE_LOG(LogCompile, Log, TEXT("Unable to determine output directory for %s. Cannot export script glue."), *GeneratedCodeModule);
-					bCanGenerateScripts = false;
-				}
-			}
-			if (!bCanGenerateScripts)
-			{
-				ScriptPlugins.Empty();
+			// Can only export scripts for game targets
+			if (GManifest.IsGameTarget)
+			{				
+				GetScriptPlugins(ScriptPlugins);
 			}
 
 			for (const auto& Module : GManifest.Modules)
@@ -4623,12 +4633,9 @@ ECompilationResult::Type UnrealHeaderTool_Main(const FString& ModuleInfoFilename
 				}
 			}
 
-			if (bCanGenerateScripts)
+			for (auto ScriptGenerator : ScriptPlugins)
 			{
-				for (auto ScriptGenerator : ScriptPlugins)
-				{
-					ScriptGenerator->FinishExport();
-				}
+				ScriptGenerator->FinishExport();
 			}
 		}
 	}
diff --git a/Engine/Source/Programs/UnrealHeaderTool/Private/Manifest.cpp b/Engine/Source/Programs/UnrealHeaderTool/Private/Manifest.cpp
--- a/Engine/Source/Programs/UnrealHeaderTool/Private/Manifest.cpp
+++ b/Engine/Source/Programs/UnrealHeaderTool/Private/Manifest.cpp
@@ -62,10 +62,20 @@ FManifest FManifest::LoadFromFile(const FString& Filename)
 	TArray<TSharedPtr<FJsonValue>> ModulesArray;
 
 	GetJsonFieldValue(Result.UseRelativePaths, RootObject, TEXT("UseRelativePaths"), TEXT("{manifest root}"));
+	GetJsonFieldValue(Result.IsGameTarget, RootObject, TEXT("IsGameTarget"), TEXT("{manifest root}"));	
 	GetJsonFieldValue(Result.RootLocalPath,    RootObject, TEXT("RootLocalPath"),    TEXT("{manifest root}"));
 	GetJsonFieldValue(Result.RootBuildPath,    RootObject, TEXT("RootBuildPath"),    TEXT("{manifest root}"));
+	GetJsonFieldValue(Result.TargetName, RootObject, TEXT("TargetName"), TEXT("{manifest root}"));
 	GetJsonFieldValue(ModulesArray,            RootObject, TEXT("Modules"),          TEXT("{manifest root}"));
 
+	UE_LOG(LogCompile, Log, TEXT("Loaded manifest: %s"), *Filename);
+	UE_LOG(LogCompile, Log, TEXT("Manifest.UseRelativePaths=%s"), Result.UseRelativePaths ? TEXT("True") : TEXT("False"));
+	UE_LOG(LogCompile, Log, TEXT("Manifest.IsGameTarget=%s"), Result.IsGameTarget ? TEXT("True") : TEXT("False"));
+	UE_LOG(LogCompile, Log, TEXT("Manifest.RootLocalPath=%s"), *Result.RootLocalPath);
+	UE_LOG(LogCompile, Log, TEXT("Manifest.RootBuildPath=%s"), *Result.RootBuildPath);
+	UE_LOG(LogCompile, Log, TEXT("Manifest.TargetName=%s"), *Result.TargetName);
+	UE_LOG(LogCompile, Log, TEXT("Manifest.Modules=%d"), ModulesArray.Num());
+
 	Result.RootLocalPath = FPaths::ConvertRelativePathToFull(FilenamePath, Result.RootLocalPath);
 	Result.RootBuildPath = FPaths::ConvertRelativePathToFull(FilenamePath, Result.RootBuildPath);
 
diff --git a/Engine/Source/Programs/UnrealHeaderTool/Private/Manifest.h b/Engine/Source/Programs/UnrealHeaderTool/Private/Manifest.h
--- a/Engine/Source/Programs/UnrealHeaderTool/Private/Manifest.h
+++ b/Engine/Source/Programs/UnrealHeaderTool/Private/Manifest.h
@@ -43,9 +43,11 @@ struct FManifestModule
 struct FManifest
 {
 	bool    UseRelativePaths;
+	bool IsGameTarget;
 	FString RootLocalPath;
 	FString RootBuildPath;
-
+	FString TargetName;
+	
 	/** Ordered list of modules that define UObjects or UStructs, which we may need to generate
 	    code for.  The list is in module dependency order, such that most dependent modules appear first. */
 	TArray<FManifestModule> Modules;
diff --git a/Engine/Source/Programs/UnrealHeaderTool/Public/IScriptGeneratorPluginInterface.h b/Engine/Source/Programs/UnrealHeaderTool/Public/IScriptGeneratorPluginInterface.h
--- a/Engine/Source/Programs/UnrealHeaderTool/Public/IScriptGeneratorPluginInterface.h
+++ b/Engine/Source/Programs/UnrealHeaderTool/Public/IScriptGeneratorPluginInterface.h
@@ -10,9 +10,16 @@ class IScriptGeneratorPluginInterface : public IModuleInterface
 {
 public:
 
+	/** Name of module that is going to be compiling generated script glue */
 	virtual FString GetGeneratedCodeModuleName() const = 0;
+	/** Returns true if this plugin supports exporting scripts for the specified target. This should handle game as well as editor target names */
+	virtual bool SupportsTarget(const FString& TargetName) const = 0;
+	/** Initializes this plugin with build information */
 	virtual void Initialize(const FString& RootLocalPath, const FString& RootBuildPath, const FString& OutputDirectory) = 0;
+	/** Exports a single class. May be called multiple times for the same class (as UHT processes the entire hierarchy inside modules. */
 	virtual void ExportClass(class UClass* Class, const FString& SourceHeaderFilename, const FString& GeneratedHeaderFilename, bool bHasChanged) = 0;
+	/** Called once all classes have been exported */
 	virtual void FinishExport() = 0;
+
 };
 
