#ttp 333303: [GitHub] 67 : Fix UnrealHeaderTool to properly support #ifdef and #ifndef precompiler directives

#fix Integrated code.
#fix Added other code to SimplifiedClassParse to make #elif etc. work in the middle of #ifdef/#ifndef.
#fix Added missing UHTDebugging.manifest fields.
#change Added some test preprocessor directives to TestObject.h.
#change Fixed some alignment.

diff --git a/Engine/Source/Programs/UnrealHeaderTool/Private/HeaderParser.cpp b/Engine/Source/Programs/UnrealHeaderTool/Private/HeaderParser.cpp
--- a/Engine/Source/Programs/UnrealHeaderTool/Private/HeaderParser.cpp
+++ b/Engine/Source/Programs/UnrealHeaderTool/Private/HeaderParser.cpp
@@ -1358,7 +1358,15 @@ UScriptStruct* FHeaderParser::CompileStructDeclaration(FClasses& AllClasses, FCl
 			// Eat a semicolon if present (not required)
 			SafeMatchSymbol(TEXT(";"));
 		}
-		else if ( Token.Matches(TEXT("#")) && MatchIdentifier(TEXT("endif")) )
+		else if ( Token.Matches(TEXT("#")) && MatchIdentifier(TEXT("ifdef")) )
+		{
+			PushCompilerDirective(ECompilerDirective::Insignificant);
+		}
+		else if ( Token.Matches(TEXT("#")) && MatchIdentifier(TEXT("ifndef")) )
+		{
+			PushCompilerDirective(ECompilerDirective::Insignificant);
+		}
+		else if (Token.Matches(TEXT("#")) && MatchIdentifier(TEXT("endif")))
 		{
 			if (CompilerDirectiveStack.Num() < 1)
 			{
@@ -2184,7 +2192,11 @@ void FHeaderParser::CompileDirective(UClass* Class)
 		// Ignore the define directive (can be multiline).
 		bDefineDirective = true;
 	}
-	else if (Directive.Matches(TEXT("undef")) || Directive.Matches(TEXT("ifdef")) || Directive.Matches(TEXT("ifndef")) || Directive.Matches(TEXT("else")))
+	else if (Directive.Matches(TEXT("ifdef")) || Directive.Matches(TEXT("ifndef")))
+	{
+		PushCompilerDirective(ECompilerDirective::Insignificant);
+	}
+	else if (Directive.Matches(TEXT("undef")) || Directive.Matches(TEXT("else")))
 	{
 		// Ignore. UHT can only handle #if directive
 	}
@@ -6747,7 +6759,14 @@ void FHeaderParser::SimplifiedClassParse(const TCHAR* InBuffer, bool& bIsInterfa
 		bool bProcess = CommentDim <= 0;	// for skipping nested multi-line comments
 		int32 BraceCount = 0;
 
-		if( bProcess && FParse::Command(&Str,TEXT("#if")) )
+		if( !bProcess )
+		{
+			ClassHeaderTextStrippedOfCppText.Logf( TEXT("%s\r\n"), *StrLine );
+			continue;
+		}
+
+		bool bIf = FParse::Command(&Str,TEXT("#if"));
+		if( bIf || FParse::Command(&Str,TEXT("#ifdef")) || FParse::Command(&Str,TEXT("#ifndef")) )
 		{
 			FStringOutputDevice TextDumpDummy;
 			int32 PreprocessorNest = 1;
@@ -6758,26 +6777,26 @@ void FHeaderParser::SimplifiedClassParse(const TCHAR* InBuffer, bool& bIsInterfa
 			bool bCPP = false;
 			bool bUnknownDirective = false;
 
-			if( FParse::Command(&Str,TEXT("CPP")) )
+			if( bIf && FParse::Command(&Str,TEXT("CPP")) )
 			{
 				Target = &TextDumpDummy;
 				SpacerTarget = &ClassHeaderTextStrippedOfCppText;
 				bCPP = true;
 			}
-			else if( FParse::Command(&Str,TEXT("!CPP")) )
+			else if( bIf && FParse::Command(&Str,TEXT("!CPP")) )
 			{
 				Target = &ClassHeaderTextStrippedOfCppText;
 				bKeepPreprocessorDirectives = false;
 				bNotCPP = true;
 			}
-			else if (FParse::Command(&Str,TEXT("WITH_EDITORONLY_DATA")) || FParse::Command(&Str,TEXT("WITH_EDITOR")))
+			else if (bIf && FParse::Command(&Str,TEXT("WITH_EDITORONLY_DATA")) || FParse::Command(&Str,TEXT("WITH_EDITOR")))
 			{
 				Target = &ClassHeaderTextStrippedOfCppText;
 				bUnknownDirective = true;
 			}
 			else
 			{
-				// Unknown directives are always treated as CPP
+				// Unknown directives or #ifdef or #ifndef are always treated as CPP
 				bUnknownDirective = true;
 				Target = &TextDumpDummy;
 				SpacerTarget = &ClassHeaderTextStrippedOfCppText;
@@ -6814,7 +6833,7 @@ void FHeaderParser::SimplifiedClassParse(const TCHAR* InBuffer, bool& bIsInterfa
 					PreprocessorNest--;
 					bIsPrep = true;
 				}
-				else if( FParse::Command(&Str,TEXT("#if")) )
+				else if( FParse::Command(&Str,TEXT("#if")) || FParse::Command(&Str,TEXT("#ifdef")) || FParse::Command(&Str,TEXT("#ifndef")) )
 				{
 					PreprocessorNest++;
 					bIsPrep = true;
@@ -6859,7 +6878,7 @@ void FHeaderParser::SimplifiedClassParse(const TCHAR* InBuffer, bool& bIsInterfa
 				}
 			}
 		}
-		else if ( bProcess && FParse::Command(&Str,TEXT("#include")) )
+		else if ( FParse::Command(&Str,TEXT("#include")) )
 		{
 			ClassHeaderTextStrippedOfCppText.Logf( TEXT("%s\r\n"), *StrLine );
 		}
diff --git a/Engine/Source/Programs/UnrealHeaderTool/Resources/UHTDebugging.manifest b/Engine/Source/Programs/UnrealHeaderTool/Resources/UHTDebugging.manifest
--- a/Engine/Source/Programs/UnrealHeaderTool/Resources/UHTDebugging.manifest
+++ b/Engine/Source/Programs/UnrealHeaderTool/Resources/UHTDebugging.manifest
@@ -1,7 +1,9 @@
 {
 	"UseRelativePaths": false,
+	"IsGameTarget": false,
 	"RootLocalPath": "..\\..\\..\\..",
 	"RootBuildPath": "..\\..\\..\\..",
+	"TargetName": "UHTDebugging",
 	"Modules":[{
 			"Name": "CoreUObject",
 			"BaseDirectory": "..\\..\\..\\Runtime\\CoreUObject",
diff --git a/Engine/Source/Programs/UnrealHeaderTool/Resources/UHTDebugging/TestObject.h b/Engine/Source/Programs/UnrealHeaderTool/Resources/UHTDebugging/TestObject.h
--- a/Engine/Source/Programs/UnrealHeaderTool/Resources/UHTDebugging/TestObject.h
+++ b/Engine/Source/Programs/UnrealHeaderTool/Resources/UHTDebugging/TestObject.h
@@ -11,7 +11,50 @@ class UTestObject : public UObject
 	void TestForNullPtrDefaults(UObject* Obj1 = NULL, UObject* Obj2 = nullptr, UObject* Obj3 = 0);
 
 #if BLAH
+	UFUNCTION() int x; // This should not compile if UHT parses it, which it shouldn't
 #elif BLAH2
+	UFUNCTION() int x;
 #else
+	UFUNCTION() int x;
+#endif
+
+#ifdef X
+	UFUNCTION() int x;
+	#ifndef X
+		UFUNCTION() int x;
+	#elif BLAH4
+		UFUNCTION() int x;
+	#else
+		UFUNCTION() int x;
+	#endif
+	UFUNCTION() int x;
+#elif BLAH3
+	UFUNCTION() int x;
+	#ifndef X
+		UFUNCTION() int x;
+	#elif BLAH4
+		UFUNCTION() int x;
+	#else
+		UFUNCTION() int x;
+	#endif
+	UFUNCTION() int x;
+#else
+	UFUNCTION() int x;
+	#ifndef X
+		UFUNCTION() int x;
+	#elif BLAH4
+		UFUNCTION() int x;
+	#else
+		UFUNCTION() int x;
+	#endif
+	UFUNCTION() int x;
+#endif
+
+#ifndef X
+	UFUNCTION() int x;
+#elif BLAH4
+	UFUNCTION() int x;
+#else
+	UFUNCTION() int x;
 #endif
 };
