Warnings issued when including a .generated.inl file.

#add Wrote out a .generated.inl file for every module which contains a #pragma message asking the user to delete their include.

diff --git a/Engine/Source/Programs/UnrealHeaderTool/Private/CodeGenerator.cpp b/Engine/Source/Programs/UnrealHeaderTool/Private/CodeGenerator.cpp
--- a/Engine/Source/Programs/UnrealHeaderTool/Private/CodeGenerator.cpp
+++ b/Engine/Source/Programs/UnrealHeaderTool/Private/CodeGenerator.cpp
@@ -4262,6 +4262,7 @@ void FNativeClassHeaderGenerator::ExportGeneratedCPP()
 	FStringOutputDevice GeneratedCPPClassesIncludes;
 	FStringOutputDevice GeneratedCPPEpilogue;
 	FStringOutputDevice GeneratedCPPText;
+	FStringOutputDevice GeneratedINLText;
 
 	TArray<FStringOutputDevice> GeneratedCPPFiles;
 
@@ -4277,6 +4278,18 @@ void FNativeClassHeaderGenerator::ExportGeneratedCPP()
 																								LINE_TERMINATOR
 		);
 
+	FString DeprecationMessage = TEXT(".generated.inl files have been deprecated - please remove the associated #include from your code");
+	GeneratedINLText.Logf(
+		TEXT("#ifdef _MSC_VER")										LINE_TERMINATOR
+		TEXT("    #pragma message(__FILE__\"(9): warning : %s\")")	LINE_TERMINATOR
+		TEXT("#else")												LINE_TERMINATOR
+		TEXT("    #pragma message(\"%s\")")							LINE_TERMINATOR
+		TEXT("#endif")												LINE_TERMINATOR
+																	LINE_TERMINATOR,
+		*DeprecationMessage,
+		*DeprecationMessage
+	);
+
 	FString ModulePCHInclude;
 
 	const auto* ModuleInfo = GPackageToManifestModuleMap.FindChecked(Package);
@@ -4357,13 +4370,12 @@ void FNativeClassHeaderGenerator::ExportGeneratedCPP()
 			FileText.Logf(TEXT("#endif\r\n"));
 		}
 
-		FString HeaderPath = ModuleInfo->GeneratedCPPFilenameBase + (GeneratedFunctionBodyTextSplit.Num() > 1 ? *FString::Printf(TEXT(".%d.cpp"), FileIdx + 1) : TEXT(".cpp"));
-		
-		SaveHeaderIfChanged(*HeaderPath, *(GeneratedCPPPreamble + ModulePCHInclude + GeneratedCPPClassesIncludes + FileText + GeneratedCPPEpilogue));
+		FString CppPath = ModuleInfo->GeneratedCPPFilenameBase + (GeneratedFunctionBodyTextSplit.Num() > 1 ? *FString::Printf(TEXT(".%d.cpp"), FileIdx + 1) : TEXT(".cpp"));
+		SaveHeaderIfChanged(*CppPath, *(GeneratedCPPPreamble + ModulePCHInclude + GeneratedCPPClassesIncludes + FileText + GeneratedCPPEpilogue));
 
 		if (GeneratedFunctionBodyTextSplit.Num() > 1)
 		{
-			NumberedHeaderNames.Add(FPaths::GetCleanFilename(HeaderPath));
+			NumberedHeaderNames.Add(FPaths::GetCleanFilename(CppPath));
 		}
 	}
 
@@ -4376,10 +4388,13 @@ void FNativeClassHeaderGenerator::ExportGeneratedCPP()
 			FileText.Logf(TEXT("#include \"%s\"") LINE_TERMINATOR, *NumberedHeaderNames[i]);
 		}
 
-		FString HeaderPath = ModuleInfo->GeneratedCPPFilenameBase + TEXT(".cpp");
-		SaveHeaderIfChanged(*HeaderPath, *(GeneratedCPPPreamble + ModulePCHInclude + FileText + GeneratedCPPEpilogue));
+		FString CppPath = ModuleInfo->GeneratedCPPFilenameBase + TEXT(".cpp");
+		SaveHeaderIfChanged(*CppPath, *(GeneratedCPPPreamble + ModulePCHInclude + FileText + GeneratedCPPEpilogue));
 	}
 
+	FString InlPath = ModuleInfo->GeneratedCPPFilenameBase + TEXT(".inl");
+	SaveHeaderIfChanged(*InlPath, *(GeneratedCPPPreamble + GeneratedINLText));
+
 	if( FParse::Param( FCommandLine::Get(), TEXT("GenerateConstructors") ) && AllConstructors.Len())
 	{
 		const FString PrivateLoc = PkgDir / PkgName / TEXT("Private");
