UHT no longer assumes an include path will always contain a folder

TTP# 333824 - Add Code To Project not including game-specific headers

It would previously assume that all include paths contained a Public/Private/Classes folder, and would blindly try and strip them out, meaning that a project without these folders (like the ones created by the game templates) would produce incorrect or empty include paths.

It instead now explicitly checks for a Public/Private/Classes folder, rather than just doing it blindly.

diff --git a/Engine/Source/Programs/UnrealHeaderTool/Private/CodeGenerator.cpp b/Engine/Source/Programs/UnrealHeaderTool/Private/CodeGenerator.cpp
--- a/Engine/Source/Programs/UnrealHeaderTool/Private/CodeGenerator.cpp
+++ b/Engine/Source/Programs/UnrealHeaderTool/Private/CodeGenerator.cpp
@@ -4564,19 +4564,36 @@ ECompilationResult::Type UnrealHeaderTool_Main(const FString& ModuleInfoFilename
 						const TCHAR *ModuleRelativePath = *FullModulePath + Module.BaseDirectory.Len();
 						GClassModuleRelativePathMap.Add(ResultClass, ModuleRelativePath);
 
-						// Add the include path
+						// Calculate the include path
 						const TCHAR *IncludePath = ModuleRelativePath;
+
+						// Walk over the first potential slash
 						if(*IncludePath == '/')
 						{
 							IncludePath++;
 						}
-						while(*IncludePath != 0 && *IncludePath != '/')
+
+						// Does this module path start with a known include path location? If so, we can cut that part out of the include path
+						static const TCHAR PublicFolderName[]  = TEXT("Public/");
+						static const TCHAR PrivateFolderName[] = TEXT("Private/");
+						static const TCHAR ClassesFolderName[] = TEXT("Classes/");
+						if(FCString::Strnicmp(IncludePath, PublicFolderName, ARRAYSIZE(PublicFolderName) - 1) == 0)
 						{
-							IncludePath++;
+							IncludePath += (ARRAYSIZE(PublicFolderName) - 1);
+						}
+						else if(FCString::Strnicmp(IncludePath, PrivateFolderName, ARRAYSIZE(PrivateFolderName) - 1) == 0)
+						{
+							IncludePath += (ARRAYSIZE(PrivateFolderName) - 1);
+						}
+						else if(FCString::Strnicmp(IncludePath, ClassesFolderName, ARRAYSIZE(ClassesFolderName) - 1) == 0)
+						{
+							IncludePath += (ARRAYSIZE(ClassesFolderName) - 1);
 						}
-						if(*IncludePath == '/' && *(IncludePath + 1) != 0)
+
+						// Add the include path
+						if(*IncludePath != 0)
 						{
-							GClassIncludePathMap.Add(ResultClass, IncludePath + 1);
+							GClassIncludePathMap.Add(ResultClass, IncludePath);
 						}
 					}
 				}
