CIS fix - Replaced ARRAYSIZE with our cross platform ARRAY_COUNT

diff --git a/Engine/Source/Programs/UnrealHeaderTool/Private/CodeGenerator.cpp b/Engine/Source/Programs/UnrealHeaderTool/Private/CodeGenerator.cpp
--- a/Engine/Source/Programs/UnrealHeaderTool/Private/CodeGenerator.cpp
+++ b/Engine/Source/Programs/UnrealHeaderTool/Private/CodeGenerator.cpp
@@ -4577,17 +4577,17 @@ ECompilationResult::Type UnrealHeaderTool_Main(const FString& ModuleInfoFilename
 						static const TCHAR PublicFolderName[]  = TEXT("Public/");
 						static const TCHAR PrivateFolderName[] = TEXT("Private/");
 						static const TCHAR ClassesFolderName[] = TEXT("Classes/");
-						if(FCString::Strnicmp(IncludePath, PublicFolderName, ARRAYSIZE(PublicFolderName) - 1) == 0)
+						if(FCString::Strnicmp(IncludePath, PublicFolderName, ARRAY_COUNT(PublicFolderName) - 1) == 0)
 						{
-							IncludePath += (ARRAYSIZE(PublicFolderName) - 1);
+							IncludePath += (ARRAY_COUNT(PublicFolderName) - 1);
 						}
-						else if(FCString::Strnicmp(IncludePath, PrivateFolderName, ARRAYSIZE(PrivateFolderName) - 1) == 0)
+						else if(FCString::Strnicmp(IncludePath, PrivateFolderName, ARRAY_COUNT(PrivateFolderName) - 1) == 0)
 						{
-							IncludePath += (ARRAYSIZE(PrivateFolderName) - 1);
+							IncludePath += (ARRAY_COUNT(PrivateFolderName) - 1);
 						}
-						else if(FCString::Strnicmp(IncludePath, ClassesFolderName, ARRAYSIZE(ClassesFolderName) - 1) == 0)
+						else if(FCString::Strnicmp(IncludePath, ClassesFolderName, ARRAY_COUNT(ClassesFolderName) - 1) == 0)
 						{
-							IncludePath += (ARRAYSIZE(ClassesFolderName) - 1);
+							IncludePath += (ARRAY_COUNT(ClassesFolderName) - 1);
 						}
 
 						// Add the include path
