UHT: Fixing bad error message when defining UFUNCTIONs inside of USTRUCTs

diff --git a/Engine/Source/Programs/UnrealHeaderTool/Private/HeaderParser.cpp b/Engine/Source/Programs/UnrealHeaderTool/Private/HeaderParser.cpp
--- a/Engine/Source/Programs/UnrealHeaderTool/Private/HeaderParser.cpp
+++ b/Engine/Source/Programs/UnrealHeaderTool/Private/HeaderParser.cpp
@@ -1326,7 +1326,7 @@ UScriptStruct* FHeaderParser::CompileStructDeclaration(FClasses& AllClasses, FCl
 		}
 		else if (Token.Matches(TEXT("UFUNCTION"), ESearchCase::CaseSensitive))
 		{
-			FError::Throwf(TEXT("UFUNCTIONs cannot contain USTRUCTs."));
+			FError::Throwf(TEXT("USTRUCTs cannot contain UFUNCTIONs."));
 		}
 		else if (Token.Matches(TEXT("GENERATED_USTRUCT_BODY")))
 		{
