Build UHT monolithically on Linux (so far).

diff --git a/Engine/Source/Programs/UnrealHeaderTool/UnrealHeaderTool.Target.cs b/Engine/Source/Programs/UnrealHeaderTool/UnrealHeaderTool.Target.cs
--- a/Engine/Source/Programs/UnrealHeaderTool/UnrealHeaderTool.Target.cs
+++ b/Engine/Source/Programs/UnrealHeaderTool/UnrealHeaderTool.Target.cs
@@ -28,7 +28,7 @@ public class UnrealHeaderToolTarget : TargetRules
 
 	public override bool ShouldCompileMonolithic(UnrealTargetPlatform InPlatform, UnrealTargetConfiguration InConfiguration)
 	{
-		if (UnrealBuildTool.UnrealBuildTool.CommandLineContains("-monolithic") == true)
+		if (UnrealBuildTool.UnrealBuildTool.CommandLineContains("-monolithic") == true || InPlatform == UnrealTargetPlatform.Linux)
 		{
 			return true;
 		}
