[GitHub] 109 : Unreal Header Tool bug fix for function parameter default declarations
#ttp 333858

diff --git a/Engine/Source/Programs/UnrealHeaderTool/Private/HeaderParser.cpp b/Engine/Source/Programs/UnrealHeaderTool/Private/HeaderParser.cpp
--- a/Engine/Source/Programs/UnrealHeaderTool/Private/HeaderParser.cpp
+++ b/Engine/Source/Programs/UnrealHeaderTool/Private/HeaderParser.cpp
@@ -7211,6 +7211,11 @@ bool FHeaderParser::DefaultValueStringCppFormatToInnerFormat(const UProperty* Pr
 			{
 				return true;
 			}
+			if(FDefaultValueHelper::Is(CppForm, TEXT("FVector::UpVector")))
+			{
+				OutForm = FString::Printf(TEXT("(X=%3.3f,Y=%3.3f,Z=%3.3f)"),
+					FVector::UpVector.X, FVector::UpVector.Y, FVector::UpVector.Z);
+			}
 			FString Parameters;
 			if( FDefaultValueHelper::GetParameters(CppForm, TEXT("FVector"), Parameters) )
 			{
@@ -7221,7 +7226,8 @@ bool FHeaderParser::DefaultValueStringCppFormatToInnerFormat(const UProperty* Pr
 				FVector Vector;
 				if(FDefaultValueHelper::ParseVector(Parameters, Vector))
 				{
-					OutForm = FString::Printf(TEXT("%f,%f,%f"), Vector.X, Vector.Y, Vector.Z);
+					OutForm = FString::Printf(TEXT("(X=%3.3f,Y=%3.3f,Z=%3.3f)"),
+						Vector.X, Vector.Y, Vector.Z);
 				}
 			}
 		}
@@ -7241,7 +7247,8 @@ bool FHeaderParser::DefaultValueStringCppFormatToInnerFormat(const UProperty* Pr
 				FRotator Rotator;
 				if(FDefaultValueHelper::ParseRotator(Parameters, Rotator))
 				{
-					OutForm = FString::Printf(TEXT("%f,%f,%f"), Rotator.Pitch, Rotator.Yaw, Rotator.Roll);
+					OutForm = FString::Printf(TEXT("(P=%f,Y=%f,R=%f)"),
+						Rotator.Pitch, Rotator.Yaw, Rotator.Roll);
 				}
 			}
 		}
@@ -7251,6 +7258,11 @@ bool FHeaderParser::DefaultValueStringCppFormatToInnerFormat(const UProperty* Pr
 			{
 				return true;
 			}
+			if(FDefaultValueHelper::Is(CppForm, TEXT("FVector2D::UnitVector")))
+			{
+				OutForm = FString::Printf(TEXT("(X=%3.3f,Y=%3.3f)"),
+					FVector2D::UnitVector.X, FVector2D::UnitVector.Y);
+			}
 			FString Parameters;
 			if( FDefaultValueHelper::GetParameters(CppForm, TEXT("FVector2D"), Parameters) )
 			{
@@ -7261,7 +7273,8 @@ bool FHeaderParser::DefaultValueStringCppFormatToInnerFormat(const UProperty* Pr
 				FVector2D Vector2D;
 				if(FDefaultValueHelper::ParseVector2D(Parameters, Vector2D))
 				{
-					OutForm = FString::Printf(TEXT("%f,%f"), Vector2D.X, Vector2D.Y);
+					OutForm = FString::Printf(TEXT("(X=%3.3f,Y=%3.3f)"),
+						Vector2D.X, Vector2D.Y);
 				}
 			}
 		}
