#ttp 333737: UHT should not output anything differently for Rocket

#change Root module path added to include paths.
#change Plugins now write intermediates to the project folder.
#change All generated #includes are now relative.
#change Updated UHTDebugging.manifest to new format.
#remove Rocket-specific UHT execution removed.
#remove Some bHackHeaderGenerator usage removed.
#remove iOS header hack removed.

diff --git a/Engine/Source/Programs/UnrealBuildTool/Configuration/UEBuildModule.cs b/Engine/Source/Programs/UnrealBuildTool/Configuration/UEBuildModule.cs
--- a/Engine/Source/Programs/UnrealBuildTool/Configuration/UEBuildModule.cs
+++ b/Engine/Source/Programs/UnrealBuildTool/Configuration/UEBuildModule.cs
@@ -476,6 +476,9 @@ namespace UnrealBuildTool
 					var IncludePathModule = Target.GetModuleByName(IncludePathModuleName);
 					IncludePathModule.SetupPublicCompileEnvironment( SourceBinary, bInnerIncludePathsOnly, ref IncludePaths, ref SystemIncludePaths, ref Definitions, ref VisitedModules );
 				}
+
+				// Add the module's directory to the include path, so we can root #includes to it
+				IncludePaths.Add(ModuleDirectory);
 			}
 		}
 		
@@ -1413,7 +1416,7 @@ namespace UnrealBuildTool
 
 		private FileItem ProcessDependencies(FileItem CPPFile, CPPEnvironment ModuleCompileEnvironment)
 		{
-			List<DependencyInclude> DirectIncludeFilenames = CPPEnvironment.GetDirectIncludeDependencies(CPPFile, ModuleCompileEnvironment.Config.TargetPlatform, ModuleCompileEnvironment.bHackHeaderGenerator);
+			List<DependencyInclude> DirectIncludeFilenames = CPPEnvironment.GetDirectIncludeDependencies(CPPFile, ModuleCompileEnvironment.Config.TargetPlatform);
 			if (BuildConfiguration.bPrintDebugInfo)
 			{
 				Log.TraceVerbose("Found direct includes for {0}: {1}", Path.GetFileName(CPPFile.AbsolutePath), string.Join(", ", DirectIncludeFilenames.Select(F => F.IncludeName)));
@@ -1709,7 +1712,7 @@ namespace UnrealBuildTool
 
 								// If it's an engine module, output intermediates to the engine intermediates directory. 
 								string IntermediateDirectory = Binary.Config.IntermediateDirectory;
-								if (PluginInfo == null && IntermediateDirectory != Target.EngineIntermediateDirectory && Path.GetFullPath(DependencyModule.ModuleDirectory).StartsWith(Path.GetFullPath(BuildConfiguration.RelativeEnginePath)))
+								if (IntermediateDirectory != Target.EngineIntermediateDirectory && Path.GetFullPath(DependencyModule.ModuleDirectory).StartsWith(Path.GetFullPath(BuildConfiguration.RelativeEnginePath)))
 								{
 									IntermediateDirectory = Target.EngineIntermediateDirectory;
 								}
diff --git a/Engine/Source/Programs/UnrealBuildTool/Configuration/UEBuildTarget.cs b/Engine/Source/Programs/UnrealBuildTool/Configuration/UEBuildTarget.cs
--- a/Engine/Source/Programs/UnrealBuildTool/Configuration/UEBuildTarget.cs
+++ b/Engine/Source/Programs/UnrealBuildTool/Configuration/UEBuildTarget.cs
@@ -1075,13 +1075,13 @@ namespace UnrealBuildTool
 					Manifest.AddBinaryNames(Path.Combine(Binary.Config.IntermediateDirectory, Path.GetFileNameWithoutExtension(Binary.Config.OutputFilePath) + ".lib"), "");
 				}
 			}
-            {
-                string DebugInfoExtension = BuildPlatform.GetDebugInfoExtension(UEBuildBinaryType.StaticLibrary);
-                foreach (var RedistLib in SpecialRocketLibFilesThatAreBuildProducts)
-                {
-                    Manifest.AddBinaryNames(RedistLib, DebugInfoExtension);
-                }
-            }
+			{
+				string DebugInfoExtension = BuildPlatform.GetDebugInfoExtension(UEBuildBinaryType.StaticLibrary);
+				foreach (var RedistLib in SpecialRocketLibFilesThatAreBuildProducts)
+				{
+					Manifest.AddBinaryNames(RedistLib, DebugInfoExtension);
+				}
+			}
 
 
 			if (UEBuildConfiguration.bCleanProject)
@@ -1502,18 +1502,7 @@ namespace UnrealBuildTool
 				string HeaderFilename                    = LinkerFixupsName + "Name.h";
 				string LinkerFixupHeaderFilenameWithPath = Path.Combine(GlobalCompileEnvironment.Config.OutputDirectory, HeaderFilename);
 
-				UnrealTargetPlatform TargetPlatform = UEBuildTarget.CPPTargetPlatformToUnrealTargetPlatform( GlobalCompileEnvironment.Config.TargetPlatform );
-
-				// Quick fix. 
-				// For some reason, when compiled over RPCUtility, the path can't be found, and we don't need the long path on Mac anyhow
-				if ( ExternalExecution.GetRuntimePlatform() != UnrealTargetPlatform.Mac && TargetPlatform == UnrealTargetPlatform.IOS )
-				{
-					LinkerFixupsFileContents.Add("#include \"" + HeaderFilename + "\"");
-				}
-				else
-				{
-					LinkerFixupsFileContents.Add("#include \"" + LinkerFixupHeaderFilenameWithPath + "\"");
-				}
+				LinkerFixupsFileContents.Add("#include \"" + HeaderFilename + "\"");
 
 				// Add a function that is not referenced by anything that invokes all the empty functions in the different static libraries
 				LinkerFixupsFileContents.Add("void " + LinkerFixupsName + "()");
@@ -1726,7 +1715,7 @@ namespace UnrealBuildTool
 			UEBuildBinaryType BinaryType = ShouldCompileMonolithic()? UEBuildBinaryType.StaticLibrary : UEBuildBinaryType.DynamicLinkLibrary;
 
 			// Make the plugin intermediate path
-			string IntermediateDirectory = Path.Combine(Plugin.Directory, BuildConfiguration.PlatformIntermediateFolder, Plugins.GetPluginSubfolderName(BinaryType, AppName), Configuration.ToString());
+			string IntermediateDirectory = ProjectIntermediateDirectory;
 
 			// Get the output path. Don't prefix the app name for Rocket
 			string OutputFilePath;
@@ -2103,7 +2092,7 @@ namespace UnrealBuildTool
             {
                 GlobalCompileEnvironment.Config.Definitions.Add("UE_ROCKET=0");
             }
-			
+
 			// Rocket intermediates go to the project's intermediate folder.  Rocket never writes to the engine intermediate folder. (Those files are immutable)
 			// Also, when compiling in monolithic, all intermediates go to the project's folder.  This is because a project can change definitions that affects all engine translation
 			// units too, so they can't be shared between different targets.  They are effectively project-specific engine intermediates.
@@ -2113,11 +2102,11 @@ namespace UnrealBuildTool
 				var IntermediateConfiguration = Configuration;
 				if( UnrealBuildTool.RunningRocket() )
 				{
-					// Only Development and Shipping are supported for engine modules
+				// Only Development and Shipping are supported for engine modules
 					if( IntermediateConfiguration != UnrealTargetConfiguration.Development && IntermediateConfiguration != UnrealTargetConfiguration.Shipping )
-					{
-						IntermediateConfiguration = UnrealTargetConfiguration.Development;
-					}
+				{
+					IntermediateConfiguration = UnrealTargetConfiguration.Development;
+				}
 				}
 
 				GlobalCompileEnvironment.Config.OutputDirectory = Path.Combine(BuildConfiguration.PlatformIntermediatePath, OutputAppName, IntermediateConfiguration.ToString());
diff --git a/Engine/Source/Programs/UnrealBuildTool/System/ActionGraph.cs b/Engine/Source/Programs/UnrealBuildTool/System/ActionGraph.cs
--- a/Engine/Source/Programs/UnrealBuildTool/System/ActionGraph.cs
+++ b/Engine/Source/Programs/UnrealBuildTool/System/ActionGraph.cs
@@ -1030,7 +1030,7 @@ namespace UnrealBuildTool
 
 				if( ExpandCPPHeaderDependencies && bIsCPPFile )
 				{
-					List<DependencyInclude> DirectlyIncludedFilenames = CPPEnvironment.GetDirectIncludeDependencies( FileItem, CPPTargetPlatform.Win64, bHackHeaderGenerator:false );
+					List<DependencyInclude> DirectlyIncludedFilenames = CPPEnvironment.GetDirectIncludeDependencies( FileItem, CPPTargetPlatform.Win64 );
 
 					// Resolve the included file name to an actual file.
 					var DirectlyIncludedFiles =
diff --git a/Engine/Source/Programs/UnrealBuildTool/System/CPPHeaders.cs b/Engine/Source/Programs/UnrealBuildTool/System/CPPHeaders.cs
--- a/Engine/Source/Programs/UnrealBuildTool/System/CPPHeaders.cs
+++ b/Engine/Source/Programs/UnrealBuildTool/System/CPPHeaders.cs
@@ -170,7 +170,7 @@ namespace UnrealBuildTool
 				IncludedFilesMap.Add(CPPFile, IncludedFileList);
 
 				// Gather a list of names of files directly included by this C++ file.
-				List<DependencyInclude> DirectIncludes = GetDirectIncludeDependencies(CPPFile, Config.TargetPlatform, bHackHeaderGenerator);
+				List<DependencyInclude> DirectIncludes = GetDirectIncludeDependencies(CPPFile, Config.TargetPlatform);
 
 				// Build a list of the unique set of files that are included by this file.
 				string SourceFilesDirectory = Path.GetDirectoryName( CPPFile.AbsolutePath );
@@ -323,7 +323,7 @@ namespace UnrealBuildTool
 													RegexOptions.Compiled | RegexOptions.Singleline | RegexOptions.ExplicitCapture );
 
 		/** Finds the names of files directly included by the given C++ file. */
-		public static List<DependencyInclude> GetDirectIncludeDependencies( FileItem CPPFile, CPPTargetPlatform TargetPlatform, bool bHackHeaderGenerator )
+		public static List<DependencyInclude> GetDirectIncludeDependencies( FileItem CPPFile, CPPTargetPlatform TargetPlatform )
 		{
 			// Try to fulfill request from cache first.
 			List<DependencyInclude> Result = null;
@@ -379,14 +379,14 @@ namespace UnrealBuildTool
 						EndIndex = FileContents.Length;
 					}
 
-					CollectHeaders(CPPFile, bHackHeaderGenerator, Result, ref HasUObjects, FileToRead, FileContents, InstalledFolder, StartIndex, EndIndex);
+					CollectHeaders(CPPFile, Result, ref HasUObjects, FileToRead, FileContents, InstalledFolder, StartIndex, EndIndex);
 
 					StartIndex = EndIndex + 1;
 				}
 			}
 			else
 			{
-				CollectHeaders(CPPFile, bHackHeaderGenerator, Result, ref HasUObjects, FileToRead, FileContents, InstalledFolder, 0, FileContents.Length);
+				CollectHeaders(CPPFile, Result, ref HasUObjects, FileToRead, FileContents, InstalledFolder, 0, FileContents.Length);
 			}
 
 			// Store whether this C++ file has UObjects right inside the FileItem for this file
@@ -404,7 +404,6 @@ namespace UnrealBuildTool
 		/// Collects all header files included in a CPPFile
 		/// </summary>
 		/// <param name="CPPFile"></param>
-		/// <param name="bHackHeaderGenerator"></param>
 		/// <param name="Result"></param>
 		/// <param name="HasUObjects"></param>
 		/// <param name="FileToRead"></param>
@@ -412,7 +411,7 @@ namespace UnrealBuildTool
 		/// <param name="InstalledFolder"></param>
 		/// <param name="StartIndex"></param>
 		/// <param name="EndIndex"></param>
-		private static void CollectHeaders(FileItem CPPFile, bool bHackHeaderGenerator, List<DependencyInclude> Result, ref bool HasUObjects, string FileToRead, string FileContents, string InstalledFolder, int StartIndex, int EndIndex)
+		private static void CollectHeaders(FileItem CPPFile, List<DependencyInclude> Result, ref bool HasUObjects, string FileToRead, string FileContents, string InstalledFolder, int StartIndex, int EndIndex)
 		{
 			Match M = CPPHeaderRegex.Match(FileContents, StartIndex, EndIndex - StartIndex);
 			CaptureCollection CC = M.Groups["HeaderFile"].Captures;
diff --git a/Engine/Source/Programs/UnrealBuildTool/System/ExternalExecution.cs b/Engine/Source/Programs/UnrealBuildTool/System/ExternalExecution.cs
--- a/Engine/Source/Programs/UnrealBuildTool/System/ExternalExecution.cs
+++ b/Engine/Source/Programs/UnrealBuildTool/System/ExternalExecution.cs
@@ -58,6 +58,7 @@ namespace UnrealBuildTool
 		{
 			public string       Name;
 			public string       BaseDirectory;
+			public string       IncludeBase;     // The include path which all UHT-generated includes should be relative to
 			public string       OutputDirectory;
 			public List<string> ClassesHeaders;
 			public List<string> PublicHeaders;
@@ -72,17 +73,17 @@ namespace UnrealBuildTool
 			}
 		}
 
-		public UHTManifest(bool InUseRelativePaths, UEBuildTarget Target, string InRootLocalPath, string InRootBuildPath, IEnumerable<UHTModuleInfo> ModuleInfo)
+		public UHTManifest(UEBuildTarget Target, string InRootLocalPath, string InRootBuildPath, IEnumerable<UHTModuleInfo> ModuleInfo)
 		{
-			UseRelativePaths = InUseRelativePaths;
-			IsGameTarget     = TargetRules.IsGameType(Target.Rules.Type);
-			RootLocalPath    = InRootLocalPath;
-			RootBuildPath    = InRootBuildPath;
-			TargetName       = Target.GetTargetName();
+			IsGameTarget  = TargetRules.IsGameType(Target.Rules.Type);
+			RootLocalPath = InRootLocalPath;
+			RootBuildPath = InRootBuildPath;
+			TargetName    = Target.GetTargetName();
 
 			Modules = ModuleInfo.Select(Info => new Module{
 				Name                     = Info.ModuleName,
 				BaseDirectory            = Info.ModuleDirectory,
+				IncludeBase              = Info.ModuleDirectory,
 				OutputDirectory          = UEBuildModuleCPP.GetGeneratedCodeDirectoryForModule(Target, Info.ModuleDirectory, Info.ModuleName),
 				ClassesHeaders           = Info.PublicUObjectClassesHeaders.Select((Header) => Header.AbsolutePath).ToList(),
 				PublicHeaders            = Info.PublicUObjectHeaders       .Select((Header) => Header.AbsolutePath).ToList(),
@@ -95,11 +96,10 @@ namespace UnrealBuildTool
 			}).ToList();
 		}
 
-		public bool         UseRelativePaths; // Generate relative paths or absolute paths
 		public bool         IsGameTarget;     // True if the current target is a game target
 		public string       RootLocalPath;    // The engine path on the local machine
 		public string       RootBuildPath;    // The engine path on the build machine, if different (e.g. Mac/iOS builds)
-		public string		TargetName;       // Name of the target currently being compiled		
+		public string       TargetName;       // Name of the target currently being compiled
 		public List<Module> Modules;
 	}
 
@@ -385,7 +385,7 @@ namespace UnrealBuildTool
 			var CppPlatform    = BuildPlatform.GetCPPTargetPlatform(Target.Platform);
 			var ToolChain      = UEToolChain.GetPlatformToolChain(CppPlatform);
 			var RootLocalPath  = Path.GetFullPath(ProjectFileGenerator.RootRelativePath);
-			var Manifest       = new UHTManifest(UnrealBuildTool.BuildingRocket() || UnrealBuildTool.RunningRocket(), Target, RootLocalPath, ToolChain.ConvertPath(RootLocalPath + '\\'), UObjectModules);
+			var Manifest       = new UHTManifest(Target, RootLocalPath, ToolChain.ConvertPath(RootLocalPath + '\\'), UObjectModules);
 
 			// ensure the headers are up to date
 			if (!bIsBuildingUHT && 
diff --git a/Engine/Source/Programs/UnrealHeaderTool/Private/CodeGenerator.cpp b/Engine/Source/Programs/UnrealHeaderTool/Private/CodeGenerator.cpp
--- a/Engine/Source/Programs/UnrealHeaderTool/Private/CodeGenerator.cpp
+++ b/Engine/Source/Programs/UnrealHeaderTool/Private/CodeGenerator.cpp
@@ -163,6 +163,11 @@ static FString Tabify(const TCHAR *Input)
 	return Result;
 }
 
+void ConvertToBuildIncludePath(UPackage* Package, FString& LocalPath)
+{
+	FPaths::MakePathRelativeTo(LocalPath, *GPackageToManifestModuleMap.FindChecked(Package)->IncludeBase);
+}
+
 /**
  *	Helper function for finding the location of a package
  *	This is required as source now lives in several possible directories
@@ -2013,17 +2018,7 @@ void FNativeClassHeaderGenerator::ExportClassHeaderWrapper( UClass* Class, bool
 
 		check(SourceFilename.EndsWith(TEXT(".h")));
 
-		if(bUseRelativePaths)
-		{
-			// When building Rocket, we need relative paths in the generated headers so that they can be installed to any directory.
-			FPaths::MakePathRelativeTo(NewFileName, *ClassHeaderPath);
-			NewFileName.ReplaceInline(TEXT("\\"), TEXT("/"));
-		}
-		// ...we cannot rebase paths relative to the install directory. It may be on a different drive.
-		else if (FPaths::MakePathRelativeTo(NewFileName, *GManifest.RootLocalPath))
-		{
-			NewFileName = GManifest.RootBuildPath / NewFileName;
-		}
+		ConvertToBuildIncludePath(Package, NewFileName);
 
 		auto IncludeStr = FString::Printf(
 			TEXT("#ifndef %s_%s_generated_h")													LINE_TERMINATOR
@@ -3819,14 +3814,13 @@ bool IsExportOrTemporaryClass(FClass* Class)
 }
 
 // Constructor.
-FNativeClassHeaderGenerator::FNativeClassHeaderGenerator( UPackage* InPackage, FClasses& AllClasses, bool InAllowSaveExportedHeaders, bool bInUseRelativePaths )
+FNativeClassHeaderGenerator::FNativeClassHeaderGenerator( UPackage* InPackage, FClasses& AllClasses, bool InAllowSaveExportedHeaders )
 	: CurrentClass                          (NULL)
 	, API                                   (FPackageName::GetShortName(InPackage).ToUpper())
 	, Package                               (InPackage)
 	, bIsExportingForOffsetDeterminationOnly(false)
 	, bAllowSaveExportedHeaders             (InAllowSaveExportedHeaders)
 	, bFailIfGeneratedCodeChanges           (false)
-	, bUseRelativePaths                     (bInUseRelativePaths)
 {
 	const FString PackageName = FPackageName::GetShortName(Package);
 
@@ -4296,18 +4290,7 @@ void FNativeClassHeaderGenerator::ExportGeneratedCPP()
 	if (ModuleInfo->PCH.Len())
 	{
 		FString PCH = ModuleInfo->PCH;
-		if(bUseRelativePaths)
-		{
-			// When building Rocket, we need relative paths in the generated headers so that they can be installed to any directory.
-			FPaths::MakePathRelativeTo(PCH, *ModuleInfo->GeneratedCPPFilenameBase);
-			PCH.ReplaceInline(TEXT("\\"), TEXT("/"));
-		}
-		// ...we cannot rebase paths relative to the install directory. It may be on a different drive.
-		else if (FPaths::MakePathRelativeTo(PCH, *GManifest.RootLocalPath))
-		{
-			PCH = GManifest.RootBuildPath / PCH;
-		}
-
+		ConvertToBuildIncludePath(Package, PCH);
 		ModulePCHInclude = FString::Printf(TEXT("#include \"%s\"") LINE_TERMINATOR, *PCH);
 	}
 
@@ -4656,7 +4639,7 @@ ECompilationResult::Type UnrealHeaderTool_Main(const FString& ModuleInfoFilename
 			{
 				if (UPackage* Package = Cast<UPackage>(StaticFindObjectFast(UPackage::StaticClass(), NULL, FName(*Module.LongPackageName), false, false)))
 				{
-					Result = FHeaderParser::ParseAllHeadersInside(GWarn, Package, Module.SaveExportedHeaders, GManifest.UseRelativePaths, ScriptPlugins);
+					Result = FHeaderParser::ParseAllHeadersInside(GWarn, Package, Module.SaveExportedHeaders, ScriptPlugins);
 					if (Result != ECompilationResult::Succeeded)
 					{
 						++NumFailures;
diff --git a/Engine/Source/Programs/UnrealHeaderTool/Private/HeaderParser.cpp b/Engine/Source/Programs/UnrealHeaderTool/Private/HeaderParser.cpp
--- a/Engine/Source/Programs/UnrealHeaderTool/Private/HeaderParser.cpp
+++ b/Engine/Source/Programs/UnrealHeaderTool/Private/HeaderParser.cpp
@@ -6445,7 +6445,7 @@ bool FHeaderParser::DependentClassNameFromHeader(const TCHAR* HeaderFilename, FS
 }
 
 // Begins the process of exporting C++ class declarations for native classes in the specified package
-void FHeaderParser::ExportNativeHeaders( UPackage* CurrentPackage, FClasses& AllClasses, bool bAllowSaveExportedHeaders, bool bUseRelativePaths )
+void FHeaderParser::ExportNativeHeaders( UPackage* CurrentPackage, FClasses& AllClasses, bool bAllowSaveExportedHeaders )
 {
 	// Build a list of header filenames
 	TArray<FString>	ClassHeaderFilenames;
@@ -6490,7 +6490,7 @@ void FHeaderParser::ExportNativeHeaders( UPackage* CurrentPackage, FClasses& All
 		}
 
 		// Export native class definitions to package header files.
-		FNativeClassHeaderGenerator(CurrentPackage, AllClasses, bAllowSaveExportedHeaders, bUseRelativePaths);
+		FNativeClassHeaderGenerator(CurrentPackage, AllClasses, bAllowSaveExportedHeaders);
 	}
 }
 
@@ -6606,7 +6606,7 @@ void ExportClassTreeToScriptPlugins(const FClassTree* Node, TArray<IScriptGenera
 }
 
 // Parse all headers for classes that are inside CurrentPackage.
-ECompilationResult::Type FHeaderParser::ParseAllHeadersInside(FFeedbackContext* Warn, UPackage* CurrentPackage, bool bAllowSaveExportedHeaders, bool bUseRelativePaths, TArray<IScriptGeneratorPluginInterface*>& ScriptPlugins)
+ECompilationResult::Type FHeaderParser::ParseAllHeadersInside(FFeedbackContext* Warn, UPackage* CurrentPackage, bool bAllowSaveExportedHeaders, TArray<IScriptGeneratorPluginInterface*>& ScriptPlugins)
 {
 	// Disable loading of objects outside of this package (or more exactly, objects which aren't UFields, CDO, or templates)
 	TGuardValue<bool> AutoRestoreVerifyObjectRefsFlag(GVerifyObjectReferencesOnly, true);
@@ -6643,7 +6643,7 @@ ECompilationResult::Type FHeaderParser::ParseAllHeadersInside(FFeedbackContext*
 			// from the feedback context.
 			Warn->SetContext(NULL);
 
-			ExportNativeHeaders(CurrentPackage, AllClasses, bAllowSaveExportedHeaders, bUseRelativePaths);
+			ExportNativeHeaders(CurrentPackage, AllClasses, bAllowSaveExportedHeaders);
 
 			// Done with header generation
 			if (HeaderParser.LinesParsed > 0)
diff --git a/Engine/Source/Programs/UnrealHeaderTool/Private/HeaderParser.h b/Engine/Source/Programs/UnrealHeaderTool/Private/HeaderParser.h
--- a/Engine/Source/Programs/UnrealHeaderTool/Private/HeaderParser.h
+++ b/Engine/Source/Programs/UnrealHeaderTool/Private/HeaderParser.h
@@ -80,7 +80,7 @@ class FHeaderParser : public FBaseParser, public FContextSupplier
 {
 public:
 	// Parse all headers for classes that are inside LimitOuter.
-	static ECompilationResult::Type ParseAllHeadersInside(FFeedbackContext* Warn, UPackage* LimitOuter, bool bAllowSaveExportedHeaders, bool bUseRelativePaths, TArray<class IScriptGeneratorPluginInterface*>& ScriptPlugins);
+	static ECompilationResult::Type ParseAllHeadersInside(FFeedbackContext* Warn, UPackage* LimitOuter, bool bAllowSaveExportedHeaders, TArray<class IScriptGeneratorPluginInterface*>& ScriptPlugins);
 
 	// Performs a preliminary parse of the text in the specified buffer, pulling out:
 	//   Class name and parent class name
@@ -299,7 +299,7 @@ protected:
 	 * @param	CurrentPackage	the package being compiled
 	 * @param	AllClasses		the class tree for CurrentPackage
 	 */
-	static void ExportNativeHeaders( UPackage* CurrentPackage, FClasses& AllClasses, bool bAllowSaveExportedHeaders, bool bUseRelativePaths );
+	static void ExportNativeHeaders( UPackage* CurrentPackage, FClasses& AllClasses, bool bAllowSaveExportedHeaders );
 
 	// FContextSupplier interface.
 	virtual FString GetContext() OVERRIDE;
diff --git a/Engine/Source/Programs/UnrealHeaderTool/Private/Manifest.cpp b/Engine/Source/Programs/UnrealHeaderTool/Private/Manifest.cpp
--- a/Engine/Source/Programs/UnrealHeaderTool/Private/Manifest.cpp
+++ b/Engine/Source/Programs/UnrealHeaderTool/Private/Manifest.cpp
@@ -61,20 +61,18 @@ FManifest FManifest::LoadFromFile(const FString& Filename)
 
 	TArray<TSharedPtr<FJsonValue>> ModulesArray;
 
-	GetJsonFieldValue(Result.UseRelativePaths, RootObject, TEXT("UseRelativePaths"), TEXT("{manifest root}"));
-	GetJsonFieldValue(Result.IsGameTarget,     RootObject, TEXT("IsGameTarget"),     TEXT("{manifest root}"));
-	GetJsonFieldValue(Result.RootLocalPath,    RootObject, TEXT("RootLocalPath"),    TEXT("{manifest root}"));
-	GetJsonFieldValue(Result.RootBuildPath,    RootObject, TEXT("RootBuildPath"),    TEXT("{manifest root}"));
-	GetJsonFieldValue(Result.TargetName,       RootObject, TEXT("TargetName"),       TEXT("{manifest root}"));
-	GetJsonFieldValue(ModulesArray,            RootObject, TEXT("Modules"),          TEXT("{manifest root}"));
-
-	UE_LOG(LogCompile, Log, TEXT("Loaded manifest: %s"),          *Filename);
-	UE_LOG(LogCompile, Log, TEXT("Manifest.UseRelativePaths=%s"), Result.UseRelativePaths ? TEXT("True") : TEXT("False"));
-	UE_LOG(LogCompile, Log, TEXT("Manifest.IsGameTarget=%s"),     Result.IsGameTarget ? TEXT("True") : TEXT("False"));
-	UE_LOG(LogCompile, Log, TEXT("Manifest.RootLocalPath=%s"),    *Result.RootLocalPath);
-	UE_LOG(LogCompile, Log, TEXT("Manifest.RootBuildPath=%s"),    *Result.RootBuildPath);
-	UE_LOG(LogCompile, Log, TEXT("Manifest.TargetName=%s"),       *Result.TargetName);
-	UE_LOG(LogCompile, Log, TEXT("Manifest.Modules=%d"),          ModulesArray.Num());
+	GetJsonFieldValue(Result.IsGameTarget,  RootObject, TEXT("IsGameTarget"),  TEXT("{manifest root}"));
+	GetJsonFieldValue(Result.RootLocalPath, RootObject, TEXT("RootLocalPath"), TEXT("{manifest root}"));
+	GetJsonFieldValue(Result.RootBuildPath, RootObject, TEXT("RootBuildPath"), TEXT("{manifest root}"));
+	GetJsonFieldValue(Result.TargetName,    RootObject, TEXT("TargetName"),    TEXT("{manifest root}"));
+	GetJsonFieldValue(ModulesArray,         RootObject, TEXT("Modules"),       TEXT("{manifest root}"));
+
+	UE_LOG(LogCompile, Log, TEXT("Loaded manifest: %s"),       *Filename);
+	UE_LOG(LogCompile, Log, TEXT("Manifest.IsGameTarget=%s"),  Result.IsGameTarget ? TEXT("True") : TEXT("False"));
+	UE_LOG(LogCompile, Log, TEXT("Manifest.RootLocalPath=%s"), *Result.RootLocalPath);
+	UE_LOG(LogCompile, Log, TEXT("Manifest.RootBuildPath=%s"), *Result.RootBuildPath);
+	UE_LOG(LogCompile, Log, TEXT("Manifest.TargetName=%s"),    *Result.TargetName);
+	UE_LOG(LogCompile, Log, TEXT("Manifest.Modules=%d"),       ModulesArray.Num());
 
 	Result.RootLocalPath = FPaths::ConvertRelativePathToFull(FilenamePath, Result.RootLocalPath);
 	Result.RootBuildPath = FPaths::ConvertRelativePathToFull(FilenamePath, Result.RootBuildPath);
@@ -99,6 +97,7 @@ FManifest FManifest::LoadFromFile(const FString& Filename)
 
 		GetJsonFieldValue(KnownModule.Name,                      ModuleObj, TEXT("Name"),                     *Outer);
 		GetJsonFieldValue(KnownModule.BaseDirectory,             ModuleObj, TEXT("BaseDirectory"),            *Outer);
+		GetJsonFieldValue(KnownModule.IncludeBase,               ModuleObj, TEXT("IncludeBase"),              *Outer);
 		GetJsonFieldValue(KnownModule.GeneratedIncludeDirectory, ModuleObj, TEXT("OutputDirectory"),          *Outer);
 		GetJsonFieldValue(KnownModule.SaveExportedHeaders,       ModuleObj, TEXT("SaveExportedHeaders"),      *Outer);
 		GetJsonFieldValue(ClassesHeaders,                        ModuleObj, TEXT("ClassesHeaders"),           *Outer);
@@ -111,11 +110,13 @@ FManifest FManifest::LoadFromFile(const FString& Filename)
 
 		// Convert relative paths
 		KnownModule.BaseDirectory             = FPaths::ConvertRelativePathToFull(FilenamePath, KnownModule.BaseDirectory);
+		KnownModule.IncludeBase               = FPaths::ConvertRelativePathToFull(FilenamePath, KnownModule.IncludeBase);
 		KnownModule.GeneratedIncludeDirectory = FPaths::ConvertRelativePathToFull(FilenamePath, KnownModule.GeneratedIncludeDirectory);
 		KnownModule.GeneratedCPPFilenameBase  = FPaths::ConvertRelativePathToFull(FilenamePath, KnownModule.GeneratedCPPFilenameBase);
 
 		// Ensure directories end with a slash, because this aids their use with FPaths::MakePathRelativeTo.
 		if (!KnownModule.BaseDirectory            .EndsWith(TEXT("/"))) { KnownModule.BaseDirectory            .AppendChar(TEXT('/')); }
+		if (!KnownModule.IncludeBase              .EndsWith(TEXT("/"))) { KnownModule.IncludeBase              .AppendChar(TEXT('/')); }
 		if (!KnownModule.GeneratedIncludeDirectory.EndsWith(TEXT("/"))) { KnownModule.GeneratedIncludeDirectory.AppendChar(TEXT('/')); }
 
 		KnownModule.PublicUObjectClassesHeaders.AddZeroed(ClassesHeaders.Num());
diff --git a/Engine/Source/Programs/UnrealHeaderTool/Private/Manifest.h b/Engine/Source/Programs/UnrealHeaderTool/Private/Manifest.h
--- a/Engine/Source/Programs/UnrealHeaderTool/Private/Manifest.h
+++ b/Engine/Source/Programs/UnrealHeaderTool/Private/Manifest.h
@@ -18,6 +18,9 @@ struct FManifestModule
 	/** Base directory of this module on disk */
 	FString BaseDirectory;
 
+	/** The directory to which #includes from this module should be relative */
+	FString IncludeBase;
+
 	/** Directory where generated include files should go */
 	FString GeneratedIncludeDirectory;
 
@@ -42,8 +45,7 @@ struct FManifestModule
 
 struct FManifest
 {
-	bool    UseRelativePaths;
-	bool IsGameTarget;
+	bool    IsGameTarget;
 	FString RootLocalPath;
 	FString RootBuildPath;
 	FString TargetName;
diff --git a/Engine/Source/Programs/UnrealHeaderTool/Private/NativeClassExporter.h b/Engine/Source/Programs/UnrealHeaderTool/Private/NativeClassExporter.h
--- a/Engine/Source/Programs/UnrealHeaderTool/Private/NativeClassExporter.h
+++ b/Engine/Source/Programs/UnrealHeaderTool/Private/NativeClassExporter.h
@@ -98,9 +98,6 @@ private:
 	/** If true, any change in the generated headers will result in UHT failure. */
 	bool bFailIfGeneratedCodeChanges;
 
-	/** If true, writes absolute paths to the generated files. */
-	bool bUseRelativePaths;
-
 	// This exists because it makes debugging much easier on VC2010, since the visualizers can't properly understand templates with templated args
 	struct HeaderDependents : TArray<const FString*>
 	{
@@ -466,5 +463,5 @@ private:
 public:
 
 	// Constructor
-	FNativeClassHeaderGenerator( UPackage* InPackage, FClasses& AllClasses, bool InAllowSaveExportedHeaders, bool bInUseRelativePaths );
+	FNativeClassHeaderGenerator( UPackage* InPackage, FClasses& AllClasses, bool InAllowSaveExportedHeaders );
 };
diff --git a/Engine/Source/Programs/UnrealHeaderTool/Resources/UHTDebugging.manifest b/Engine/Source/Programs/UnrealHeaderTool/Resources/UHTDebugging.manifest
--- a/Engine/Source/Programs/UnrealHeaderTool/Resources/UHTDebugging.manifest
+++ b/Engine/Source/Programs/UnrealHeaderTool/Resources/UHTDebugging.manifest
@@ -1,5 +1,4 @@
 {
-	"UseRelativePaths": false,
 	"IsGameTarget": false,
 	"RootLocalPath": "..\\..\\..\\..",
 	"RootBuildPath": "..\\..\\..\\..",
@@ -7,6 +6,7 @@
 	"Modules":[{
 			"Name": "CoreUObject",
 			"BaseDirectory": "..\\..\\..\\Runtime\\CoreUObject",
+			"IncludeBase": "..\\..\\..\\Runtime\\CoreUObject",
 			"OutputDirectory": "..\\..\\..\\..\\Intermediate/Build/Inc\\CoreUObject\\",
 			"ClassesHeaders": [
 				"..\\..\\..\\Runtime\\CoreUObject\\Classes\\Object.h"
@@ -21,6 +21,7 @@
 		}, {
 			"Name": "UHTDebugging",
 			"BaseDirectory": ".\\UHTDebugging",
+			"IncludeBase": ".\\UHTDebugging",
 			"OutputDirectory": "..\\..\\..\\..\\Intermediate/Build/Inc\\UHTDebugging\\",
 			"ClassesHeaders": [
 				".\\UHTDebugging\\TestObject.h",
