BlueprintNativeEvent and BlueprintImplementableEvent tags automatically imply BlueprintCallable

ttp#335271 EDITOR: Blueprints: BlueprintNative event cannot be added to the graph until it is implemented by a BP

diff --git a/Engine/Source/Programs/UnrealHeaderTool/Private/HeaderParser.cpp b/Engine/Source/Programs/UnrealHeaderTool/Private/HeaderParser.cpp
--- a/Engine/Source/Programs/UnrealHeaderTool/Private/HeaderParser.cpp
+++ b/Engine/Source/Programs/UnrealHeaderTool/Private/HeaderParser.cpp
@@ -165,6 +165,7 @@ namespace
 
 				FuncInfo.FunctionFlags |= FUNC_Event;
 				FuncInfo.FunctionFlags |= FUNC_BlueprintEvent;
+				FuncInfo.FunctionFlags |= FUNC_BlueprintCallable;
 			}
 			else if (Specifier.Key == TEXT("BlueprintImplementableEvent"))
 			{
@@ -184,6 +185,7 @@ namespace
 
 				FuncInfo.FunctionFlags |= FUNC_Event;
 				FuncInfo.FunctionFlags |= FUNC_BlueprintEvent;
+				FuncInfo.FunctionFlags |= FUNC_BlueprintCallable;
 				FuncInfo.FunctionFlags &= ~FUNC_Native;
 			}
 			else if (Specifier.Key == TEXT("Exec"))
@@ -5126,8 +5128,9 @@ void FHeaderParser::CompileFunctionDeclaration(FClasses& AllClasses)
 		const bool bDeprecated      = MetaData.Contains("DeprecatedFunction");       // FBlueprintMetadata::MD_DeprecatedFunction
 		const bool bHasMenuCategory = MetaData.Contains("Category");                 // FBlueprintMetadata::MD_FunctionCategory
 		const bool bInternalOnly    = InternalPtr && *InternalPtr == TEXT("true");
+		const bool bCanBeOverriden  = (FuncInfo.FunctionFlags & FUNC_BlueprintEvent) != 0; //No category is required for BlueprintNativeEvent or BlueprintImplementableEvent
 
-		if (!bHasMenuCategory && !bInternalOnly && !bDeprecated) 
+		if (!bHasMenuCategory && !bInternalOnly && !bDeprecated && !bCanBeOverriden)
 		{ 
 			FError::Throwf(TEXT("Blueprint accessible functions must have a category specified")); 
 		} 
