Back out revision 274 from //depot/UE4/Engine/Source/Programs/UnrealHeaderTool/Private/HeaderParser.cpp

diff --git a/Engine/Source/Programs/UnrealHeaderTool/Private/HeaderParser.cpp b/Engine/Source/Programs/UnrealHeaderTool/Private/HeaderParser.cpp
--- a/Engine/Source/Programs/UnrealHeaderTool/Private/HeaderParser.cpp
+++ b/Engine/Source/Programs/UnrealHeaderTool/Private/HeaderParser.cpp
@@ -165,7 +165,6 @@ namespace
 
 				FuncInfo.FunctionFlags |= FUNC_Event;
 				FuncInfo.FunctionFlags |= FUNC_BlueprintEvent;
-				FuncInfo.FunctionFlags |= FUNC_BlueprintCallable;
 			}
 			else if (Specifier.Key == TEXT("BlueprintImplementableEvent"))
 			{
@@ -185,7 +184,6 @@ namespace
 
 				FuncInfo.FunctionFlags |= FUNC_Event;
 				FuncInfo.FunctionFlags |= FUNC_BlueprintEvent;
-				FuncInfo.FunctionFlags |= FUNC_BlueprintCallable;
 				FuncInfo.FunctionFlags &= ~FUNC_Native;
 			}
 			else if (Specifier.Key == TEXT("Exec"))
@@ -5128,9 +5126,8 @@ void FHeaderParser::CompileFunctionDeclaration(FClasses& AllClasses)
 		const bool bDeprecated      = MetaData.Contains("DeprecatedFunction");       // FBlueprintMetadata::MD_DeprecatedFunction
 		const bool bHasMenuCategory = MetaData.Contains("Category");                 // FBlueprintMetadata::MD_FunctionCategory
 		const bool bInternalOnly    = InternalPtr && *InternalPtr == TEXT("true");
-		const bool bCanBeOverriden  = (FuncInfo.FunctionFlags & FUNC_BlueprintEvent) != 0; //No category is required for BlueprintNativeEvent or BlueprintImplementableEvent
 
-		if (!bHasMenuCategory && !bInternalOnly && !bDeprecated && !bCanBeOverriden)
+		if (!bHasMenuCategory && !bInternalOnly && !bDeprecated) 
 		{ 
 			FError::Throwf(TEXT("Blueprint accessible functions must have a category specified")); 
 		} 
