Fix for: error Default value '(X=0.000,Y=0.000,Z=1.000)' for  Direction  is invalid: 'Invalid value for an FVector'

diff --git a/Engine/Source/Programs/UnrealHeaderTool/Private/HeaderParser.cpp b/Engine/Source/Programs/UnrealHeaderTool/Private/HeaderParser.cpp
--- a/Engine/Source/Programs/UnrealHeaderTool/Private/HeaderParser.cpp
+++ b/Engine/Source/Programs/UnrealHeaderTool/Private/HeaderParser.cpp
@@ -7213,7 +7213,7 @@ bool FHeaderParser::DefaultValueStringCppFormatToInnerFormat(const UProperty* Pr
 			}
 			if(FDefaultValueHelper::Is(CppForm, TEXT("FVector::UpVector")))
 			{
-				OutForm = FString::Printf(TEXT("(X=%3.3f,Y=%3.3f,Z=%3.3f)"),
+				OutForm = FString::Printf(TEXT("%f,%f,%f"),
 					FVector::UpVector.X, FVector::UpVector.Y, FVector::UpVector.Z);
 			}
 			FString Parameters;
@@ -7226,7 +7226,7 @@ bool FHeaderParser::DefaultValueStringCppFormatToInnerFormat(const UProperty* Pr
 				FVector Vector;
 				if(FDefaultValueHelper::ParseVector(Parameters, Vector))
 				{
-					OutForm = FString::Printf(TEXT("(X=%3.3f,Y=%3.3f,Z=%3.3f)"),
+					OutForm = FString::Printf(TEXT("%f,%f,%f"),
 						Vector.X, Vector.Y, Vector.Z);
 				}
 			}
@@ -7247,7 +7247,7 @@ bool FHeaderParser::DefaultValueStringCppFormatToInnerFormat(const UProperty* Pr
 				FRotator Rotator;
 				if(FDefaultValueHelper::ParseRotator(Parameters, Rotator))
 				{
-					OutForm = FString::Printf(TEXT("(P=%f,Y=%f,R=%f)"),
+					OutForm = FString::Printf(TEXT("%f,%f,%f"),
 						Rotator.Pitch, Rotator.Yaw, Rotator.Roll);
 				}
 			}
