Add clickable error with file name and line number information in error/output window when superclass is missing.
#ttp 332332 UE4: Core: UnrealHeaderTool: Doesn't provide source file and line number for error about missing superclass
#proj core
#branch UE4
#summary Add clickable error with file name and line number information in error/output window when superclass is missing.

diff --git a/Engine/Source/Programs/UnrealHeaderTool/Private/ClassMaps.cpp b/Engine/Source/Programs/UnrealHeaderTool/Private/ClassMaps.cpp
--- a/Engine/Source/Programs/UnrealHeaderTool/Private/ClassMaps.cpp
+++ b/Engine/Source/Programs/UnrealHeaderTool/Private/ClassMaps.cpp
@@ -6,6 +6,7 @@
 
 TMap<UClass*, FString>                   GClassStrippedHeaderTextMap;
 TMap<UClass*, FString>                   GClassSourceFileMap;
+TMap<UClass*, int32>                     GClassDeclarationLineNumber;
 TMap<UClass*, FClassHeaderInfo>          GClassGeneratedFileMap;
 TMap<UClass*, TUniqueObj<TArray<FName>>> GClassDependentOnMap;
 TMap<UClass*, FString>                   GClassHeaderNameWithNoPathMap;
diff --git a/Engine/Source/Programs/UnrealHeaderTool/Private/ClassMaps.h b/Engine/Source/Programs/UnrealHeaderTool/Private/ClassMaps.h
--- a/Engine/Source/Programs/UnrealHeaderTool/Private/ClassMaps.h
+++ b/Engine/Source/Programs/UnrealHeaderTool/Private/ClassMaps.h
@@ -52,6 +52,7 @@ struct FManifestModule;
 
 extern TMap<UClass*, FString>                   GClassStrippedHeaderTextMap;
 extern TMap<UClass*, FString>                   GClassSourceFileMap;
+extern TMap<UClass*, int32>                     GClassDeclarationLineNumber;
 extern TMap<UClass*, FClassHeaderInfo>          GClassGeneratedFileMap;
 extern TMap<UClass*, TUniqueObj<TArray<FName>>> GClassDependentOnMap;
 extern TMap<UClass*, FString>                   GClassHeaderNameWithNoPathMap;
diff --git a/Engine/Source/Programs/UnrealHeaderTool/Private/CodeGenerator.cpp b/Engine/Source/Programs/UnrealHeaderTool/Private/CodeGenerator.cpp
--- a/Engine/Source/Programs/UnrealHeaderTool/Private/CodeGenerator.cpp
+++ b/Engine/Source/Programs/UnrealHeaderTool/Private/CodeGenerator.cpp
@@ -4550,6 +4550,7 @@ ECompilationResult::Type UnrealHeaderTool_Main(const FString& ModuleInfoFilename
 
 					UClass* ResultClass = GenerateCodeForHeader(Package, *ClassName, RF_Public|RF_Standalone, *HeaderFile, ClassDeclLine);
 					GClassSourceFileMap.Add(ResultClass, Filename);
+					GClassDeclarationLineNumber.Add(ResultClass, ClassDeclLine);
 					GClassGeneratedFileMap.Add(ResultClass, FClassHeaderInfo(IFileManager::Get().ConvertToAbsolutePathForExternalAppForRead(*Filename)));
 
 					if( CurrentlyProcessing == PublicClassesHeaders )
@@ -4637,7 +4638,31 @@ ECompilationResult::Type UnrealHeaderTool_Main(const FString& ModuleInfoFilename
 
 			if (ScriptSuperClass && !ScriptSuperClass->HasAnyClassFlags(CLASS_Intrinsic) && GClassStrippedHeaderTextMap.Contains(ScriptClass) && !GClassStrippedHeaderTextMap.Contains(ScriptSuperClass))
 			{
-				UE_LOG(LogCompile, Error, TEXT("Superclass %s of class %s not found"), *ScriptSuperClass->GetName(), *ScriptClass->GetName());
+				class FSuperClassContextSupplier : public FContextSupplier
+				{
+				public:
+					FSuperClassContextSupplier(const UClass * Class) :
+						ScriptClass(Class)
+					{ }
+
+					virtual FString GetContext() OVERRIDE
+					{
+						FString Filename = IFileManager::Get().ConvertToAbsolutePathForExternalAppForRead(*GClassSourceFileMap[ScriptClass]);
+						int32 LineNumber = GClassDeclarationLineNumber[ScriptClass];
+						return FString::Printf(TEXT("%s(%i)"), *Filename, LineNumber);
+					}
+				private:
+					const UClass * ScriptClass;
+				} ContextSupplier(ScriptClass);
+
+				auto OldContext = GWarn->GetContext();
+
+				TGuardValue<ELogTimes::Type> DisableLogTimes(GPrintLogTimes, ELogTimes::None);
+
+				GWarn->SetContext(&ContextSupplier);
+				GWarn->Log(ELogVerbosity::Error, FString::Printf(TEXT("Error: Superclass %s of class %s not found"), *ScriptSuperClass->GetName(), *ScriptClass->GetName()));
+				GWarn->SetContext(OldContext);
+
 				Result = ECompilationResult::OtherCompilationError;
 				++NumFailures;
 			}
