Fixing Non-Unity UHT CIS

diff --git a/Engine/Source/Programs/UnrealHeaderTool/Private/HeaderParser.cpp b/Engine/Source/Programs/UnrealHeaderTool/Private/HeaderParser.cpp
--- a/Engine/Source/Programs/UnrealHeaderTool/Private/HeaderParser.cpp
+++ b/Engine/Source/Programs/UnrealHeaderTool/Private/HeaderParser.cpp
@@ -10,6 +10,7 @@
 #include "UObjectAnnotation.h"
 #include "DefaultValueHelper.h"
 #include "IScriptGeneratorPluginInterface.h"
+#include "Manifest.h"
 
 /*-----------------------------------------------------------------------------
 	Constants & declarations.
