Fix UHT parsing of multiple angle closing angle brackets in nested templates (e.g. TArray<TSubclassOf<SomeClass>>).
#ttp 332272  UHT cannot correctly parse TArray<TSubclassOf<AActor>>
#proj core
#branch UE4
#summary Fix UHT parsing of multiple angle closing angle brackets in nested templates (e.g. TArray<TSubclassOf<SomeClass>>).

diff --git a/Engine/Source/Programs/UnrealHeaderTool/Private/BaseParser.cpp b/Engine/Source/Programs/UnrealHeaderTool/Private/BaseParser.cpp
--- a/Engine/Source/Programs/UnrealHeaderTool/Private/BaseParser.cpp
+++ b/Engine/Source/Programs/UnrealHeaderTool/Private/BaseParser.cpp
@@ -296,7 +296,7 @@ bool FBaseParser::IsWhitespace( TCHAR c )
 -----------------------------------------------------------------------------*/
 
 // Gets the next token from the input stream, advancing the variables which keep track of the current input position and line.
-bool FBaseParser::GetToken( FToken& Token, bool bNoConsts/*=false*/ )
+bool FBaseParser::GetToken( FToken& Token, bool bNoConsts/*=false*/, ESymbolParseOption bParseTemplateClosingBracket/*=ESymbolParseOption::Normal*/ )
 {
 	Token.TokenName	= NAME_None;
 	TCHAR c = GetLeadingChar();
@@ -483,7 +483,7 @@ bool FBaseParser::GetToken( FToken& Token, bool bNoConsts/*=false*/ )
 		TCHAR d = GetChar();
 		if
 		(	PAIR('<','<')
-		||	PAIR('>','>')
+		||	(PAIR('>','>') && (bParseTemplateClosingBracket != ESymbolParseOption::CloseTemplateBracket))
 		||	PAIR('!','=')
 		||	PAIR('<','=')
 		||	PAIR('>','=')
@@ -677,11 +677,11 @@ bool FBaseParser::GetConstInt(int32& Result, const TCHAR* Tag)
 	return false;
 }
 
-bool FBaseParser::MatchSymbol( const TCHAR* Match )
+bool FBaseParser::MatchSymbol( const TCHAR* Match, ESymbolParseOption bParseTemplateClosingBracket/*=ESymbolParseOption::Normal*/ )
 {
 	FToken Token;
 
-	if (GetToken(Token, /*bNoConsts=*/ true))
+	if (GetToken(Token, /*bNoConsts=*/ true, bParseTemplateClosingBracket))
 	{
 		if (Token.TokenType==TOKEN_Symbol && !FCString::Stricmp(Token.Identifier, Match))
 		{
@@ -805,9 +805,9 @@ void FBaseParser::UngetToken( FToken& Token )
 //
 // Require a symbol.
 //
-void FBaseParser::RequireSymbol( const TCHAR* Match, const TCHAR* Tag )
+void FBaseParser::RequireSymbol( const TCHAR* Match, const TCHAR* Tag, ESymbolParseOption bParseTemplateClosingBracket/*=ESymbolParseOption::Normal*/ )
 {
-	if (!MatchSymbol(Match))
+	if (!MatchSymbol(Match, bParseTemplateClosingBracket))
 	{
 		FError::Throwf(TEXT("Missing '%s' in %s"), Match, Tag );
 	}
diff --git a/Engine/Source/Programs/UnrealHeaderTool/Private/BaseParser.h b/Engine/Source/Programs/UnrealHeaderTool/Private/BaseParser.h
--- a/Engine/Source/Programs/UnrealHeaderTool/Private/BaseParser.h
+++ b/Engine/Source/Programs/UnrealHeaderTool/Private/BaseParser.h
@@ -11,6 +11,13 @@
 /////////////////////////////////////////////////////
 // FBaseParser
 
+
+enum class ESymbolParseOption
+{
+	Normal,
+	CloseTemplateBracket
+};
+
 //
 // Base class of header parsers.
 //
@@ -76,14 +83,15 @@ protected:
 	/**
 	 * Gets the next token from the input stream, advancing the variables which keep track of the current input position and line.
 	 *
-	 * @param	Token			receives the value of the parsed text; if Token is pre-initialized, special logic is performed
-	 *							to attempt to evaluated Token in the context of that type.  Useful for distinguishing between ambigous symbols
-	 *							like enum tags.
-	 * @param	NoConsts		specify true to indicate that tokens representing literal const values are not allowed.
+	 * @param	Token						receives the value of the parsed text; if Token is pre-initialized, special logic is performed
+	 *										to attempt to evaluated Token in the context of that type.  Useful for distinguishing between ambigous symbols
+	 *										like enum tags.
+	 * @param	NoConsts					specify true to indicate that tokens representing literal const values are not allowed.
+	 * @param	ParseTemplateClosingBracket	specify true to treat >> as two template closing brackets instead of shift operator.
 	 *
 	 * @return	true if a token was successfully processed, false otherwise.
 	 */
-	bool GetToken( FToken& Token, bool bNoConsts = false );
+	bool GetToken( FToken& Token, bool bNoConsts = false, ESymbolParseOption bParseTemplateClosingBracket = ESymbolParseOption::Normal );
 
 	/**
 	 * Put all text from the current position up to either EOL or the StopToken
@@ -114,14 +122,14 @@ protected:
 	bool MatchIdentifier( const TCHAR* Match );
 	bool PeekIdentifier( FName Match );
 	bool PeekIdentifier( const TCHAR* Match );
-	bool MatchSymbol( const TCHAR* Match );
+	bool MatchSymbol( const TCHAR* Match, ESymbolParseOption bParseTemplateClosingBracket = ESymbolParseOption::Normal );
 	void MatchSemi();
 	bool PeekSymbol( const TCHAR* Match );
 
 	// Requiring predefined text.
 	void RequireIdentifier( FName Match, const TCHAR* Tag );
 	void RequireIdentifier( const TCHAR* Match, const TCHAR* Tag );
-	void RequireSymbol( const TCHAR* Match, const TCHAR* Tag );
+	void RequireSymbol( const TCHAR* Match, const TCHAR* Tag, ESymbolParseOption bParseTemplateClosingBracket = ESymbolParseOption::Normal );
 
 	/** Clears out the stored comment. */
 	void ClearComment();
diff --git a/Engine/Source/Programs/UnrealHeaderTool/Private/HeaderParser.cpp b/Engine/Source/Programs/UnrealHeaderTool/Private/HeaderParser.cpp
--- a/Engine/Source/Programs/UnrealHeaderTool/Private/HeaderParser.cpp
+++ b/Engine/Source/Programs/UnrealHeaderTool/Private/HeaderParser.cpp
@@ -2750,7 +2750,7 @@ bool FHeaderParser::GetVarType
 		OriginalVarTypeFlags |= VarProperty.PropertyFlags & (CPF_ContainsInstancedReference | CPF_InstancedReference); // propagate these to the array, we will fix them later
 		VarType.PropertyFlags = OriginalVarTypeFlags;
 		VarProperty.ArrayType = EArrayType::Dynamic;
-		RequireSymbol( TEXT(">"), TEXT("'tarray'") );
+		RequireSymbol( TEXT(">"), TEXT("'tarray'"), ESymbolParseOption::CloseTemplateBracket );
 	}
 	else if( VarType.Matches(TEXT("TAttribute")) )
 	{
@@ -2776,7 +2776,7 @@ bool FHeaderParser::GetVarType
 		OriginalVarTypeFlags |= VarProperty.PropertyFlags & (CPF_ContainsInstancedReference | CPF_InstancedReference);
 		VarType.PropertyFlags = OriginalVarTypeFlags;
 	
-		RequireSymbol(TEXT(">"), TEXT("'tattribute'"));
+		RequireSymbol(TEXT(">"), TEXT("'tattribute'"), ESymbolParseOption::CloseTemplateBracket);
 
 	}
 	else if ( VarType.Matches(TEXT("FString")) )
@@ -2850,7 +2850,7 @@ bool FHeaderParser::GetVarType
 			FError::Throwf(TEXT("Expected the name of a previously defined enum"));
 		}
 
-		RequireSymbol(TEXT(">"), VarType.Identifier);
+		RequireSymbol(TEXT(">"), VarType.Identifier, ESymbolParseOption::CloseTemplateBracket);
 	}
 	else if (UEnum* Enum = FindObject<UEnum>( ANY_PACKAGE, VarType.Identifier ))
 	{
@@ -3013,7 +3013,7 @@ bool FHeaderParser::GetVarType
 					FError::Throwf(TEXT("%s: Missing template type"), VarType.Identifier);
 				}
 
-				RequireSymbol(TEXT(">"), VarType.Identifier);
+				RequireSymbol(TEXT(">"), VarType.Identifier, ESymbolParseOption::CloseTemplateBracket);
 			}
 			else
 			{
@@ -3047,7 +3047,7 @@ bool FHeaderParser::GetVarType
 
 						VarProperty.MetaClass = AllClasses.FindScriptClassOrThrow(Limitor.Identifier);
 
-						RequireSymbol( TEXT(">"), TEXT("'class limitor'") );
+						RequireSymbol( TEXT(">"), TEXT("'class limitor'"), ESymbolParseOption::CloseTemplateBracket );
 					}
 					else
 					{
