Fix unspecific UHT error when missing semicolon in header.
#ttp 336426 UnrealHeaderTool: Unspecific error when missing semicolon in header
#proj core
#branch UE4
#summary Fix unspecific UHT error when missing semicolon in header.

diff --git a/Engine/Source/Programs/UnrealHeaderTool/Private/HeaderParser.cpp b/Engine/Source/Programs/UnrealHeaderTool/Private/HeaderParser.cpp
--- a/Engine/Source/Programs/UnrealHeaderTool/Private/HeaderParser.cpp
+++ b/Engine/Source/Programs/UnrealHeaderTool/Private/HeaderParser.cpp
@@ -3923,7 +3923,8 @@ bool FHeaderParser::SkipDeclaration(FToken& Token)
 			bEndOfDeclarationFound = true;
 			break;
 		}
-		else if (Token.Matches(OpeningBracket))
+
+		if (Token.Matches(OpeningBracket))
 		{
 			// This is a function definition or class declaration.
 			bDefinitionFound = true;
@@ -3937,6 +3938,11 @@ bool FHeaderParser::SkipDeclaration(FToken& Token)
 				bEndOfDeclarationFound = true;
 				break;
 			}
+
+			if (Nest < 0)
+			{
+				FError::Throwf(TEXT("Unexpected '}'. Didn't you miss a semi-colon?"));
+			}
 		}
 		else if (bMacroDeclaration && Nest == 0)
 		{
