Fix unspecified error in UHT on circular header dependency.
#ttp 336732 [UHT] Getting enigmatic error code -1073741571
#proj core
#branch UE4
#summary [UHT] Getting enigmatic error code -1073741571

diff --git a/Engine/Source/Programs/UnrealHeaderTool/Private/Classes.cpp b/Engine/Source/Programs/UnrealHeaderTool/Private/Classes.cpp
--- a/Engine/Source/Programs/UnrealHeaderTool/Private/Classes.cpp
+++ b/Engine/Source/Programs/UnrealHeaderTool/Private/Classes.cpp
@@ -75,11 +75,25 @@ bool FClasses::IsDependentOn(const FClass* Suspect, const FClass* Source) const
 	// suspect is too.
 	for (auto It : Suspect->GetDependentNames())
 	{
-		const FClass* Dependency = FindClass(*It.ToString());
+		auto DependentName = It.ToString();
+		const FClass* Dependency = FindClass(*DependentName);
 
-		// Error.
 		if (!Dependency)
-			continue;
+		{
+			// Check if DependentName is a header file. If so, rerun the check on name without ".h" extension.
+			if (DependentName.EndsWith(TEXT(".h")))
+			{
+				DependentName.RemoveFromEnd(TEXT(".h"));
+				Dependency = FindClass(*DependentName);
+			}
+
+			// If still no class found, we're out of luck.
+			if (!Dependency)
+			{
+				// Error.
+				continue;
+			}
+		}
 
 		// the parser disallows declaring the parent class as a dependency, so the only way this could occur is
 		// if the parent for a native class has been changed (which causes the new parent to be inserted as a dependency),
diff --git a/Engine/Source/Programs/UnrealHeaderTool/Private/HeaderParser.cpp b/Engine/Source/Programs/UnrealHeaderTool/Private/HeaderParser.cpp
--- a/Engine/Source/Programs/UnrealHeaderTool/Private/HeaderParser.cpp
+++ b/Engine/Source/Programs/UnrealHeaderTool/Private/HeaderParser.cpp
@@ -6361,7 +6361,9 @@ ECompilationResult::Type FHeaderParser::ParseHeaders(FClasses& AllClasses, FHead
 		// Check for circular dependency. If the DependsOnClass is dependent on the SubClass, there is one.
 		if (AllClasses.IsDependentOn(DependsOnClass, Class))
 		{
-			FError::Throwf(TEXT("Class %s DependsOn(%s) is a circular dependency."),*DependsOnClass->GetName(),*Class->GetName());
+			HeaderParser.Class = Class;
+			HeaderParser.InputLine = GClassDeclarationLineNumber[Class];
+			FError::Throwf(TEXT("Error: Class %s DependsOn(%s) is a circular dependency."),*DependsOnClass->GetName(),*Class->GetName());
 		}
 
 		// Find first base class of DependsOnClass that is not a base class of Class.
