CIS fix.

#change Ignore inter-module dependencies.

diff --git a/Engine/Source/Programs/UnrealHeaderTool/Private/Classes.cpp b/Engine/Source/Programs/UnrealHeaderTool/Private/Classes.cpp
--- a/Engine/Source/Programs/UnrealHeaderTool/Private/Classes.cpp
+++ b/Engine/Source/Programs/UnrealHeaderTool/Private/Classes.cpp
@@ -89,10 +89,7 @@ bool FClasses::IsDependentOn(const FClass* Suspect, const FClass* Source) const
 
 			// If still no class found, we're out of luck.
 			if (!Dependency)
-			{
-				// Error.
 				continue;
-			}
 		}
 
 		// the parser disallows declaring the parent class as a dependency, so the only way this could occur is
@@ -101,6 +98,10 @@ bool FClasses::IsDependentOn(const FClass* Suspect, const FClass* Source) const
 		if (Suspect->GetSuperClass() == Dependency)
 			continue;
 
+		// Ignore inter-module dependencies, since modules should be self-contained when they are compiled.
+		if (Dependency->GetOutermost() != Suspect->GetOutermost())
+			continue;
+
 		if (Dependency == Source || IsDependentOn(Dependency, Source))
 			return true;
 	}
