WIP refactors

#ttp 331467: Make UHT not depend on CoreUObject

#change GScriptHelper made an object instead of a heap-allocated pointer.
#change NameLookupCPP made an object instead of a heap-allocated pointer.
#add FClasses::FindAnyClass added.
#change FClasses construction moved earlier.
#change FClass::GetInnerClass added.
#change More UClasses changed to FClasses.

diff --git a/Engine/Source/Programs/UnrealHeaderTool/Private/Class.cpp b/Engine/Source/Programs/UnrealHeaderTool/Private/Class.cpp
--- a/Engine/Source/Programs/UnrealHeaderTool/Private/Class.cpp
+++ b/Engine/Source/Programs/UnrealHeaderTool/Private/Class.cpp
@@ -2,6 +2,7 @@
 
 #include "UnrealHeaderTool.h"
 #include "Class.h"
+#include "Classes.h"
 #include "ClassMaps.h"
 
 bool FClass::Inherits(const FClass* SuspectBase) const
@@ -56,6 +57,11 @@ FClass* FClass::GetSuperClass() const
 	return static_cast<FClass*>(static_cast<const UClass*>(this)->GetSuperClass());
 }
 
+FClass* FClass::GetClassWithin() const
+{
+	return (FClass*)ClassWithin;
+}
+
 TArray<FName> FClass::GetDependentNames() const
 {
 	TArray<FName> Result;
diff --git a/Engine/Source/Programs/UnrealHeaderTool/Private/Class.h b/Engine/Source/Programs/UnrealHeaderTool/Private/Class.h
--- a/Engine/Source/Programs/UnrealHeaderTool/Private/Class.h
+++ b/Engine/Source/Programs/UnrealHeaderTool/Private/Class.h
@@ -16,6 +16,8 @@ struct EEnforceInterfacePrefix
 	};
 };
 
+class FClasses;
+
 class FClass : public UClass
 {
 public:
@@ -43,6 +45,13 @@ public:
 	 */
 	FClass* GetSuperClass() const;
 
+	/**
+	 * Returns the 'within' class of this class.
+	 *
+	 * @return The 'within' class of this class.
+	 */
+	FClass* GetClassWithin() const;
+
 	TArray<FName> GetDependentNames() const;
 
 	TArray<FClass*> GetInterfaceTypes() const;
diff --git a/Engine/Source/Programs/UnrealHeaderTool/Private/Classes.cpp b/Engine/Source/Programs/UnrealHeaderTool/Private/Classes.cpp
--- a/Engine/Source/Programs/UnrealHeaderTool/Private/Classes.cpp
+++ b/Engine/Source/Programs/UnrealHeaderTool/Private/Classes.cpp
@@ -61,9 +61,7 @@ FClass* FClasses::GetRootClass() const
 
 bool FClasses::IsDependentOn(const FClass* Suspect, const FClass* Source) const
 {
-	// A class is not dependent on itself.
-	if (Suspect == Source)
-		return false;
+	check(Suspect != Source);
 
 	// Children are all implicitly dependent on their parent, that is, children require their parent
 	// to be compiled first therefore if the source is a parent of the suspect, the suspect is
@@ -106,15 +104,6 @@ bool FClasses::IsDependentOn(const FClass* Suspect, const FClass* Source) const
 			return true;
 	}
 
-	// Consider all children of the suspect if any of them are dependent on the source, the suspect is
-	// too because it itself is dependent on its children.
-
-	if(!Suspect->HasAnyClassFlags(CLASS_Interface))
-	{
-		if (!ClassTree.FindNode(const_cast<FClass*>(Suspect)))
-			FError::Throwf(TEXT("Unparsed class '%s' found while validating DependsOn entries for '%s'"), *Suspect->GetName(), *Source->GetName());
-	}
-
 	return false;
 }
 
@@ -148,6 +137,13 @@ TArray<FClass*> FClasses::GetDerivedClasses(FClass* Parent) const
 	return Result;
 }
 
+FClass* FClasses::FindAnyClass(const TCHAR* ClassName) const
+{
+	check(ClassName);
+
+	return (FClass*)FindObject<UClass>(ANY_PACKAGE, ClassName);
+}
+
 FClass* FClasses::FindScriptClass(const FString& InClassName) const
 {
 	FString ErrorMsg;
@@ -209,12 +205,21 @@ TArray<FClass*> FClasses::GetClassesInPackage(UPackage* InPackage) const
 	return Result;
 }
 
+#if WIP_UHT_REFACTOR
+
 void FClasses::ChangeParentClass(FClass* Class)
 {
 	ClassTree.ChangeParentClass(Class);
 }
 
+bool FClasses::ContainsClass(const FClass* Class) const
+{
+	return !!ClassTree.FindNode(const_cast<FClass*>(Class));
+}
+
 void FClasses::Validate()
 {
 	ClassTree.Validate();
 }
+
+#endif
diff --git a/Engine/Source/Programs/UnrealHeaderTool/Private/Classes.h b/Engine/Source/Programs/UnrealHeaderTool/Private/Classes.h
--- a/Engine/Source/Programs/UnrealHeaderTool/Private/Classes.h
+++ b/Engine/Source/Programs/UnrealHeaderTool/Private/Classes.h
@@ -38,6 +38,8 @@ public:
 
 	TArray<FClass*> GetDerivedClasses(FClass* Parent) const;
 
+	FClass* FindAnyClass(const TCHAR* ClassName) const;
+
 	/** 
 	 * Attempts to find a script class based on the given name. Will attempt to strip
 	 * the prefix of the given name while searching. Throws script errors when appropriate.
@@ -88,6 +90,8 @@ public:
 	 */
 	void ChangeParentClass(FClass* Class);
 
+	bool ContainsClass(const FClass* Class) const;
+
 	/**
 	 * Validates the state of the tree (shouldn't be needed once this class has well-defined invariants).
 	 */
diff --git a/Engine/Source/Programs/UnrealHeaderTool/Private/CodeGenerator.cpp b/Engine/Source/Programs/UnrealHeaderTool/Private/CodeGenerator.cpp
--- a/Engine/Source/Programs/UnrealHeaderTool/Private/CodeGenerator.cpp
+++ b/Engine/Source/Programs/UnrealHeaderTool/Private/CodeGenerator.cpp
@@ -26,10 +26,10 @@ static const bool bMultiLineUPROPERTY = true;
 
 static UClass* GenerateCodeForHeader(UObject* InParent, const TCHAR* Name, EObjectFlags Flags, const TCHAR* Buffer, int32& OutClassDeclLine);
 
-FCompilerMetadataManager* GScriptHelper = NULL;
+FCompilerMetadataManager GScriptHelper;
 
 /** C++ name lookup helper */
-FNameLookupCPP* NameLookupCPP;
+FNameLookupCPP NameLookupCPP;
 
 /////////////////////////////////////////////////////
 // FFlagAudit
@@ -476,7 +476,7 @@ void FNativeClassHeaderGenerator::ExportProperties( UStruct* Struct, int32 TextI
 			if (bAccessSpecifiers)
 			{
 				// find the class info for this class
-				FClassMetaData* ClassData = GScriptHelper->FindClassData(CurrentClass);
+				FClassMetaData* ClassData = GScriptHelper.FindClassData(CurrentClass);
 
 				// find the compiler token for this property
 				FTokenData* PropData = ClassData->FindTokenData(Current);
@@ -582,16 +582,18 @@ static TMap<FString, FString> SingletonNameToExternDecl;
 
 FString FNativeClassHeaderGenerator::GetSingletonName(UField* Item, bool bRequiresValidObject)
 {
-	UClass* ItemClass = Cast<UClass>(Item);
-	if (ItemClass && ItemClass->HasAllClassFlags(CLASS_Intrinsic))
-	{
-		return FString::Printf(TEXT("%s::StaticClass()"), NameLookupCPP->GetNameCPP(ItemClass));
-	}
-
 	FString Suffix;
-	if (ItemClass && !bRequiresValidObject && ItemClass->HasAllClassFlags(CLASS_Native))
+	if (UClass* ItemClass = Cast<UClass>(Item))
 	{
-		Suffix = TEXT("_NoRegister");
+		if (ItemClass->HasAllClassFlags(CLASS_Intrinsic))
+		{
+			return FString::Printf(TEXT("%s::StaticClass()"), NameLookupCPP.GetNameCPP(ItemClass));
+		}
+
+		if (!bRequiresValidObject && ItemClass->HasAllClassFlags(CLASS_Native))
+		{
+			Suffix = TEXT("_NoRegister");
+		}
 	}
 
 	FString Result;
@@ -600,7 +602,7 @@ FString FNativeClassHeaderGenerator::GetSingletonName(UField* Item, bool bRequir
 	{
 		if (Cast<UClass>(Outer) || Cast<UScriptStruct>(Outer))
 		{
-			FString OuterName = NameLookupCPP->GetNameCPP(Cast<UStruct>(Outer));
+			FString OuterName = NameLookupCPP.GetNameCPP(Cast<UStruct>(Outer));
 			if (Result.Len())
 			{
 				Result = OuterName + TEXT("_") + Result;
@@ -626,7 +628,7 @@ FString FNativeClassHeaderGenerator::GetSingletonName(UField* Item, bool bRequir
 		}
 		Outer = Outer->GetOuter();
 	}
-	FString ClassString = NameLookupCPP->GetNameCPP(Item->GetClass());
+	FString ClassString = NameLookupCPP.GetNameCPP(Item->GetClass());
 	// Can't use long package names in function names.
 	if (Result.StartsWith(TEXT("/Script/")))
 	{
@@ -655,6 +657,11 @@ FString FNativeClassHeaderGenerator::GetSingletonName(UField* Item, bool bRequir
 	return Result;
 }
 
+FString FNativeClassHeaderGenerator::GetSingletonName(FClass* Item, bool bRequiresValidObject)
+{
+	return GetSingletonName((UClass*)Item, bRequiresValidObject);
+}
+
 FString FNativeClassHeaderGenerator::PropertyNew(FString& Meta, UProperty* Prop, const FString& OuterString, const FString& PropMacro, const TCHAR* NameSuffix, const TCHAR* Spaces, const TCHAR* SourceStruct)
 {
 	FString ExtraArgs;
@@ -804,7 +811,7 @@ void FNativeClassHeaderGenerator::OutputProperty(FString& Meta, FOutputDevice& O
 		}
 		else
 		{			
-			SourceStruct = NameLookupCPP->GetNameCPP(CastChecked<UStruct>(Prop->GetOuter()));
+			SourceStruct = NameLookupCPP.GetNameCPP(CastChecked<UStruct>(Prop->GetOuter()));
 		}
 		FString PropMacroOuterClass;
 		FString PropNameDep = Prop->GetName();
@@ -916,7 +923,7 @@ static void FindNoExportStructs(TArray<UScriptStruct*>& Structs, UStruct* Start)
 
 FString GetPackageSingletonName(UPackage* Package)
 {
-	static FString ClassString = NameLookupCPP->GetNameCPP(UPackage::StaticClass());
+	static FString ClassString = NameLookupCPP.GetNameCPP(UPackage::StaticClass());
 	return FString(TEXT("Z_Construct_")) + ClassString + TEXT("_") + FPackageName::GetShortName(Package) + TEXT("()");
 }
 
@@ -969,12 +976,12 @@ void FNativeClassHeaderGenerator::ExportGeneratedPackageInitCode(UPackage* Packa
 
 }
 
-void FNativeClassHeaderGenerator::ExportNativeGeneratedInitCode(UClass* Class)
+void FNativeClassHeaderGenerator::ExportNativeGeneratedInitCode(FClass* Class)
 {
 	// Emit the IMPLEMENT_CLASS macro to go in the generated cpp file (unless it's the temporary UHT class).
 	if (!Class->HasAnyClassFlags(CLASS_Temporary))
 	{
-		GeneratedPackageCPP.Logf(TEXT("    IMPLEMENT_CLASS(%s);\r\n"),NameLookupCPP->GetNameCPP(Class));
+		GeneratedPackageCPP.Logf(TEXT("    IMPLEMENT_CLASS(%s);\r\n"),NameLookupCPP.GetNameCPP(Class));
 	}
 
 	int32 MaxLinesPerCpp = 30000;
@@ -1091,11 +1098,11 @@ void FNativeClassHeaderGenerator::ExportNativeGeneratedInitCode(UClass* Class)
 		if ( (ScriptStruct->StructFlags&STRUCT_Native) != 0 )
 		{
 			//@todo .u we don't need the auto register versions of these (except for hotreload, which should be fixed)
-			CppStructOpsString = FString::Printf( TEXT("new UScriptStruct::TCppStructOps<%s>"), NameLookupCPP->GetNameCPP(ScriptStruct) );
+			CppStructOpsString = FString::Printf( TEXT("new UScriptStruct::TCppStructOps<%s>"), NameLookupCPP.GetNameCPP(ScriptStruct) );
 		}
 		else
 		{
-			ExplicitSizeString = FString::Printf( TEXT(", sizeof(%s), ALIGNOF(%s)"), NameLookupCPP->GetNameCPP(ScriptStruct), NameLookupCPP->GetNameCPP(ScriptStruct) );
+			ExplicitSizeString = FString::Printf( TEXT(", sizeof(%s), ALIGNOF(%s)"), NameLookupCPP.GetNameCPP(ScriptStruct), NameLookupCPP.GetNameCPP(ScriptStruct) );
 		}
 
 		GeneratedFunctionText.Logf(TEXT("            ReturnStruct = new(Outer, TEXT(\"%s\"), RF_Public|RF_Transient|RF_Native) UScriptStruct(FPostConstructInitializeProperties(), %s, %s, EStructFlags(0x%08X)%s);\r\n"), 
@@ -1141,7 +1148,7 @@ void FNativeClassHeaderGenerator::ExportNativeGeneratedInitCode(UClass* Class)
 	FunctionsToExport.Sort();
 
 	// find the class info for this class
-	FClassMetaData* ClassData = GScriptHelper->FindClassData(Class);
+	FClassMetaData* ClassData = GScriptHelper.FindClassData(Class);
 
 	// Export the init code for each function
 	for (int32 FuncIndex = 0; FuncIndex < FunctionsToExport.Num(); FuncIndex++)
@@ -1262,7 +1269,7 @@ void FNativeClassHeaderGenerator::ExportNativeGeneratedInitCode(UClass* Class)
 
 			GeneratedFunctionText.Logf(TEXT("    UClass* %s\r\n"), *SingletonNameNoRegister);
 			GeneratedFunctionText.Logf(TEXT("    {\r\n"));
-			GeneratedFunctionText.Logf(TEXT("        return %s::StaticClass();\r\n"), NameLookupCPP->GetNameCPP(Class));
+			GeneratedFunctionText.Logf(TEXT("        return %s::StaticClass();\r\n"), NameLookupCPP.GetNameCPP(Class));
 			GeneratedFunctionText.Logf(TEXT("    }\r\n"));
 
 
@@ -1284,13 +1291,13 @@ void FNativeClassHeaderGenerator::ExportNativeGeneratedInitCode(UClass* Class)
 			GeneratedFunctionText.Logf(TEXT("            %s;\r\n"), *GetSingletonName(Class->GetSuperClass()));
 		}
 		GeneratedFunctionText.Logf(TEXT("            %s;\r\n"), *GetPackageSingletonName(CastChecked<UPackage>(Class->GetOutermost())));
-		GeneratedFunctionText.Logf(TEXT("            OuterClass = %s::StaticClass();\r\n"), NameLookupCPP->GetNameCPP(Class));
+		GeneratedFunctionText.Logf(TEXT("            OuterClass = %s::StaticClass();\r\n"), NameLookupCPP.GetNameCPP(Class));
 		GeneratedFunctionText.Logf(TEXT("            UObjectForceRegistration(OuterClass);\r\n"));
 		uint32 Flags = Class->ClassFlags & CLASS_SaveInCompiledInClasses;
 		GeneratedFunctionText.Logf(TEXT("            OuterClass->ClassFlags |= 0x%08X;\r\n"), Flags);
 		//if (Class->IsChildOf(UObject::StaticClass()))
 		//{
-		//	GeneratedFunctionText.Logf(TEXT("            OuterClass->ClassAddReferencedObjects = &%s::AddReferencedObjects;\r\n"), NameLookupCPP->GetNameCPP(Class));
+		//	GeneratedFunctionText.Logf(TEXT("            OuterClass->ClassAddReferencedObjects = &%s::AddReferencedObjects;\r\n"), NameLookupCPP.GetNameCPP(Class));
 		//}
 		TheFlagAudit.Add(Class, TEXT("ClassFlags"), Flags);
 		GeneratedFunctionText.Logf(TEXT("\r\n"));
@@ -1339,7 +1346,7 @@ void FNativeClassHeaderGenerator::ExportNativeGeneratedInitCode(UClass* Class)
 			FString OffsetString(TEXT("0"));
 			if (Inter.PointerOffset)
 			{
-				OffsetString = FString::Printf(TEXT("VTABLE_OFFSET(%s, %s)"), NameLookupCPP->GetNameCPP(Class), NameLookupCPP->GetNameCPP(Inter.Class, true));
+				OffsetString = FString::Printf(TEXT("VTABLE_OFFSET(%s, %s)"), NameLookupCPP.GetNameCPP(Class), NameLookupCPP.GetNameCPP(Inter.Class, true));
 			}
 			GeneratedFunctionText.Logf(TEXT("            OuterClass->Interfaces.Add(FImplementedInterface(%s, %s, %s ));\r\n"), 
 				*GetSingletonName(Inter.Class, false),
@@ -1372,7 +1379,7 @@ void FNativeClassHeaderGenerator::ExportNativeGeneratedInitCode(UClass* Class)
 	{
 		if (bIsNoExport)
 		{
-			GeneratedFunctionText.Logf(TEXT("    /* friend declarations for pasting into noexport class %s\r\n"), NameLookupCPP->GetNameCPP(Class));
+			GeneratedFunctionText.Logf(TEXT("    /* friend declarations for pasting into noexport class %s\r\n"), NameLookupCPP.GetNameCPP(Class));
 			GeneratedFunctionText.Log(FriendText);
 			GeneratedFunctionText.Logf(TEXT("    */\r\n"));
 			FriendText.Empty();
@@ -1384,17 +1391,17 @@ void FNativeClassHeaderGenerator::ExportNativeGeneratedInitCode(UClass* Class)
 	// Skip temporary UHT classes.
 	if (!Class->HasAnyClassFlags(CLASS_Temporary))
 	{
-		GeneratedFunctionText.Logf(TEXT("    static FCompiledInDefer Z_CompiledInDefer_UClass_%s(%s);\r\n"), NameLookupCPP->GetNameCPP(Class), *SingletonName);
+		GeneratedFunctionText.Logf(TEXT("    static FCompiledInDefer Z_CompiledInDefer_UClass_%s(%s);\r\n"), NameLookupCPP.GetNameCPP(Class), *SingletonName);
 	}
 }
 
-void FNativeClassHeaderGenerator::ExportNatives(UClass* Class)
+void FNativeClassHeaderGenerator::ExportNatives(FClass* Class)
 {
 	// Skip temporary UHT classes.
 	if (Class->HasAnyClassFlags(CLASS_NoExport|CLASS_Temporary))
 		return;
 
-	GeneratedPackageCPP.Logf(TEXT("    void %s::StaticRegisterNatives%s()\r\n"),NameLookupCPP->GetNameCPP(Class),NameLookupCPP->GetNameCPP(Class));
+	GeneratedPackageCPP.Logf(TEXT("    void %s::StaticRegisterNatives%s()\r\n"),NameLookupCPP.GetNameCPP(Class),NameLookupCPP.GetNameCPP(Class));
 	GeneratedPackageCPP.Logf(TEXT("    {\r\n"));
 
 	{
@@ -1412,9 +1419,9 @@ void FNativeClassHeaderGenerator::ExportNatives(UClass* Class)
 		for (auto Func : FunctionsToExport)
 		{
 			GeneratedPackageCPP.Logf(TEXT("        FNativeFunctionRegistrar::RegisterFunction(%s::StaticClass(),\"%s\",(Native)&%s::exec%s);\r\n"),
-				NameLookupCPP->GetNameCPP(Class),
+				NameLookupCPP.GetNameCPP(Class),
 				*Func->GetName(),
-				Class->HasAnyClassFlags(CLASS_Interface) ? *FString::Printf(TEXT("I%s"), *Class->GetName()) : NameLookupCPP->GetNameCPP(Class),
+				Class->HasAnyClassFlags(CLASS_Interface) ? *FString::Printf(TEXT("I%s"), *Class->GetName()) : NameLookupCPP.GetNameCPP(Class),
 				*Func->GetName()
 			);
 		}
@@ -1433,7 +1440,7 @@ void FNativeClassHeaderGenerator::ExportNatives(UClass* Class)
 
 void FNativeClassHeaderGenerator::ExportInterfaceCallFunctions( const TArray<UFunction*>& InCallbackFunctions, FStringOutputDevice& HeaderOutput )
 {
-	FClassMetaData* ClassData = GScriptHelper->FindClassData(CurrentClass);
+	FClassMetaData* ClassData = GScriptHelper.FindClassData(CurrentClass);
 
 	TArray<UFunction*> CallbackFunctions = InCallbackFunctions;
 	CallbackFunctions.Sort();
@@ -1606,9 +1613,9 @@ void FNativeClassHeaderGenerator::ExportNames()
 /**
  * Exports the C++ class declarations for a native interface class.
  */
-void FNativeClassHeaderGenerator::ExportInterfaceClassDeclaration( UClass* Class )
+void FNativeClassHeaderGenerator::ExportInterfaceClassDeclaration( FClass* Class )
 {
-	FClassMetaData* ClassData = GScriptHelper->FindClassData(CurrentClass);
+	FClassMetaData* ClassData = GScriptHelper.FindClassData(CurrentClass);
 
 	TArray<UEnum*>			Enums;
 	TArray<UScriptStruct*>	Structs;
@@ -1639,12 +1646,12 @@ void FNativeClassHeaderGenerator::ExportInterfaceClassDeclaration( UClass* Class
 	// export delegate wrapper function implementations
 	ExportDelegateDefinitions(DelegateFunctions, true);
 
-	UClass* SuperClass = Class->GetSuperClass();
+	FClass* SuperClass = Class->GetSuperClass();
 
 
 	// the name for the C++ version of the UClass
-	const TCHAR* ClassCPPName		= NameLookupCPP->GetNameCPP( Class );
-	const TCHAR* SuperClassCPPName	= NameLookupCPP->GetNameCPP( SuperClass );
+	const TCHAR* ClassCPPName		= NameLookupCPP.GetNameCPP( Class );
+	const TCHAR* SuperClassCPPName	= NameLookupCPP.GetNameCPP( SuperClass );
 
 	// =============================================
 	// Export the UClass declaration
@@ -1657,7 +1664,7 @@ void FNativeClassHeaderGenerator::ExportInterfaceClassDeclaration( UClass* Class
 
 		// Export the class's native function registration.
 		UInterfaceBoilerplate.Logf(TEXT("    private:\r\n"));
-		UInterfaceBoilerplate.Logf(TEXT("    static void StaticRegisterNatives%s();\r\n"),NameLookupCPP->GetNameCPP(Class));
+		UInterfaceBoilerplate.Logf(TEXT("    static void StaticRegisterNatives%s();\r\n"),NameLookupCPP.GetNameCPP(Class));
 		UInterfaceBoilerplate.Log(*FriendText);
 		FriendText.Empty();
 		UInterfaceBoilerplate.Logf( TEXT("public:\r\n") );
@@ -1686,7 +1693,7 @@ void FNativeClassHeaderGenerator::ExportInterfaceClassDeclaration( UClass* Class
 
 		if (Class->ClassWithin != Class->GetSuperClass()->ClassWithin)
 		{
-			UInterfaceBoilerplate.Logf(TEXT("    DECLARE_WITHIN(%s)\r\n"), NameLookupCPP->GetNameCPP( Class->ClassWithin ) );
+			UInterfaceBoilerplate.Logf(TEXT("    DECLARE_WITHIN(%s)\r\n"), NameLookupCPP.GetNameCPP( Class->GetClassWithin() ) );
 		}
 
 		FString MacroName = TEXT("GENERATED_UINTERFACE_BODY()");
@@ -1744,7 +1751,7 @@ void FNativeClassHeaderGenerator::ExportInterfaceClassDeclaration( UClass* Class
 			InterfaceBoilerplate.Logf(TEXT("\tvirtual void GetLifetimeReplicatedProps( TArray< FLifetimeProperty > & OutLifetimeProps ) const OVERRIDE;\r\n"));
 		}
 		const TCHAR* Suffix = TEXT("_INCLASS_IINTERFACE");
-		FString MacroName = FString(NameLookupCPP->GetNameCPP(Class)) + Suffix;
+		FString MacroName = FString(NameLookupCPP.GetNameCPP(Class)) + Suffix;
 		FString Macroized = Macroize(*MacroName, *InterfaceBoilerplate);
 		GeneratedHeaderText.Log(*Macroized);
 		InClassMacroCalls.Logf( TEXT("%s%s\r\n"), FCString::Spc(4),  *MacroName);
@@ -1776,15 +1783,15 @@ int32 FindByIndirectName( const TArray<T>& Array, const FName& Name )
  *
  * @param	Class				The class to be exported.
  */
-void FNativeClassHeaderGenerator::ExportClassHeader( UClass* Class )
+void FNativeClassHeaderGenerator::ExportClassHeader( FClasses& AllClasses, FClass* Class )
 {
-	TArray<UClass*> DependencyChain;
+	TArray<FClass*> DependencyChain;
 	TSet<const UClass*>	VisitedSet;
-	ExportClassHeaderRecursive(Class, DependencyChain, VisitedSet, false);
+	ExportClassHeaderRecursive(AllClasses, Class, DependencyChain, VisitedSet, false);
 }
 
 
-void FNativeClassHeaderGenerator::ExportClassHeaderInner(UClass* Class, bool bValidNonTemporaryClass)
+void FNativeClassHeaderGenerator::ExportClassHeaderInner(FClass* Class, bool bValidNonTemporaryClass)
 {
 	if ( !Class->HasAnyClassFlags(CLASS_Interface) )
 	{
@@ -1839,17 +1846,17 @@ void FNativeClassHeaderGenerator::ExportClassHeaderInner(UClass* Class, bool bVa
 			// VM -> C++ proxies (events and delegates).
 			ExportCallbackFunctions(CallbackFunctions);
 
-			UClass* SuperClass = Class->GetSuperClass();
+			FClass* SuperClass = Class->GetSuperClass();
 
-			const TCHAR* ClassCPPName = NameLookupCPP->GetNameCPP( Class );
-			const TCHAR* SuperClassCPPName = (SuperClass != NULL) ? NameLookupCPP->GetNameCPP( SuperClass ) : NULL;
+			const TCHAR* ClassCPPName = NameLookupCPP.GetNameCPP( Class );
+			const TCHAR* SuperClassCPPName = (SuperClass != NULL) ? NameLookupCPP.GetNameCPP( SuperClass ) : NULL;
 
 			{
 				FStringOutputDevice ClassBoilerplate;
 
 				// Export the class's native function registration.
 				ClassBoilerplate.Logf(TEXT("    private:\r\n"));
-				ClassBoilerplate.Logf(TEXT("    static void StaticRegisterNatives%s();\r\n"),NameLookupCPP->GetNameCPP(Class));
+				ClassBoilerplate.Logf(TEXT("    static void StaticRegisterNatives%s();\r\n"),NameLookupCPP.GetNameCPP(Class));
 				ClassBoilerplate.Log(FriendText);
 				FriendText.Empty();
 				ClassBoilerplate.Logf(TEXT("    public:\r\n"));
@@ -1880,7 +1887,7 @@ void FNativeClassHeaderGenerator::ExportClassHeaderInner(UClass* Class, bool bVa
 
 				if(SuperClass && Class->ClassWithin != SuperClass->ClassWithin)
 				{
-					ClassBoilerplate.Logf(TEXT("    DECLARE_WITHIN(%s)\r\n"), NameLookupCPP->GetNameCPP( Class->ClassWithin ) );
+					ClassBoilerplate.Logf(TEXT("    DECLARE_WITHIN(%s)\r\n"), NameLookupCPP.GetNameCPP( Class->GetClassWithin() ) );
 				}
 
 				// export the class's config name
@@ -1926,7 +1933,7 @@ void FNativeClassHeaderGenerator::ExportClassHeaderInner(UClass* Class, bool bVa
 
 
 				const TCHAR* Suffix = TEXT("_INCLASS");
-				FString MacroName = FString(NameLookupCPP->GetNameCPP(Class)) + Suffix;
+				FString MacroName = FString(NameLookupCPP.GetNameCPP(Class)) + Suffix;
 				FString Macroized = Macroize(*MacroName, *ClassBoilerplate);
 				GeneratedHeaderText.Log(*Macroized);
 				InClassMacroCalls.Logf( TEXT("%s%s\r\n"), FCString::Spc(4),  *MacroName);
@@ -1940,7 +1947,7 @@ void FNativeClassHeaderGenerator::ExportClassHeaderInner(UClass* Class, bool bVa
 	}
 }
 
-void FNativeClassHeaderGenerator::ExportClassHeaderWrapper( UClass* Class, bool bIsExportClass )
+void FNativeClassHeaderGenerator::ExportClassHeaderWrapper( FClass* Class, bool bIsExportClass )
 {
 	const bool bTemporaryJunkClass = Class->HasAnyClassFlags(CLASS_Temporary);
 	const bool bValidClass = !bTemporaryJunkClass;
@@ -1955,7 +1962,7 @@ void FNativeClassHeaderGenerator::ExportClassHeaderWrapper( UClass* Class, bool
 	}
 
 	GeneratedHeaderText.Logf(TEXT("#undef UCLASS_CURRENT_FILE_NAME\r\n"));
-	GeneratedHeaderText.Logf(TEXT("#define UCLASS_CURRENT_FILE_NAME %s\r\n\r\n\r\n"), NameLookupCPP->GetNameCPP(CurrentClass));
+	GeneratedHeaderText.Logf(TEXT("#define UCLASS_CURRENT_FILE_NAME %s\r\n\r\n\r\n"), NameLookupCPP.GetNameCPP(CurrentClass));
 
 	if (bValidClass)
 	{
@@ -2044,7 +2051,7 @@ void FNativeClassHeaderGenerator::ExportClassHeaderWrapper( UClass* Class, bool
 	EnumForeachText.Empty();
 }
 
-void FNativeClassHeaderGenerator::ExportClassHeaderRecursive( UClass* Class, TArray<UClass*>& DependencyChain, TSet<const UClass*>& VisitedSet, bool bCheckDependenciesOnly )
+void FNativeClassHeaderGenerator::ExportClassHeaderRecursive( FClasses& AllClasses, FClass* Class, TArray<FClass*>& DependencyChain, TSet<const UClass*>& VisitedSet, bool bCheckDependenciesOnly )
 {
 	bool bIsExportClass   = GClassStrippedHeaderTextMap.Contains(Class) && Class->HasAnyClassFlags(CLASS_Native|CLASS_Temporary) && !Class->HasAnyClassFlags(CLASS_NoExport);
 	bool bIsCorrectHeader = Class->GetOuter() == Package;
@@ -2087,9 +2094,9 @@ void FNativeClassHeaderGenerator::ExportClassHeaderRecursive( UClass* Class, TAr
 	}
 
 	// Export the super class first.
-	if (UClass* SuperClass = Class->GetSuperClass())
+	if (FClass* SuperClass = Class->GetSuperClass())
 	{
-		ExportClassHeaderRecursive(SuperClass, DependencyChain, VisitedSet, bCheckDependenciesOnly);
+		ExportClassHeaderRecursive(AllClasses, SuperClass, DependencyChain, VisitedSet, bCheckDependenciesOnly);
 	}
 
 	// Export all classes we depend on.
@@ -2103,19 +2110,18 @@ void FNativeClassHeaderGenerator::ExportClassHeaderRecursive( UClass* Class, TAr
 		DependencyClassNameStripped = GetClassNameWithPrefixRemoved( DependencyClassNameStripped );
 
 		// Only export the class if it's in Classes array (it may be in another package).
-		UClass* FoundClass = FindObject<UClass>(ANY_PACKAGE, *DependencyClassNameStripped);
-		if (FoundClass == NULL)
+		FClass* FoundClass = AllClasses.FindAnyClass(*DependencyClassNameStripped);
+		if (!FoundClass)
 		{
 			// Try again with temporary class name. This may be struct only header.
-			DependencyClassNameStripped = GetClassNameWithoutPrefix(DependencyClassName.ToString());
-			DependencyClassNameStripped = FHeaderParser::GenerateTemporaryClassName(*DependencyClassNameStripped);
-			FoundClass                  = FindObject<UClass>(ANY_PACKAGE, *DependencyClassNameStripped);
+			DependencyClassNameStripped = FHeaderParser::GenerateTemporaryClassName(*GetClassNameWithoutPrefix(DependencyClassName.ToString()));
+			FoundClass                  = AllClasses.FindAnyClass(*DependencyClassNameStripped);
 		}
 
 		int32 ClassIndex = FindByIndirectName( Classes, *DependencyClassNameStripped );
-		if (UClass* DependsOnClass = (ClassIndex != INDEX_NONE) ? FoundClass : NULL)
+		if (FClass* DependsOnClass = (ClassIndex != INDEX_NONE) ? FoundClass : NULL)
 		{
-			ExportClassHeaderRecursive(DependsOnClass, DependencyChain, VisitedSet, bCheckDependenciesOnly);
+			ExportClassHeaderRecursive(AllClasses, DependsOnClass, DependencyChain, VisitedSet, bCheckDependenciesOnly);
 		}
 		else if ( !FoundClass && !bClassNameFromHeader )
 		{
@@ -2279,14 +2285,14 @@ void FNativeClassHeaderGenerator::ExportGeneratedStructBodyMacros(const TArray<U
 
 			const FString CombinedLine = FriendLine + StaticClassLine;
 			const FString MacroName = 
-				FString::Printf(TEXT("%s_USTRUCT_BODY_LINE_%d"), NameLookupCPP->GetNameCPP(CurrentClass), Struct->StructMacroDeclaredLineNumber);
+				FString::Printf(TEXT("%s_USTRUCT_BODY_LINE_%d"), NameLookupCPP.GetNameCPP(CurrentClass), Struct->StructMacroDeclaredLineNumber);
 
 			const FString Macroized = Macroize(*MacroName, *CombinedLine);
 			GeneratedHeaderText.Log(*Macroized);
 
 			FString SingletonName = StaticConstructionString.Replace(TEXT("()"), TEXT("")); // function address
 
-			GeneratedPackageCPP.Logf(TEXT("class UScriptStruct* %s::StaticStruct()\r\n"), NameLookupCPP->GetNameCPP(Struct));
+			GeneratedPackageCPP.Logf(TEXT("class UScriptStruct* %s::StaticStruct()\r\n"), NameLookupCPP.GetNameCPP(Struct));
 			GeneratedPackageCPP.Logf(TEXT("{\r\n"));
 			GeneratedPackageCPP.Logf(TEXT("    static class UScriptStruct* Singleton = NULL;\r\n"));
 			GeneratedPackageCPP.Logf(TEXT("    if (!Singleton)\r\n"));
@@ -2296,7 +2302,7 @@ void FNativeClassHeaderGenerator::ExportGeneratedStructBodyMacros(const TArray<U
 			// UStructs can have UClass or UPackage outer (if declared in non-UClass headers).
 			if (Struct->GetOuter()->IsA(UStruct::StaticClass()))
 			{
-				GeneratedPackageCPP.Logf(TEXT("        Singleton = GetStaticStruct(%s, %s::StaticClass(), TEXT(\"%s\"));\r\n"), *SingletonName, NameLookupCPP->GetNameCPP(CastChecked<UStruct>(Struct->GetOuter())), *Struct->GetName());
+				GeneratedPackageCPP.Logf(TEXT("        Singleton = GetStaticStruct(%s, %s::StaticClass(), TEXT(\"%s\"));\r\n"), *SingletonName, NameLookupCPP.GetNameCPP(CastChecked<UStruct>(Struct->GetOuter())), *Struct->GetName());
 			}
 			else
 			{
@@ -2308,21 +2314,21 @@ void FNativeClassHeaderGenerator::ExportGeneratedStructBodyMacros(const TArray<U
 			GeneratedPackageCPP.Logf(TEXT("    }\r\n"));
 			GeneratedPackageCPP.Logf(TEXT("    return Singleton;\r\n"));
 			GeneratedPackageCPP.Logf(TEXT("}\r\n"));
-			GeneratedPackageCPP.Logf(TEXT("static FCompiledInDeferStruct Z_CompiledInDeferStruct_UScriptStruct_%s(%s::StaticStruct);\r\n"), NameLookupCPP->GetNameCPP(Struct), NameLookupCPP->GetNameCPP(Struct));
+			GeneratedPackageCPP.Logf(TEXT("static FCompiledInDeferStruct Z_CompiledInDeferStruct_UScriptStruct_%s(%s::StaticStruct);\r\n"), NameLookupCPP.GetNameCPP(Struct), NameLookupCPP.GetNameCPP(Struct));
 
 			// Generate StaticRegisterNatives equivalent for structs without classes.
 			if (!Struct->GetOuter()->IsA(UStruct::StaticClass()))
 			{
 				const FString ShortPackageName = FPackageName::GetShortName(Struct->GetOuter()->GetName());
-				GeneratedPackageCPP.Logf(TEXT("static struct FScriptStruct_%s_StaticRegisterNatives%s\r\n"), *ShortPackageName, NameLookupCPP->GetNameCPP(Struct));
+				GeneratedPackageCPP.Logf(TEXT("static struct FScriptStruct_%s_StaticRegisterNatives%s\r\n"), *ShortPackageName, NameLookupCPP.GetNameCPP(Struct));
 				GeneratedPackageCPP.Logf(TEXT("{\r\n"));
-				GeneratedPackageCPP.Logf(TEXT("\tFScriptStruct_%s_StaticRegisterNatives%s()\r\n"), *ShortPackageName, NameLookupCPP->GetNameCPP(Struct));
+				GeneratedPackageCPP.Logf(TEXT("\tFScriptStruct_%s_StaticRegisterNatives%s()\r\n"), *ShortPackageName, NameLookupCPP.GetNameCPP(Struct));
 				GeneratedPackageCPP.Logf(TEXT("\t{\r\n"));
 
 				GeneratedPackageCPP.Logf( TEXT("\t\tUScriptStruct::DeferCppStructOps(FName(TEXT(\"%s\")),new UScriptStruct::TCppStructOps<%s%s>);\r\n"), *Struct->GetName(), Struct->GetPrefixCPP(), *Struct->GetName() );
 
 				GeneratedPackageCPP.Logf(TEXT("\t}\r\n"));
-				GeneratedPackageCPP.Logf(TEXT("} ScriptStruct_%s_StaticRegisterNatives%s;\r\n"), *ShortPackageName, NameLookupCPP->GetNameCPP(Struct));
+				GeneratedPackageCPP.Logf(TEXT("} ScriptStruct_%s_StaticRegisterNatives%s;\r\n"), *ShortPackageName, NameLookupCPP.GetNameCPP(Struct));
 			}
 		}
 	}
@@ -2370,10 +2376,10 @@ void FNativeClassHeaderGenerator::ExportMirrorsForNoexportStructs(const TArray<U
 		UScriptStruct* Struct = NativeStructs[i];
 
 		// Export struct.
-		HeaderOutput.Logf( TEXT("%sstruct %s"), FCString::Spc(TextIndent), NameLookupCPP->GetNameCPP( Struct ) );
+		HeaderOutput.Logf( TEXT("%sstruct %s"), FCString::Spc(TextIndent), NameLookupCPP.GetNameCPP( Struct ) );
 		if (Struct->GetSuperStruct() != NULL)
 		{
-			HeaderOutput.Logf(TEXT(" : public %s"), NameLookupCPP->GetNameCPP(Struct->GetSuperStruct()));
+			HeaderOutput.Logf(TEXT(" : public %s"), NameLookupCPP.GetNameCPP(Struct->GetSuperStruct()));
 		}
 		HeaderOutput.Logf( TEXT("\r\n%s{\r\n"), FCString::Spc(TextIndent) );
 
@@ -2485,7 +2491,7 @@ void FNativeClassHeaderGenerator::ExportDelegateDefinitions( const TArray<UFunct
 	}
 
 	// find the class info for this class
-	FClassMetaData* ClassData = GScriptHelper->FindClassData(CurrentClass);
+	FClassMetaData* ClassData = GScriptHelper.FindClassData(CurrentClass);
 	UClass*         Class     = CurrentClass;
 	const FString   ClassName = Class->GetName();
 
@@ -2554,7 +2560,7 @@ void FNativeClassHeaderGenerator::ExportDelegateDefinitions( const TArray<UFunct
 		else
 		{
 			const TCHAR* Suffix = TEXT("_DELEGATES");
-			FString MacroName = FString(NameLookupCPP->GetNameCPP(CurrentClass)) + Suffix;
+			FString MacroName = FString(NameLookupCPP.GetNameCPP(CurrentClass)) + Suffix;
 			FString Macroized = Macroize(*MacroName, *HeaderOutput);
 			GeneratedHeaderText.Log(*Macroized);
 			PrologMacroCalls.Logf( TEXT("%s%s\r\n"), FCString::Spc(0),  *MacroName);
@@ -2690,7 +2696,7 @@ void ExportProtoDeclaration(FOutputDevice& Out, const FString& MessageName, TFie
 void FNativeClassHeaderGenerator::ExportProtoMessage(const TArray<UFunction*>& InCallbackFunctions, int32 Indent, class FStringOutputDevice* Output)
 {
 	// find the class info for this class
-	FClassMetaData* ClassData = GScriptHelper->FindClassData(CurrentClass);
+	FClassMetaData* ClassData = GScriptHelper.FindClassData(CurrentClass);
 
 	// Parms struct definitions.
 	FStringOutputDevice HeaderOutput;
@@ -2920,7 +2926,7 @@ void ExportMCPDeclaration(FOutputDevice& Out, const FString& MessageName, TField
 void FNativeClassHeaderGenerator::ExportMCPMessage(const TArray<UFunction*>& InCallbackFunctions, int32 Indent, class FStringOutputDevice* Output)
 {
 	// find the class info for this class
-	FClassMetaData* ClassData = GScriptHelper->FindClassData(CurrentClass);
+	FClassMetaData* ClassData = GScriptHelper.FindClassData(CurrentClass);
 
 	// Parms struct definitions.
 	FStringOutputDevice HeaderOutput;
@@ -3081,7 +3087,7 @@ void FNativeClassHeaderGenerator::ExportEventParms( const TArray<UFunction*>& In
 	if (!Output)
 	{
 		const TCHAR* Suffix = TEXT("_EVENTPARMS");
-		FString MacroName = FString(NameLookupCPP->GetNameCPP(CurrentClass)) + Suffix + MacroSuffix;
+		FString MacroName = FString(NameLookupCPP.GetNameCPP(CurrentClass)) + Suffix + MacroSuffix;
 		FString Macroized = Macroize(*MacroName, *HeaderOutput);
 		GeneratedHeaderText.Log(*Macroized);
 		PrologMacroCalls.Logf( TEXT("%s%s\r\n"), FCString::Spc(Indent),  *MacroName);
@@ -3226,7 +3232,7 @@ void FNativeClassHeaderGenerator::ExportNativeFunctionHeader( const FFuncInfo& F
 	FString FunctionName;
 	if (FunctionHeaderStyle == EExportFunctionHeaderStyle::Definition)
 	{
-		FunctionName = FString(NameLookupCPP->GetNameCPP(CastChecked<UClass>(Function->GetOuter()), bIsInterface || FunctionType == EExportFunctionType::Interface)) + TEXT("::");
+		FunctionName = FString(NameLookupCPP.GetNameCPP(CastChecked<UClass>(Function->GetOuter()), bIsInterface || FunctionType == EExportFunctionType::Interface)) + TEXT("::");
 	}
 
 	if (FunctionType == EExportFunctionType::Interface)
@@ -3426,7 +3432,7 @@ void FNativeClassHeaderGenerator::ExportFunctionThunk(FStringOutputDevice& RPCWr
 	// export the call to the C++ version
 	if (FunctionData.FunctionExportFlags & FUNCEXPORT_CppStatic)
 	{
-		RPCWrappers.Logf(TEXT("%s::%s(%s);%s"), NameLookupCPP->GetNameCPP(CurrentClass), *FunctionData.CppImplName, *ParameterList, LINE_TERMINATOR);
+		RPCWrappers.Logf(TEXT("%s::%s(%s);%s"), NameLookupCPP.GetNameCPP(CurrentClass), *FunctionData.CppImplName, *ParameterList, LINE_TERMINATOR);
 	}
 	else
 	{
@@ -3439,7 +3445,7 @@ void FNativeClassHeaderGenerator::ExportNativeFunctions(const TArray<UFunction*>
 	FStringOutputDevice RPCWrappers;
 
 	// find the class info for this class
-	FClassMetaData* ClassData = GScriptHelper->FindClassData(CurrentClass);
+	FClassMetaData* ClassData = GScriptHelper.FindClassData(CurrentClass);
 
 	//@todo: UCREMOVAL this multipass stuff is only used for diffs, remove it
 	TArray<UFunction*> SortedNativeFunctions = NativeFunctions;
@@ -3522,7 +3528,7 @@ void FNativeClassHeaderGenerator::ExportNativeFunctions(const TArray<UFunction*>
 	}
 
 	const TCHAR* Suffix = TEXT("_RPC_WRAPPERS");
-	FString MacroName = FString(NameLookupCPP->GetNameCPP(CurrentClass)) + Suffix;
+	FString MacroName = FString(NameLookupCPP.GetNameCPP(CurrentClass)) + Suffix;
 	FString Macroized = Macroize(*MacroName, *RPCWrappers);
 	GeneratedHeaderText.Log(*Macroized);
 	InClassMacroCalls.Logf( TEXT("%s%s\r\n"), FCString::Spc(4),  *MacroName);
@@ -3538,7 +3544,7 @@ void FNativeClassHeaderGenerator::ExportCallbackFunctions( const TArray<UFunctio
 	FStringOutputDevice RPCWrappers;
 	FStringOutputDevice RPCWrappersCPP;
 	// find the class info for this class
-	FClassMetaData* ClassData = GScriptHelper->FindClassData(CurrentClass);
+	FClassMetaData* ClassData = GScriptHelper.FindClassData(CurrentClass);
 
 	// Skip interfaces; they don't have any callbacks
 	if (CurrentClass->HasAnyClassFlags(CLASS_Interface))
@@ -3556,7 +3562,7 @@ void FNativeClassHeaderGenerator::ExportCallbackFunctions( const TArray<UFunctio
 		// Never expecting to export delegate functions this way
 		check( !Function->HasAnyFunctionFlags( FUNC_Delegate ) );
 
-		UClass* Class = CurrentClass;
+		FClass* Class = CurrentClass;
 
 		FFunctionData* CompilerInfo = ClassData->FindFunctionData(Function);
 		check(CompilerInfo);
@@ -3601,7 +3607,7 @@ void FNativeClassHeaderGenerator::ExportCallbackFunctions( const TArray<UFunctio
 			RPCWrappersCPP.Logf(
 				TEXT("%s%sProcessEvent(FindFunctionChecked(%s_%s),%s);\r\n"),
 				FCString::Spc(8),
-				(Function->HasAllFunctionFlags(FUNC_Const)) ? *FString::Printf(TEXT("const_cast<%s*>(this)->"), NameLookupCPP->GetNameCPP(Class)) : TEXT(""),
+				(Function->HasAllFunctionFlags(FUNC_Const)) ? *FString::Printf(TEXT("const_cast<%s*>(this)->"), NameLookupCPP.GetNameCPP(Class)) : TEXT(""),
 				*API,
 				*FunctionName,
 				Parameters.HasParms() ? TEXT("&Parms") : TEXT("NULL")
@@ -3617,7 +3623,7 @@ void FNativeClassHeaderGenerator::ExportCallbackFunctions( const TArray<UFunctio
 	// else drop the implementation on the floor
 
 	const TCHAR* Suffix = TEXT("_CALLBACK_WRAPPERS");
-	FString MacroName = FString(NameLookupCPP->GetNameCPP(CurrentClass)) + Suffix;
+	FString MacroName = FString(NameLookupCPP.GetNameCPP(CurrentClass)) + Suffix;
 	FString Macroized = Macroize(*MacroName, *RPCWrappers);
 	GeneratedHeaderText.Log(*Macroized);
 	InClassMacroCalls.Logf( TEXT("%s%s\r\n"), FCString::Spc(4),  *MacroName);
@@ -3654,7 +3660,7 @@ void FNativeClassHeaderGenerator::ApplyAlternatePropertyExportText( UProperty* P
 		}
 	}
 	// find the class info for this class
-	FClassMetaData* ClassData = GScriptHelper->FindClassData(CurrentClass);
+	FClassMetaData* ClassData = GScriptHelper.FindClassData(CurrentClass);
 	// find the compiler token for this property
 	FTokenData* PropData = ClassData->FindTokenData(Prop);
 
@@ -3674,7 +3680,7 @@ void FNativeClassHeaderGenerator::ApplyAlternatePropertyExportText( UProperty* P
 	}
 	if ( ObjectProp != NULL && PropData && PropData->Token.ExportInfo.Len() > 0)
 	{
-		PropertyText.ReplaceInline(*(FString("class ") + NameLookupCPP->GetNameCPP(ObjectProp->PropertyClass) + TEXT("*")), *PropData->Token.ExportInfo);
+		PropertyText.ReplaceInline(*(FString("class ") + NameLookupCPP.GetNameCPP(ObjectProp->PropertyClass) + TEXT("*")), *PropData->Token.ExportInfo);
 		return;
 	}
 
@@ -3715,7 +3721,7 @@ void FNativeClassHeaderGenerator::ApplyAlternatePropertyExportText( UProperty* P
 
 			if ( ExportText.Len() > 0 )
 			{
-				(FString&)PropertyText = PropertyText.Replace(NameLookupCPP->GetNameCPP(StructProp->Struct), *ExportText);
+				(FString&)PropertyText = PropertyText.Replace(NameLookupCPP.GetNameCPP(StructProp->Struct), *ExportText);
 			}
 		}
 	}
@@ -3866,7 +3872,7 @@ FNativeClassHeaderGenerator::FNativeClassHeaderGenerator( UPackage* InPackage, F
 		ClassesHeaderPath = GeneratedIncludeDirectory / ClassesHeaderName;
 
 		int32 ClassCount = 0;
-		for (UClass* Class : Classes)
+		for (FClass* Class : Classes)
 		{
 			if( Class->GetOuter()==Package && Class->HasAnyClassFlags(CLASS_Native|CLASS_Temporary))
 			{
@@ -3922,11 +3928,11 @@ FNativeClassHeaderGenerator::FNativeClassHeaderGenerator( UPackage* InPackage, F
 			}
 
 			// Export an include line for each header
-			for (UClass* Class : Classes)
+			for (FClass* Class : Classes)
 			{
 				if (Class->GetOuter() == Package)
 				{
-					ExportClassHeader(Class);
+					ExportClassHeader(AllClasses, Class);
 				}
 			}
 
@@ -3947,11 +3953,11 @@ FNativeClassHeaderGenerator::FNativeClassHeaderGenerator( UPackage* InPackage, F
 	if (!bPackageHasAnyExportOrTemporaryClasses)
 	{
 		// If no headers are generated, check if there's any no export classes that need to have the inl file generated.
-		for (UClass* Class : Classes)
+		for (FClass* Class : Classes)
 		{
 			if (Class->GetOuter() == Package && !Class->HasAnyClassFlags(CLASS_Intrinsic | CLASS_Temporary))
 			{
-				ExportClassHeader(Class);
+				ExportClassHeader(AllClasses, Class);
 			}
 		}
 	}
@@ -4460,11 +4466,6 @@ ECompilationResult::Type UnrealHeaderTool_Main(const FString& ModuleInfoFilename
 		return GCompilationResult;
 	}
 
-	NameLookupCPP = new FNameLookupCPP();
-
-	FCompilerMetadataManager* ScriptHelper = new FCompilerMetadataManager();
-	GScriptHelper = ScriptHelper;
-
 	// Load classes for editing.
 	int32 NumFailures = 0;
 
@@ -4664,7 +4665,11 @@ ECompilationResult::Type UnrealHeaderTool_Main(const FString& ModuleInfoFilename
 			{
 				if (UPackage* Package = Cast<UPackage>(StaticFindObjectFast(UPackage::StaticClass(), NULL, FName(*Module.LongPackageName), false, false)))
 				{
-					Result = FHeaderParser::ParseAllHeadersInside(GWarn, Package, Module, ScriptPlugins);
+					// Object which represents all parsed classes
+					FClasses AllClasses(Package);
+					AllClasses.Validate();
+
+					Result = FHeaderParser::ParseAllHeadersInside(AllClasses, GWarn, Package, Module, ScriptPlugins);
 					if (Result != ECompilationResult::Succeeded)
 					{
 						++NumFailures;
@@ -4681,7 +4686,7 @@ ECompilationResult::Type UnrealHeaderTool_Main(const FString& ModuleInfoFilename
 	}
 
 	// Avoid TArray slack for meta data.
-	GScriptHelper->Shrink();
+	GScriptHelper.Shrink();
 
 	if( bWriteContents )
 	{
@@ -4705,12 +4710,6 @@ ECompilationResult::Type UnrealHeaderTool_Main(const FString& ModuleInfoFilename
 	}
 	TheFlagAudit.WriteResults();
 
-	delete NameLookupCPP;
-	NameLookupCPP = NULL;
-
-	delete ScriptHelper;
-	ScriptHelper = NULL;
-
 	GIsRequestingExit = true;
 
 	if(Result == ECompilationResult::Succeeded && NumFailures > 0)
diff --git a/Engine/Source/Programs/UnrealHeaderTool/Private/HeaderParser.cpp b/Engine/Source/Programs/UnrealHeaderTool/Private/HeaderParser.cpp
--- a/Engine/Source/Programs/UnrealHeaderTool/Private/HeaderParser.cpp
+++ b/Engine/Source/Programs/UnrealHeaderTool/Private/HeaderParser.cpp
@@ -4147,7 +4147,7 @@ void FHeaderParser::CompileClassDeclaration(FClasses& AllClasses)
 	Class->GetAutoCollapseCategories(AutoCollapseCategories);
 
 	// Class attributes.
-	FClassMetaData* ClassData = GScriptHelper->FindClassData(Class);
+	FClassMetaData* ClassData = GScriptHelper.FindClassData(Class);
 	check(ClassData);
 
 	// New-style UCLASS() syntax
@@ -4720,7 +4720,7 @@ void FHeaderParser::CompileInterfaceDeclaration(FClasses& AllClasses)
 	}
 
 	// Try parsing metadata for the interface
-	FClassMetaData* ClassData = GScriptHelper->FindClassData(Class);
+	FClassMetaData* ClassData = GScriptHelper.FindClassData(Class);
 	check(ClassData);
 
 	// Register the metadata
@@ -6089,7 +6089,7 @@ ECompilationResult::Type FHeaderParser::ParseHeaderForOneClass(FClasses& AllClas
 
 	ECompilationResult::Type Result = ECompilationResult::OtherCompilationError;
 
-	ClassData = GScriptHelper->AddClassData(Class);
+	ClassData = GScriptHelper.AddClassData(Class);
 
 	// Message.
 	if (FParse::Param(FCommandLine::Get(), TEXT("VERBOSE")))
@@ -6348,24 +6348,25 @@ ECompilationResult::Type FHeaderParser::ParseHeaders(FClasses& AllClasses, FHead
 
 		// Detect potentially unnecessary usage of dependson.
 		if (DependsOnClass->HasAnyClassFlags(CLASS_Parsed))
-		{
 			continue;
-		}
 
 		// Treat manual dependency on a base class as an error to detect bad habits early on.
 		if (Class->IsChildOf(DependsOnClass))
-		{
 			FError::Throwf(TEXT("%s is derived from %s - please remove the DependsOn"),*Class->GetName(),*DependsOnClass->GetName() );
-		}
 
 		// Check for circular dependency. If the DependsOnClass is dependent on the SubClass, there is one.
-		if (AllClasses.IsDependentOn(DependsOnClass, Class))
+		if (DependsOnClass != Class && AllClasses.IsDependentOn(DependsOnClass, Class))
 		{
 			HeaderParser.Class = Class;
 			HeaderParser.InputLine = GClassDeclarationLineNumber[Class];
 			FError::Throwf(TEXT("Error: Class %s DependsOn(%s) is a circular dependency."),*DependsOnClass->GetName(),*Class->GetName());
 		}
 
+		// Consider all children of the suspect if any of them are dependent on the source, the suspect is
+		// too because it itself is dependent on its children.
+		if (!DependsOnClass->HasAnyClassFlags(CLASS_Interface) && !AllClasses.ContainsClass(DependsOnClass))
+			FError::Throwf(TEXT("Unparsed class '%s' found while validating DependsOn entries for '%s'"), *DependsOnClass->GetName(), *Class->GetName());
+
 		// Find first base class of DependsOnClass that is not a base class of Class.
 		TArray<FClass*> ClassesToParse;
 		ClassesToParse.Add(DependsOnClass);
@@ -6613,23 +6614,18 @@ void ExportClassTreeToScriptPlugins(const FClassTree* Node, const FManifestModul
 }
 
 // Parse all headers for classes that are inside CurrentPackage.
-ECompilationResult::Type FHeaderParser::ParseAllHeadersInside(FFeedbackContext* Warn, UPackage* CurrentPackage, const FManifestModule& Module, TArray<IScriptGeneratorPluginInterface*>& ScriptPlugins)
+ECompilationResult::Type FHeaderParser::ParseAllHeadersInside(FClasses& ModuleClasses, FFeedbackContext* Warn, UPackage* CurrentPackage, const FManifestModule& Module, TArray<IScriptGeneratorPluginInterface*>& ScriptPlugins)
 {
 	// Disable loading of objects outside of this package (or more exactly, objects which aren't UFields, CDO, or templates)
 	TGuardValue<bool> AutoRestoreVerifyObjectRefsFlag(GVerifyObjectReferencesOnly, true);
 
-	// Object which represents all parsed classes
-	FClasses AllClasses(CurrentPackage);
-
-	AllClasses.Validate();
-
 	// Create the header parser and register it as the warning context.
 	// Note: This must be declared outside the try block, since the catch block will log into it.
 	FHeaderParser HeaderParser(Warn);
 	Warn->SetContext(&HeaderParser);
 
 	// Set up a filename for the error context if we don't even get as far parsing a class
-	HeaderParser.Filename = IFileManager::Get().ConvertToAbsolutePathForExternalAppForRead(*GClassSourceFileMap[AllClasses.GetRootClass()]);
+	HeaderParser.Filename = IFileManager::Get().ConvertToAbsolutePathForExternalAppForRead(*GClassSourceFileMap[ModuleClasses.GetRootClass()]);
 
 	// Hierarchically parse all classes.
 	ECompilationResult::Type Result = ECompilationResult::Succeeded;
@@ -6640,7 +6636,7 @@ ECompilationResult::Type FHeaderParser::ParseAllHeadersInside(FFeedbackContext*
 		// Parse the headers
 		const bool bParseSubclasses = true;
 
-		Result = FHeaderParser::ParseHeaders(AllClasses, HeaderParser, AllClasses.GetRootClass(), bParseSubclasses);
+		Result = FHeaderParser::ParseHeaders(ModuleClasses, HeaderParser, ModuleClasses.GetRootClass(), bParseSubclasses);
 
 		// Export the autogenerated code wrappers
 		if (Result == ECompilationResult::Succeeded)
@@ -6650,7 +6646,7 @@ ECompilationResult::Type FHeaderParser::ParseAllHeadersInside(FFeedbackContext*
 			// from the feedback context.
 			Warn->SetContext(NULL);
 
-			ExportNativeHeaders(CurrentPackage, AllClasses, Module.SaveExportedHeaders);
+			ExportNativeHeaders(CurrentPackage, ModuleClasses, Module.SaveExportedHeaders);
 
 			// Done with header generation
 			if (HeaderParser.LinesParsed > 0)
@@ -6675,7 +6671,7 @@ ECompilationResult::Type FHeaderParser::ParseAllHeadersInside(FFeedbackContext*
 
 	if (Result == ECompilationResult::Succeeded && ScriptPlugins.Num())
 	{
-		auto RootNode = &AllClasses.GetClassTree();
+		auto RootNode = &ModuleClasses.GetClassTree();
 		for (auto Plugin : ScriptPlugins)
 		{
 			if (Plugin->ShouldExportClassesForModule(Module.Name, Module.ModuleType))
diff --git a/Engine/Source/Programs/UnrealHeaderTool/Private/HeaderParser.h b/Engine/Source/Programs/UnrealHeaderTool/Private/HeaderParser.h
--- a/Engine/Source/Programs/UnrealHeaderTool/Private/HeaderParser.h
+++ b/Engine/Source/Programs/UnrealHeaderTool/Private/HeaderParser.h
@@ -80,7 +80,7 @@ class FHeaderParser : public FBaseParser, public FContextSupplier
 {
 public:
 	// Parse all headers for classes that are inside LimitOuter.
-	static ECompilationResult::Type ParseAllHeadersInside(FFeedbackContext* Warn, UPackage* LimitOuter, const FManifestModule& Module, TArray<class IScriptGeneratorPluginInterface*>& ScriptPlugins);
+	static ECompilationResult::Type ParseAllHeadersInside(FClasses& ModuleClasses, FFeedbackContext* Warn, UPackage* LimitOuter, const FManifestModule& Module, TArray<class IScriptGeneratorPluginInterface*>& ScriptPlugins);
 
 	// Performs a preliminary parse of the text in the specified buffer, pulling out:
 	//   Class name and parent class name
diff --git a/Engine/Source/Programs/UnrealHeaderTool/Private/NativeClassExporter.h b/Engine/Source/Programs/UnrealHeaderTool/Private/NativeClassExporter.h
--- a/Engine/Source/Programs/UnrealHeaderTool/Private/NativeClassExporter.h
+++ b/Engine/Source/Programs/UnrealHeaderTool/Private/NativeClassExporter.h
@@ -32,11 +32,15 @@ namespace EExportFunctionType
 	};
 }
 
+class FClass;
+class FClasses;
+class FModuleClasses;
+
 struct FNativeClassHeaderGenerator
 {
 private:
 
-	UClass*				CurrentClass;
+	FClass*				CurrentClass;
 	TArray<FClass*>		Classes;
 	FString				API;
 	FString			ClassesHeaderPath;
@@ -164,6 +168,7 @@ private:
 
 	/** Return the name of the singleton function that returns the UObject for Item **/
 	FString GetSingletonName(UField* Item, bool bRequiresValidObject=true);
+	FString GetSingletonName(FClass* Item, bool bRequiresValidObject=true);
 
 	/** 
 	 * Export functions used to find and call C++ or script implementation of a script function in the interface 
@@ -173,15 +178,16 @@ private:
 	/**
 	 * Exports the C++ class declarations for a native interface class.
 	 */
-	void ExportInterfaceClassDeclaration( UClass* Class );
+	void ExportInterfaceClassDeclaration( FClass* Class );
 
 	/**
 	 * Appends the header definition for an inheritance hierarchy of classes to the header.
 	 * Wrapper for ExportClassHeaderRecursive().
 	 *
+	 * @param	AllClasses			The classes being processed.
 	 * @param	Class				The class to be exported.
 	 */
-	void ExportClassHeader( UClass* Class );
+	void ExportClassHeader( FClasses& AllClasses, FClass* Class );
 
 	/**
 	 * After all of the dependency checking, and setup for isolating the generated code, actually export the class
@@ -189,7 +195,7 @@ private:
 	 * @param	Class				The class to be exported.
 	 * @param	bValidNonTemporaryClass	Is the class permanent, or a temporary shell?
 	 */
-	void ExportClassHeaderInner( UClass* Class, bool bValidNonTemporaryClass );
+	void ExportClassHeaderInner( FClass* Class, bool bValidNonTemporaryClass );
 
 	/**
 	 * After all of the dependency checking, but before actually exporting the class, set up the generated code
@@ -197,17 +203,18 @@ private:
 	 * @param	Class				The class to be exported.
 	 * @param	bIsExportClass		true if this is an export class
 	 */
-	void ExportClassHeaderWrapper( UClass* Class, bool bIsExportClass );
+	void ExportClassHeaderWrapper( FClass* Class, bool bIsExportClass );
 
 	/**
 	 * Appends the header definition for an inheritance hierarchy of classes to the header.
 	 *
+	 * @param	AllClasses				The classes being processed.
 	 * @param	Class					The class to be exported.
 	 * @param	DependencyChain			Used for finding errors. Must be empty before the first call.
 	 * @param	VisitedSet				The set of classes visited so far. Must be empty before the first call.
 	 * @param	bCheckDependenciesOnly	Whether we should just keep checking for dependency errors, without exporting anything.
 	 */
-	void ExportClassHeaderRecursive( UClass* Class, TArray<UClass*>& DependencyChain, TSet<const UClass*>& VisitedSet, bool bCheckDependenciesOnly );
+	void ExportClassHeaderRecursive( FClasses& AllClasses, FClass* Class, TArray<FClass*>& DependencyChain, TSet<const UClass*>& VisitedSet, bool bCheckDependenciesOnly );
 
 	/**
 	 * Returns a string in the format CLASS_Something|CLASS_Something which represents all class flags that are set for the specified
@@ -370,14 +377,14 @@ private:
 	void ExportFunctionThunk(FStringOutputDevice& RPCWrappers, const FFuncInfo& FunctionData, const TArray<UProperty*>& Parameters, UProperty* Return);
 
 	/** Exports the native function registration code for the given class. */
-	void ExportNatives(UClass* Class);
+	void ExportNatives(FClass* Class);
 
 	/**
 	 * Exports generated singleton functions for UObjects that used to be stored in .u files.
 	 * 
 	 * @param	Class			Class to export
 	**/
-	void ExportNativeGeneratedInitCode(UClass* Class);
+	void ExportNativeGeneratedInitCode(FClass* Class);
 
 	/**
 	 * Exports a generated singleton function to setup the package for compiled-in classes.
diff --git a/Engine/Source/Programs/UnrealHeaderTool/Private/ParserHelper.cpp b/Engine/Source/Programs/UnrealHeaderTool/Private/ParserHelper.cpp
--- a/Engine/Source/Programs/UnrealHeaderTool/Private/ParserHelper.cpp
+++ b/Engine/Source/Programs/UnrealHeaderTool/Private/ParserHelper.cpp
@@ -21,7 +21,7 @@ FFunctionData* FClassMetaData::FindFunctionData( UFunction* Function )
 	if ( Result == NULL )
 	{
 		UClass* OwnerClass = Function->GetOwnerClass();
-		FClassMetaData* OwnerClassData = GScriptHelper->FindClassData(OwnerClass);
+		FClassMetaData* OwnerClassData = GScriptHelper.FindClassData(OwnerClass);
 		if ( OwnerClassData && OwnerClassData != this )
 		{
 			Result = OwnerClassData->FindFunctionData(Function);
@@ -54,7 +54,7 @@ FStructData* FClassMetaData::FindStructData( UScriptStruct* Struct )
 	if ( Result == NULL )
 	{
 		UClass* OwnerClass = Struct->GetOwnerClass();
-		FClassMetaData* OwnerClassData = GScriptHelper->FindClassData(OwnerClass);
+		FClassMetaData* OwnerClassData = GScriptHelper.FindClassData(OwnerClass);
 		if ( OwnerClassData && OwnerClassData != this )
 		{
 			Result = OwnerClassData->FindStructData(Struct);
@@ -137,7 +137,7 @@ FTokenData* FClassMetaData::FindTokenData( UProperty* Prop )
 
 	if ( Result == NULL && OuterClass != NULL )
 	{
-		FClassMetaData* SuperClassData = GScriptHelper->FindClassData(OuterClass);
+		FClassMetaData* SuperClassData = GScriptHelper.FindClassData(OuterClass);
 		if ( SuperClassData && SuperClassData != this )
 		{
 			Result = SuperClassData->FindTokenData(Prop);
diff --git a/Engine/Source/Programs/UnrealHeaderTool/Private/ParserHelper.h b/Engine/Source/Programs/UnrealHeaderTool/Private/ParserHelper.h
--- a/Engine/Source/Programs/UnrealHeaderTool/Private/ParserHelper.h
+++ b/Engine/Source/Programs/UnrealHeaderTool/Private/ParserHelper.h
@@ -5,7 +5,7 @@
 
 #pragma once
 
-extern class FCompilerMetadataManager*	GScriptHelper;
+extern class FCompilerMetadataManager GScriptHelper;
 
 /*-----------------------------------------------------------------------------
 	FPropertyBase.
@@ -1591,15 +1591,6 @@ public:
 class FCompilerMetadataManager : protected TMap<UClass*, TScopedPointer<FClassMetaData> >
 {
 public:
-
-	~FCompilerMetadataManager()
-	{
-		if ( this == GScriptHelper )
-		{
-			GScriptHelper = NULL;
-		}
-	}
-
 	/**
 	 * Adds a new class to be tracked
 	 * 
@@ -1734,39 +1725,35 @@ struct FNameLookupCPP
 	const TCHAR* GetNameCPP( UStruct* Struct, bool bForceInterface = false )
 	{
 		TCHAR* NameCPP = StructNameMap.FindRef( Struct );
-		if ((NameCPP != NULL) && !bForceInterface)
-		{
+		if (NameCPP && !bForceInterface)
 			return NameCPP;
-		}
-		else
+
+		FString DesiredStructName = Struct->GetName();
+		if (UClass* TestClass = Cast<UClass>(Struct))
 		{
-			FString DesiredStructName = Struct->GetName();
-			if (UClass* TestClass = Cast<UClass>(Struct))
+			if (TestClass->HasAnyClassFlags(CLASS_Temporary))
 			{
-				if (TestClass->HasAnyClassFlags(CLASS_Temporary))
-				{
-					DesiredStructName = GClassHeaderNameWithNoPathMap[TestClass];
-				}
+				DesiredStructName = GClassHeaderNameWithNoPathMap[TestClass];
 			}
+		}
 
-			FString	TempName = FString(bForceInterface ? TEXT("I") : Struct->GetPrefixCPP()) + DesiredStructName;
-			int32 StringLength = TempName.Len();
+		FString	TempName = FString(bForceInterface ? TEXT("I") : Struct->GetPrefixCPP()) + DesiredStructName;
+		int32 StringLength = TempName.Len();
 
-			NameCPP = new TCHAR[StringLength + 1];
-			FCString::Strcpy( NameCPP, StringLength + 1, *TempName );
-			NameCPP[StringLength] = 0;
+		NameCPP = new TCHAR[StringLength + 1];
+		FCString::Strcpy( NameCPP, StringLength + 1, *TempName );
+		NameCPP[StringLength] = 0;
 
-			if (bForceInterface)
-			{
-				InterfaceAllocations.Add(NameCPP);
-			}
-			else
-			{
-				StructNameMap.Add( Struct, NameCPP );
-			}
-			
-			return NameCPP;
+		if (bForceInterface)
+		{
+			InterfaceAllocations.Add(NameCPP);
+		}
+		else
+		{
+			StructNameMap.Add( Struct, NameCPP );
 		}
+			
+		return NameCPP;
 	}
 
 private:
@@ -1775,7 +1762,7 @@ private:
 	TArray<TCHAR*> InterfaceAllocations;
 };
 
-extern FNameLookupCPP* NameLookupCPP;
+extern FNameLookupCPP NameLookupCPP;
 
 
 /////////////////////////////////////////////////////
