Fix for circular header includes crashing UnrealHeaderTool

diff --git a/Engine/Source/Programs/UnrealHeaderTool/Private/Classes.cpp b/Engine/Source/Programs/UnrealHeaderTool/Private/Classes.cpp
--- a/Engine/Source/Programs/UnrealHeaderTool/Private/Classes.cpp
+++ b/Engine/Source/Programs/UnrealHeaderTool/Private/Classes.cpp
@@ -62,12 +62,31 @@ FClass* FClasses::GetRootClass() const
 bool FClasses::IsDependentOn(const FClass* Suspect, const FClass* Source) const
 {
 	check(Suspect != Source);
+	TSet<const FClass*> VisitedDpendencies;
+	return IsDependentOn(Suspect, Source, VisitedDpendencies);
+}
 
+bool FClasses::IsDependentOn(const FClass* Suspect, const FClass* Source, TSet<const FClass*>& VisitedDpendencies) const
+{
 	// Children are all implicitly dependent on their parent, that is, children require their parent
 	// to be compiled first therefore if the source is a parent of the suspect, the suspect is
 	// dependent on the source.
 	if (Suspect->Inherits(Source))
+	{
 		return true;
+	}
+
+	// Prevent circular #includes from causing inifinite recursion
+	// Note that although it may mean there's a circular dependency somewhere, it does not
+	// necessarily mean it's the one we're looking for
+	if (VisitedDpendencies.Contains(Suspect))
+	{
+		return false;
+	}
+	else
+	{
+		VisitedDpendencies.Add(Suspect);
+	}
 
 	// Now consider all dependents of the suspect. If any of them are dependent on the source, the
 	// suspect is too.
@@ -87,21 +106,29 @@ bool FClasses::IsDependentOn(const FClass* Suspect, const FClass* Source) const
 
 			// If still no class found, we're out of luck.
 			if (!Dependency)
+			{
 				continue;
+			}				
 		}
 
 		// the parser disallows declaring the parent class as a dependency, so the only way this could occur is
 		// if the parent for a native class has been changed (which causes the new parent to be inserted as a dependency),
 		// if this is the case, skip it or we'll go into a loop
 		if (Suspect->GetSuperClass() == Dependency)
+		{
 			continue;
+		}
 
 		// Ignore inter-module dependencies, since modules should be self-contained when they are compiled.
 		if (Dependency->GetOutermost() != Suspect->GetOutermost())
+		{
 			continue;
+		}
 
-		if (Dependency == Source || IsDependentOn(Dependency, Source))
+		if (Dependency == Source || IsDependentOn(Dependency, Source, VisitedDpendencies))
+		{
 			return true;
+		}
 	}
 
 	return false;
diff --git a/Engine/Source/Programs/UnrealHeaderTool/Private/Classes.h b/Engine/Source/Programs/UnrealHeaderTool/Private/Classes.h
--- a/Engine/Source/Programs/UnrealHeaderTool/Private/Classes.h
+++ b/Engine/Source/Programs/UnrealHeaderTool/Private/Classes.h
@@ -108,6 +108,12 @@ private:
 	FClass*     UObjectClass;
 	FClassTree  ClassTree;
 
+	/**
+	* Determines whether the class hierarchy rooted at Suspect is dependent on the hierarchy rooted at Source.
+	* Used by the public overload of IsDependentOn to recursively track dependencies and handle circular references
+	*/
+	bool IsDependentOn(const FClass* Suspect, const FClass* Source, TSet<const FClass*>& VisitedDpendencies) const;
+
 	friend auto begin(const FClasses& Classes) -> decltype(begin(TObjectRange<FClass>())) { return begin(TObjectRange<FClass>()); }
 	friend auto end  (const FClasses& Classes) -> decltype(end  (TObjectRange<FClass>())) { return end  (TObjectRange<FClass>()); }
 };
