UnrealHeaderTool: making sure command line is always logged (UBT overrides LogInit verbosity setting)

diff --git a/Engine/Source/Programs/UnrealHeaderTool/Private/UnrealHeaderToolMain.cpp b/Engine/Source/Programs/UnrealHeaderTool/Private/UnrealHeaderToolMain.cpp
--- a/Engine/Source/Programs/UnrealHeaderTool/Private/UnrealHeaderToolMain.cpp
+++ b/Engine/Source/Programs/UnrealHeaderTool/Private/UnrealHeaderToolMain.cpp
@@ -60,6 +60,9 @@ INT32_MAIN_INT32_ARGC_TCHAR_ARGV()
 	GIsUCCMakeStandaloneHeaderGenerator = true;
 	GEngineLoop.PreInit(*ShortCmdLine);
 
+	// Log full command line for UHT as UBT overrides LogInit verbosity settings
+	UE_LOG(LogCompile, Log, TEXT("UHT Command Line: %s"), *CmdLine);
+
 	if (ModuleInfoFilename.IsEmpty())
 	{
 		if (!FPlatformMisc::IsDebuggerPresent())
