Always use "override" and "final" keyword on all platforms

diff --git a/Engine/Source/Programs/UnrealHeaderTool/Private/CodeGenerator.cpp b/Engine/Source/Programs/UnrealHeaderTool/Private/CodeGenerator.cpp
--- a/Engine/Source/Programs/UnrealHeaderTool/Private/CodeGenerator.cpp
+++ b/Engine/Source/Programs/UnrealHeaderTool/Private/CodeGenerator.cpp
@@ -1731,7 +1731,7 @@ void FNativeClassHeaderGenerator::ExportInterfaceClassDeclaration( FClass* Class
 		}
 		if (bNeedsRep)
 		{
-			InterfaceBoilerplate.Logf(TEXT("\tvirtual void GetLifetimeReplicatedProps( TArray< FLifetimeProperty > & OutLifetimeProps ) const OVERRIDE;\r\n"));
+			InterfaceBoilerplate.Logf(TEXT("\tvirtual void GetLifetimeReplicatedProps( TArray< FLifetimeProperty > & OutLifetimeProps ) const override;\r\n"));
 		}
 		const TCHAR* Suffix = TEXT("_INCLASS_IINTERFACE");
 		FString MacroName = FString(NameLookupCPP.GetNameCPP(Class)) + Suffix;
@@ -1909,7 +1909,7 @@ void FNativeClassHeaderGenerator::ExportClassHeaderInner(FClass* Class, bool bVa
 				{
 					if( (It->PropertyFlags & CPF_Net) != 0 )
 					{
-						ClassBoilerplate.Logf(TEXT("\tvirtual void GetLifetimeReplicatedProps( TArray< FLifetimeProperty > & OutLifetimeProps ) const OVERRIDE;\r\n"));
+						ClassBoilerplate.Logf(TEXT("\tvirtual void GetLifetimeReplicatedProps( TArray< FLifetimeProperty > & OutLifetimeProps ) const override;\r\n"));
 						break;
 					}
 				}
diff --git a/Engine/Source/Runtime/Core/Public/Android/AndroidPlatform.h b/Engine/Source/Runtime/Core/Public/Android/AndroidPlatform.h
--- a/Engine/Source/Runtime/Core/Public/Android/AndroidPlatform.h
+++ b/Engine/Source/Runtime/Core/Public/Android/AndroidPlatform.h
@@ -57,8 +57,6 @@ typedef FAndroidTypes FPlatformTypes;
 #define TEXT_HELPER(a,b)	a ## b
 #define TEXT(s)				TEXT_HELPER(L, s)
 
-#define OVERRIDE override
-#define FINAL	 final
 #define ABSTRACT abstract
 
 // Alignment.
diff --git a/Engine/Source/Runtime/Core/Public/HAL/Platform.h b/Engine/Source/Runtime/Core/Public/HAL/Platform.h
--- a/Engine/Source/Runtime/Core/Public/HAL/Platform.h
+++ b/Engine/Source/Runtime/Core/Public/HAL/Platform.h
@@ -288,13 +288,15 @@
 	#define TYPENAME_OUTSIDE_TEMPLATE	typename
 #endif
 
-// Method modifiers
-#ifndef OVERRIDE						
-	#define OVERRIDE
+// Legacy method modifier macros.  You shouldn't use these macros in modern code.  Use the built-in keyword directly.
+#ifndef OVERRIDE
+	#define OVERRIDE override
 #endif
 #ifndef FINAL						
-	#define FINAL
+	#define FINAL final
 #endif
+
+// Method modifiers
 #ifndef ABSTRACT						
 	#define ABSTRACT
 #endif
diff --git a/Engine/Source/Runtime/Core/Public/HTML5/HTML5Platform.h b/Engine/Source/Runtime/Core/Public/HTML5/HTML5Platform.h
--- a/Engine/Source/Runtime/Core/Public/HTML5/HTML5Platform.h
+++ b/Engine/Source/Runtime/Core/Public/HTML5/HTML5Platform.h
@@ -73,15 +73,10 @@ typedef FHTML5Types FPlatformTypes;
 #pragma warning(disable : 4996) // 'function' was was declared deprecated  (needed for the secure string functions)
 #pragma warning( disable : 4530 ) // warning C4530: C++ exception handler used, but unwind semantics are not enabled. Specify /EHsc. for HTML5Win32 code, which uses stl. 
 #endif
-#define OVERRIDE			override //__attribute__((override))
 #if PLATFORM_HTML5_WIN32
 // Backwater of the spec. All compilers support this except microsoft, and they will soon
 #define TYPENAME_OUTSIDE_TEMPLATE
 #define EXCEPTION_EXECUTE_HANDLER       1
-
-#define FINAL				sealed
-#else
-#define FINAL				final
 #endif
 #define ABSTRACT			abstract
 
diff --git a/Engine/Source/Runtime/Core/Public/IOS/IOSPlatform.h b/Engine/Source/Runtime/Core/Public/IOS/IOSPlatform.h
--- a/Engine/Source/Runtime/Core/Public/IOS/IOSPlatform.h
+++ b/Engine/Source/Runtime/Core/Public/IOS/IOSPlatform.h
@@ -57,8 +57,6 @@ typedef FIOSPlatformTypes FPlatformTypes;
 #define TEXT_HELPER(a,b)	a ## b
 #define TEXT(s)				TEXT_HELPER(L, s)
 
-#define OVERRIDE override
-#define FINAL
 #define ABSTRACT abstract
 
 // Strings.
diff --git a/Engine/Source/Runtime/Core/Public/Linux/LinuxPlatform.h b/Engine/Source/Runtime/Core/Public/Linux/LinuxPlatform.h
--- a/Engine/Source/Runtime/Core/Public/Linux/LinuxPlatform.h
+++ b/Engine/Source/Runtime/Core/Public/Linux/LinuxPlatform.h
@@ -64,8 +64,6 @@ typedef FLinuxPlatformTypes FPlatformTypes;
 #define PRAGMA_DISABLE_OPTIMIZATION_ACTUAL _Pragma("optimize(\"\",off)")
 #define PRAGMA_ENABLE_OPTIMIZATION_ACTUAL  _Pragma("optimize(\"\",on)")
 
-#define OVERRIDE override
-#define FINAL	 final
 #define ABSTRACT abstract
 
 // DLL export and import definitions
diff --git a/Engine/Source/Runtime/Core/Public/Mac/MacPlatform.h b/Engine/Source/Runtime/Core/Public/Mac/MacPlatform.h
--- a/Engine/Source/Runtime/Core/Public/Mac/MacPlatform.h
+++ b/Engine/Source/Runtime/Core/Public/Mac/MacPlatform.h
@@ -52,8 +52,6 @@ typedef FMacPlatformTypes FPlatformTypes;
 #define TEXT_HELPER(a,b)	a ## b
 #define TEXT(s)				TEXT_HELPER(L, s)
 
-#define OVERRIDE override
-#define FINAL
 #define ABSTRACT abstract
 
 // Alignment.
diff --git a/Engine/Source/Runtime/Core/Public/WinRT/WinRTARMPlatform.h b/Engine/Source/Runtime/Core/Public/WinRT/WinRTARMPlatform.h
--- a/Engine/Source/Runtime/Core/Public/WinRT/WinRTARMPlatform.h
+++ b/Engine/Source/Runtime/Core/Public/WinRT/WinRTARMPlatform.h
@@ -66,8 +66,6 @@ typedef FWinRTTypes FPlatformTypes;
 #define TYPENAME_OUTSIDE_TEMPLATE
 
 #pragma warning(disable : 4481) // nonstandard extension used: override specifier 'override'
-#define OVERRIDE 
-#define FINAL sealed
 #define ABSTRACT abstract
 
 // Strings.
diff --git a/Engine/Source/Runtime/Core/Public/WinRT/WinRTPlatform.h b/Engine/Source/Runtime/Core/Public/WinRT/WinRTPlatform.h
--- a/Engine/Source/Runtime/Core/Public/WinRT/WinRTPlatform.h
+++ b/Engine/Source/Runtime/Core/Public/WinRT/WinRTPlatform.h
@@ -78,8 +78,6 @@ typedef FWinRTTypes FPlatformTypes;
 #define TYPENAME_OUTSIDE_TEMPLATE
 
 #pragma warning(disable : 4481) // nonstandard extension used: override specifier 'override'
-#define OVERRIDE 
-#define FINAL sealed
 #define ABSTRACT abstract
 
 // Strings.
diff --git a/Engine/Source/Runtime/Core/Public/Windows/WIndowsPlatform.h b/Engine/Source/Runtime/Core/Public/Windows/WIndowsPlatform.h
--- a/Engine/Source/Runtime/Core/Public/Windows/WIndowsPlatform.h
+++ b/Engine/Source/Runtime/Core/Public/Windows/WIndowsPlatform.h
@@ -104,11 +104,8 @@ typedef FWindowsPlatformTypes FPlatformTypes;
 #endif
 
 #pragma warning(disable : 4481) // nonstandard extension used: override specifier 'override'
-#define OVERRIDE override
 #if PLATFORM_WINDOWS_CLANG
 	#define CONSTEXPR constexpr
-#else
-	#define FINAL sealed
 #endif
 #define ABSTRACT abstract
 
