UnrealHeaderTool: support for 'typename const' parameter declarations.

diff --git a/Engine/Source/Programs/UnrealHeaderTool/Private/HeaderParser.cpp b/Engine/Source/Programs/UnrealHeaderTool/Private/HeaderParser.cpp
--- a/Engine/Source/Programs/UnrealHeaderTool/Private/HeaderParser.cpp
+++ b/Engine/Source/Programs/UnrealHeaderTool/Private/HeaderParser.cpp
@@ -3052,6 +3052,9 @@ bool FHeaderParser::GetVarType
 				// Eat the star that indicates this is a pointer to the UObject
 				if (bExpectStar)
 				{
+					// Const after variable type but before pointer symbol
+					MatchIdentifier(TEXT("const"));
+
 					RequireSymbol(TEXT("*"), TEXT("Expected a pointer type"));
 
 					VarProperty.PointerType = EPointerType::Native;
@@ -3097,6 +3100,15 @@ bool FHeaderParser::GetVarType
 		}
 	}
 
+	if (VariableCategory != EVariableCategory::Member)
+	{
+		// const after the variable type support (only for params)
+		if (MatchIdentifier(TEXT("const")))
+		{
+			Flags |= CPF_ConstParm;
+		}
+	}
+
 	if (bUnconsumedConstKeyword)
 	{
 		FError::Throwf(TEXT("Inappropriate keyword 'const' on variable of type '%s'"), VarType.Identifier );
