Updated UHTDebugging.manifest to new format.

diff --git a/Engine/Source/Programs/UnrealHeaderTool/Resources/UHTDebugging.manifest b/Engine/Source/Programs/UnrealHeaderTool/Resources/UHTDebugging.manifest
--- a/Engine/Source/Programs/UnrealHeaderTool/Resources/UHTDebugging.manifest
+++ b/Engine/Source/Programs/UnrealHeaderTool/Resources/UHTDebugging.manifest
@@ -5,6 +5,7 @@
 	"TargetName": "UHTDebugging",
 	"Modules":[{
 			"Name": "CoreUObject",
+			"ModuleType": "Runtime",
 			"BaseDirectory": "..\\..\\..\\Runtime\\CoreUObject",
 			"IncludeBase": "..\\..\\..\\Runtime\\CoreUObject",
 			"OutputDirectory": "..\\..\\..\\..\\Intermediate/Build/Inc\\CoreUObject\\",
@@ -20,6 +21,7 @@
 			"SaveExportedHeaders": true
 		}, {
 			"Name": "UHTDebugging",
+			"ModuleType": "Game",
 			"BaseDirectory": ".\\UHTDebugging",
 			"IncludeBase": ".\\UHTDebugging",
 			"OutputDirectory": "..\\..\\..\\..\\Intermediate/Build/Inc\\UHTDebugging\\",
