Editor: Now with more dragons!
- Added new class metadata Experimental and EarlyAccessPreview, which will trigger a warning section in the actor details panel

Experimental:
- Here be dragons!  It's a totally unsupported and undocumented prototype

Early Access Preview:
- While not considered production-ready, it's a step beyond 'experimental' and is being provided as a preview of things to come

diff --git a/Engine/Source/Editor/DetailCustomizations/Private/ActorDetails.cpp b/Engine/Source/Editor/DetailCustomizations/Private/ActorDetails.cpp
--- a/Engine/Source/Editor/DetailCustomizations/Private/ActorDetails.cpp
+++ b/Engine/Source/Editor/DetailCustomizations/Private/ActorDetails.cpp
@@ -94,6 +94,8 @@ void FActorDetails::CustomizeDetails( IDetailLayoutBuilder& DetailLayout )
 			DetailLayout.HideProperty( GET_MEMBER_NAME_CHECKED(AActor, SpriteScale) );
 		}
 
+		AddExperimentalWarningCategory( DetailLayout );
+
 		AddTransformCategory( DetailLayout );
 
 		AddMaterialCategory( DetailLayout );
@@ -414,7 +416,52 @@ void FActorDetails::AddTransformCategory( IDetailLayoutBuilder& DetailBuilder )
 	IDetailCategoryBuilder& TransformCategory = DetailBuilder.EditCategory( "TransformCommon", LOCTEXT("TransformCommonCategory", "Transform").ToString(), ECategoryPriority::Transform );
 
 	TransformCategory.AddCustomBuilder( TransformDetails );
+}
+
+void FActorDetails::AddExperimentalWarningCategory( IDetailLayoutBuilder& DetailBuilder )
+{
+	const FSelectedActorInfo& SelectedActorInfo = DetailBuilder.GetDetailsView().GetSelectedActorInfo();
+
+	if (SelectedActorInfo.bHaveExperimentalClass || SelectedActorInfo.bHaveEarlyAccessClass)
+	{
+		const bool bExperimental = SelectedActorInfo.bHaveExperimentalClass;
+
+		const FName CategoryName(TEXT("Warning"));
+		const FString CategoryDisplayName = LOCTEXT("WarningCategoryDisplayName", "Warning").ToString();
+		const FText WarningText = bExperimental ? LOCTEXT("ExperimentalClassWarning", "Uses experimental class") : LOCTEXT("EarlyAccessClassWarning", "Uses early access class");
+		const FString SearchString = WarningText.ToString();
+		const FText Tooltip = bExperimental ? LOCTEXT("ExperimentalClassTooltip", "Here be dragons!  Uses one or more unsupported 'experimental' classes") : LOCTEXT("EarlyAccessClassTooltip", "Uses one or more 'early access' classes");
+		const FString ExcerptName = bExperimental ? TEXT("ActorUsesExperimentalClass") : TEXT("ActorUsesEarlyAccessClass");
+		const FSlateBrush* WarningIcon = FEditorStyle::GetBrush(bExperimental ? "PropertyEditor.ExperimentalClass" : "PropertyEditor.EarlyAccessClass");
+
+		IDetailCategoryBuilder& WarningCategory = DetailBuilder.EditCategory(CategoryName, CategoryDisplayName, ECategoryPriority::Transform);
 
+		FDetailWidgetRow& WarningRow = WarningCategory.AddCustomRow(SearchString);
+
+		WarningCategory.AddCustomRow(SearchString)
+		.WholeRowContent()
+		[
+			SNew(SHorizontalBox)
+			.ToolTip(IDocumentation::Get()->CreateToolTip(Tooltip, nullptr, TEXT("Shared/LevelEditor"), ExcerptName))
+			+ SHorizontalBox::Slot()
+			.VAlign(VAlign_Center)
+			.AutoWidth()
+			.Padding(4.0f, 0.0f, 0.0f, 0.0f)
+			[
+				SNew(SImage)
+				.Image(WarningIcon)
+			]
+			+SHorizontalBox::Slot()
+			.VAlign(VAlign_Center)
+			.AutoWidth()
+			.Padding(4.0f, 0.0f, 0.0f, 0.0f)
+			[
+				SNew(STextBlock)
+				.Text(WarningText)
+				.Font(IDetailLayoutBuilder::GetDetailFont())
+			]
+		];
+	}
 }
 
 void FActorDetails::AddMaterialCategory( IDetailLayoutBuilder& DetailBuilder )
diff --git a/Engine/Source/Editor/DetailCustomizations/Private/ActorDetails.h b/Engine/Source/Editor/DetailCustomizations/Private/ActorDetails.h
--- a/Engine/Source/Editor/DetailCustomizations/Private/ActorDetails.h
+++ b/Engine/Source/Editor/DetailCustomizations/Private/ActorDetails.h
@@ -73,6 +73,8 @@ private:
 
 	void AddTransformCategory( IDetailLayoutBuilder& DetailBuilder );
 
+	void AddExperimentalWarningCategory( IDetailLayoutBuilder& DetailBuilder );
+
 	/**
 	 * Creates a category for displaying materials used by the selected actors
 	 */
diff --git a/Engine/Source/Editor/EditorStyle/Private/SlateEditorStyle.cpp b/Engine/Source/Editor/EditorStyle/Private/SlateEditorStyle.cpp
--- a/Engine/Source/Editor/EditorStyle/Private/SlateEditorStyle.cpp
+++ b/Engine/Source/Editor/EditorStyle/Private/SlateEditorStyle.cpp
@@ -3673,6 +3673,10 @@ void FSlateEditorStyle::FStyle::SetupLevelEditorStyle()
 			.SetCheckedHoveredImage(BOX_BRUSH(*SmallRoundedButtonEnd, FMargin(7.f / 16.f), SelectionColor))
 			.SetCheckedPressedImage(BOX_BRUSH(*SmallRoundedButtonEnd, FMargin(7.f / 16.f), SelectionColor))
 			.SetCheckedImage(BOX_BRUSH(*SmallRoundedButtonEnd, FMargin(7.f / 16.f), SelectionColor)));
+
+		// Experimental/early access stuff
+		Set("PropertyEditor.ExperimentalClass", new IMAGE_BRUSH("/PropertyView/ExperimentalClassWarning", Icon40x40));
+		Set("PropertyEditor.EarlyAccessClass", new IMAGE_BRUSH("/PropertyView/EarlyAccessClassWarning", Icon40x40));
 	}
 
 	// Mesh Paint
diff --git a/Engine/Source/Editor/UnrealEd/Private/AssetSelection.cpp b/Engine/Source/Editor/UnrealEd/Private/AssetSelection.cpp
--- a/Engine/Source/Editor/UnrealEd/Private/AssetSelection.cpp
+++ b/Engine/Source/Editor/UnrealEd/Private/AssetSelection.cpp
@@ -27,6 +27,10 @@ extern void OnPlaceStaticMeshActor( AActor* MeshActor, bool bUseSurfaceOrientati
 
 namespace AssetSelectionUtils
 {
+	const FName DevelopmentStatusKey(TEXT("DevelopmentStatus"));
+	const FString EarlyAccessValue(TEXT("EarlyAccess"));
+	const FString ExperimentalValue(TEXT("Experimental"));
+
 	bool IsClassPlaceable(const UClass* Class)
 	{
 		const bool bIsAddable =
@@ -35,7 +39,7 @@ namespace AssetSelectionUtils
 			&& Class->IsChildOf( AActor::StaticClass() );
 		return bIsAddable;
 	}
-
+	
 	void GetSelectedAssets( TArray<FAssetData>& OutSelectedAssets )
 	{
 		// Add the selection from the content browser module
@@ -188,17 +192,36 @@ namespace AssetSelectionUtils
 						ActorInfo.bHaveAttachedActor = true;
 					}
 
-					TArray<UStaticMeshComponent*> StaticMeshComponents;
-					CurrentActor->GetComponents(StaticMeshComponents);
+					TArray<UActorComponent*> ActorComponents;
+					CurrentActor->GetComponents(ActorComponents);
 
-					for( int32 CompIndex = 0; CompIndex < StaticMeshComponents.Num(); ++CompIndex )
+					for( UActorComponent* Component : ActorComponents )
 					{
-						UStaticMeshComponent* SMComp = StaticMeshComponents[ CompIndex ];
-						if( SMComp->IsRegistered() )
+						if( UStaticMeshComponent* SMComp = Cast<UStaticMeshComponent>(Component) )
 						{
-							ActorInfo.bHaveStaticMeshComponent = true;
+							if( SMComp->IsRegistered() )
+							{
+								ActorInfo.bHaveStaticMeshComponent = true;
+							}
 						}
 
+						// Check for experimental classes in the component hierarchy
+						FString DevelopmentStatus;
+						if (Component->GetClass()->GetStringMetaDataHierarchical(DevelopmentStatusKey, /*out*/ &DevelopmentStatus))
+						{
+							ActorInfo.bHaveExperimentalClass = ActorInfo.bHaveExperimentalClass || (DevelopmentStatus == ExperimentalValue);
+							ActorInfo.bHaveEarlyAccessClass = ActorInfo.bHaveEarlyAccessClass || (DevelopmentStatus == EarlyAccessValue);
+						}
+					}
+
+					// Check for experimental classes in the actor hierarchy
+					{
+						FString DevelopmentStatus;
+						if (CurrentClass->GetStringMetaDataHierarchical(DevelopmentStatusKey, /*out*/ &DevelopmentStatus))
+						{
+							ActorInfo.bHaveExperimentalClass = ActorInfo.bHaveExperimentalClass || (DevelopmentStatus == ExperimentalValue);
+							ActorInfo.bHaveEarlyAccessClass = ActorInfo.bHaveEarlyAccessClass || (DevelopmentStatus == EarlyAccessValue);
+						}
 					}
 
 					if( CurrentActor->IsA( ALight::StaticClass() ) )
@@ -658,6 +681,8 @@ bool FActorFactoryAssetProxy::IsActorValidForMaterialApplication( AActor* Target
 {
 	bool bIsValid = false;
 
+	//@TODO: PAPER2D: Extend this to support non mesh components (or make sprites a mesh component)
+
 	// Check if the actor has a mesh or fog volume density. If so, it can likely have
 	// a material applied to it. Otherwise, it cannot.
 	if ( TargetActor )
diff --git a/Engine/Source/Editor/UnrealEd/Public/AssetSelection.h b/Engine/Source/Editor/UnrealEd/Public/AssetSelection.h
--- a/Engine/Source/Editor/UnrealEd/Public/AssetSelection.h
+++ b/Engine/Source/Editor/UnrealEd/Public/AssetSelection.h
@@ -104,6 +104,10 @@ struct FSelectedActorInfo
 	uint32 bHaveHidden:1;
 	/** true if a landscape is in the selection */
 	uint32 bHaveLandscape:1;
+	/** true if an experimental actor (or actor containing such a component) is selected */
+	uint32 bHaveExperimentalClass:1;
+	/** true if an early access actor (or actor containing such a component) is selected */
+	uint32 bHaveEarlyAccessClass:1;
 
 	FSelectedActorInfo()
 		: SelectionClass(NULL)
@@ -135,6 +139,8 @@ struct FSelectedActorInfo
 		, bHaveMatinee(false)
 		, bHaveHidden(false)
 		, bHaveLandscape(false)
+		, bHaveExperimentalClass(false)
+		, bHaveEarlyAccessClass(false)
 	{
 	}
 
diff --git a/Engine/Source/Programs/UnrealHeaderTool/Private/BaseParser.cpp b/Engine/Source/Programs/UnrealHeaderTool/Private/BaseParser.cpp
--- a/Engine/Source/Programs/UnrealHeaderTool/Private/BaseParser.cpp
+++ b/Engine/Source/Programs/UnrealHeaderTool/Private/BaseParser.cpp
@@ -104,6 +104,12 @@ namespace
 
 			FMetadataKeyword& Category = Dictionary.Add(TEXT("Category"), EMetadataValueArgument::Required);
 			Category.InsertAddAction(TEXT("Category"), TEXT(""));
+
+			FMetadataKeyword& ExperimentalFeature = Dictionary.Add(TEXT("Experimental"), EMetadataValueArgument::None);
+			ExperimentalFeature.InsertAddAction(TEXT("DevelopmentStatus"), TEXT("Experimental"));
+
+			FMetadataKeyword& EarlyAccessFeature = Dictionary.Add(TEXT("EarlyAccessPreview"), EMetadataValueArgument::None);
+			EarlyAccessFeature.InsertAddAction(TEXT("DevelopmentStatus"), TEXT("EarlyAccess"));
 		}
 
 		return Dictionary.Find(Keyword);
diff --git a/Engine/Source/Programs/UnrealHeaderTool/Private/HeaderParser.cpp b/Engine/Source/Programs/UnrealHeaderTool/Private/HeaderParser.cpp
--- a/Engine/Source/Programs/UnrealHeaderTool/Private/HeaderParser.cpp
+++ b/Engine/Source/Programs/UnrealHeaderTool/Private/HeaderParser.cpp
@@ -3698,18 +3698,22 @@ bool FHeaderParser::CompileDeclaration( FClasses& AllClasses, FToken& Token )
 		check( TopNest->NestType == NEST_Class || TopNest->NestType == NEST_Interface );
 		CurrentAccessSpecifier = AccessSpecifier;
 	}
-	else if (Token.Matches(TEXT("class")) && TopNest->NestType == NEST_GlobalScope)
+	else if (Token.Matches(TEXT("class")) && (TopNest->NestType == NEST_GlobalScope))
 	{
 		if (!bHaveSeenFirstInterfaceClass || bFinishedParsingInterfaceClasses)
+		{
 			return SkipDeclaration(Token);
+		}
 
 		// Make sure the previous class ended with valid nesting.
 		if (bEncounteredNewStyleClass_UnmatchedBrackets)
-			FError::Throwf(TEXT("Missing } at end of class") );
+		{
+			FError::Throwf(TEXT("Missing } at end of class"));
+		}
 
 		// Start parsing the second class
 		bEncounteredNewStyleClass_UnmatchedBrackets = true;
-		bHaveSeenSecondInterfaceClass               = true;
+		bHaveSeenSecondInterfaceClass = true;
 		ParseSecondInterfaceClass(AllClasses);
 	}
 	else if (Token.Matches(TEXT("GENERATED_UCLASS_BODY")))
@@ -4100,10 +4104,10 @@ void FHeaderParser::CompileClassDeclaration(FClasses& AllClasses)
 	TArray<FString> HideFunctions;
 	TArray<FString> AutoExpandCategories;
 	TArray<FString> AutoCollapseCategories;
-	Class->GetHideCategories        (HideCategories);
-	Class->GetShowCategories        (ShowSubCatgories);
-	Class->GetHideFunctions			(HideFunctions);
-	Class->GetAutoExpandCategories  (AutoExpandCategories);
+	Class->GetHideCategories(HideCategories);
+	Class->GetShowCategories(ShowSubCatgories);
+	Class->GetHideFunctions(HideFunctions);
+	Class->GetAutoExpandCategories(AutoExpandCategories);
 	Class->GetAutoCollapseCategories(AutoCollapseCategories);
 
 	// Class attributes.
@@ -4111,7 +4115,7 @@ void FHeaderParser::CompileClassDeclaration(FClasses& AllClasses)
 	check(ClassData);
 
 	// New-style UCLASS() syntax
-	TMap<FName, FString>       MetaData;
+	TMap<FName, FString> MetaData;
 
 	TArray<FPropertySpecifier> SpecifiersFound;
 	ReadSpecifierSetInsideMacro(SpecifiersFound, TEXT("Class"), MetaData);
@@ -4143,8 +4147,8 @@ void FHeaderParser::CompileClassDeclaration(FClasses& AllClasses)
 
 	// Process all of the class specifiers
 	TArray<FString> ClassGroupNames;
-	bool            bWithinSpecified         = false;
-	bool            bDeclaresConfigFile      = false;
+	bool bWithinSpecified = false;
+	bool bDeclaresConfigFile = false;
 	for (const FPropertySpecifier& PropSpecifier : SpecifiersFound)
 	{
 		const FString& Specifier = PropSpecifier.Key;
@@ -5580,6 +5584,15 @@ void FHeaderParser::ValidateMetaDataFormat(UField* Field, const FString& InKey,
 			}
 		}
 	}
+	else if (InKey == TEXT("DevelopmentStatus"))
+	{
+		const FString EarlyAccessValue(TEXT("EarlyAccess"));
+		const FString ExperimentalValue(TEXT("Experimental"));
+		if ((InValue != EarlyAccessValue) && (InValue != ExperimentalValue))
+		{
+			FError::Throwf(TEXT("'%s' metadata was '%s' but it must be %s or %s"), *InKey, *InValue, *ExperimentalValue, *EarlyAccessValue);
+		}
+	}
 }
 
 // Validates the metadata, then adds it to the class data
diff --git a/Engine/Source/Runtime/CoreUObject/Public/UObject/ObjectBase.h b/Engine/Source/Runtime/CoreUObject/Public/UObject/ObjectBase.h
--- a/Engine/Source/Runtime/CoreUObject/Public/UObject/ObjectBase.h
+++ b/Engine/Source/Runtime/CoreUObject/Public/UObject/ObjectBase.h
@@ -600,6 +600,12 @@ namespace UC
 
 		/// A root convert limits a sub-class to only be able to convert to child classes of the first root class going up the hierarchy.
 		ConversionRoot,
+
+		/// Marks this class as 'experimental' (a totally unsupported and undocumented prototype)
+		Experimental,
+
+		// Marks this class as an 'early access' preview (while not considered production-ready, it's a step beyond 'experimental' and is being provided as a preview of things to come)
+		EarlyAccessPreview,
 	};
 }
 
