UnrealHeaderTool: Improve error message when an unknown class type is used as a template parameter

diff --git a/Engine/Source/Programs/UnrealHeaderTool/Private/HeaderParser.cpp b/Engine/Source/Programs/UnrealHeaderTool/Private/HeaderParser.cpp
--- a/Engine/Source/Programs/UnrealHeaderTool/Private/HeaderParser.cpp
+++ b/Engine/Source/Programs/UnrealHeaderTool/Private/HeaderParser.cpp
@@ -2950,6 +2950,10 @@ bool FHeaderParser::GetVarType
 				if (GetIdentifier(InnerClass))
 				{
 					TempClass = AllClasses.FindScriptClass(InnerClass.Identifier);
+					if (TempClass == nullptr)
+					{
+						FError::Throwf(TEXT("Unrecognized type '%s' (in expression %s<%s>)"), InnerClass.Identifier, VarType.Identifier, InnerClass.Identifier);
+					}
 
 					if (bIsAutoweakPtrTemplate)
 					{
