UHT: Coding standards fixes

diff --git a/Engine/Source/Programs/UnrealHeaderTool/Private/Classes.cpp b/Engine/Source/Programs/UnrealHeaderTool/Private/Classes.cpp
--- a/Engine/Source/Programs/UnrealHeaderTool/Private/Classes.cpp
+++ b/Engine/Source/Programs/UnrealHeaderTool/Private/Classes.cpp
@@ -146,7 +146,7 @@ FClass* FClasses::FindClass(const TCHAR* ClassName) const
 	if (UObjectRedirector* RenamedClassRedirector = FindObject<UObjectRedirector>(ClassPackage, ClassName))
 		return (FClass*)CastChecked<UClass>(RenamedClassRedirector->DestinationObject);
 
-	return NULL;
+	return nullptr;
 }
 
 TArray<FClass*> FClasses::GetDerivedClasses(FClass* Parent) const
@@ -181,7 +181,9 @@ FClass* FClasses::FindScriptClassOrThrow(const FString& InClassName) const
 {
 	FString ErrorMsg;
 	if (FClass* Result = FindScriptClass(InClassName, ErrorMsg))
+	{
 		return Result;
+	}
 
 	FError::Throwf(*ErrorMsg);
 
@@ -196,7 +198,7 @@ FClass* FClasses::FindScriptClass(const FString& InClassName, FString& OutErrorM
 	if (FClass* FoundClass = FindClass(*ClassNameStripped))
 	{
 		// If the class was found with the stripped class name, verify that the correct prefix was used and throw an error otherwise
-		if( !ClassNameHasValidPrefix( InClassName, FoundClass ) )
+		if (!ClassNameHasValidPrefix(InClassName, FoundClass))
 		{
 			OutErrorMsg = FString::Printf(TEXT("Class '%s' has an incorrect prefix, expecting '%s'"), *InClassName, *FoundClass->GetNameWithPrefix());
 			return NULL;
