UBT will no longer use ini file to get the list of plugins to compile for programs. Removed 'ScriptPlugins' ini setting (replaced with plugin category path).

diff --git a/Engine/Plugins/Script/ScriptGeneratorPlugin/ScriptGeneratorPlugin.uplugin b/Engine/Plugins/Script/ScriptGeneratorPlugin/ScriptGeneratorPlugin.uplugin
--- a/Engine/Plugins/Script/ScriptGeneratorPlugin/ScriptGeneratorPlugin.uplugin
+++ b/Engine/Plugins/Script/ScriptGeneratorPlugin/ScriptGeneratorPlugin.uplugin
@@ -8,8 +8,9 @@
 	"CreatedByURL" : "http://epicgames.com",
 	"EngineVersion" : "4.2.0",
 	"Description" : "An example of a script generator plugin.  This can be used as a starting point when creating your own plugin.",
-	"Category" : "Examples",
-	
+	"Category" : "UnrealHeaderTool.Code Generator.Generic",
+	"EnabledByDefault" : true,
+
 	"Modules" :
 	[
 		{
diff --git a/Engine/Programs/UnrealHeaderTool/Config/DefaultEngine.ini b/Engine/Programs/UnrealHeaderTool/Config/DefaultEngine.ini
--- a/Engine/Programs/UnrealHeaderTool/Config/DefaultEngine.ini
+++ b/Engine/Programs/UnrealHeaderTool/Config/DefaultEngine.ini
@@ -1,6 +1,5 @@
 [Plugins]
 +ProgramEnabledPlugins="ScriptGeneratorPlugin"
-+ScriptPlugins="ScriptGeneratorPlugin"
 
 +ScriptSupportedModules=CoreUObject
 +ScriptSupportedModules=Engine
diff --git a/Engine/Source/Programs/UnrealBuildTool/Configuration/UEBuildTarget.cs b/Engine/Source/Programs/UnrealBuildTool/Configuration/UEBuildTarget.cs
--- a/Engine/Source/Programs/UnrealBuildTool/Configuration/UEBuildTarget.cs
+++ b/Engine/Source/Programs/UnrealBuildTool/Configuration/UEBuildTarget.cs
@@ -2129,15 +2129,6 @@ namespace UnrealBuildTool
 					{
 						FilterPluginNames.AddRange(ValidPlugins.Where(x => x.bEnabledByDefault).Select(x => x.Name));
 					}
-					else if (Rules.Type == TargetRules.TargetType.Program)
-					{
-						var Config = new ConfigCacheIni(Platform, "Engine", Path.Combine(ProjectDirectory, "Programs", GetTargetName()));
-						List<string> ConfigPlugins;
-						if (Config.GetArray("Plugins", "ProgramEnabledPlugins", out ConfigPlugins))
-						{
-							FilterPluginNames.AddRange(ConfigPlugins);
-						}
-					}
 
 					// Update the plugin list for game targets
 					if(Rules.Type != TargetRules.TargetType.Program && UnrealBuildTool.HasUProjectFile())
diff --git a/Engine/Source/Programs/UnrealFrontend/UnrealFrontend.Target.cs b/Engine/Source/Programs/UnrealFrontend/UnrealFrontend.Target.cs
--- a/Engine/Source/Programs/UnrealFrontend/UnrealFrontend.Target.cs
+++ b/Engine/Source/Programs/UnrealFrontend/UnrealFrontend.Target.cs
@@ -8,6 +8,7 @@ public class UnrealFrontendTarget : TargetRules
 	public UnrealFrontendTarget( TargetInfo Target )
 	{
 		Type = TargetType.Program;
+		AdditionalPlugins.Add("UdpMessaging");
 	}
 
 	//
diff --git a/Engine/Source/Programs/UnrealHeaderTool/Private/CodeGenerator.cpp b/Engine/Source/Programs/UnrealHeaderTool/Private/CodeGenerator.cpp
--- a/Engine/Source/Programs/UnrealHeaderTool/Private/CodeGenerator.cpp
+++ b/Engine/Source/Programs/UnrealHeaderTool/Private/CodeGenerator.cpp
@@ -11,6 +11,7 @@
 #include "IScriptGeneratorPluginInterface.h"
 #include "Manifest.h"
 #include "StringUtils.h"
+#include "IPluginManager.h"
 
 /////////////////////////////////////////////////////
 // Globals
@@ -4368,14 +4369,18 @@ void FNativeClassHeaderGenerator::ExportGeneratedCPP()
 /** Get all script plugins based on ini setting */
 void GetScriptPlugins(TArray<IScriptGeneratorPluginInterface*>& ScriptPlugins)
 {
-	TArray<FString> ScriptPluginNames;
-	GConfig->GetArray(TEXT("Plugins"), TEXT("ScriptPlugins"), ScriptPluginNames, GEngineIni);
-	for (auto& PluginName : ScriptPluginNames)
+	const FString CodeGeneratorPluginCategory(TEXT("UnrealHeaderTool.Code Generator"));
+	const TArray<FPluginStatus> Plugins = IPluginManager::Get().QueryStatusForAllPlugins();
+	for (auto PluginIt(Plugins.CreateConstIterator()); PluginIt; ++PluginIt)
 	{
-		auto GeneratorInterface = FModuleManager::LoadModulePtr<IScriptGeneratorPluginInterface>(*PluginName);
-		if (GeneratorInterface)
+		const auto& PluginStatus = *PluginIt;
+		if (PluginStatus.bIsEnabled && PluginStatus.CategoryPath.StartsWith(CodeGeneratorPluginCategory))
 		{
-			ScriptPlugins.Add(GeneratorInterface);
+			auto GeneratorInterface = FModuleManager::LoadModulePtr<IScriptGeneratorPluginInterface>(*PluginStatus.Name);
+			if (GeneratorInterface)
+			{
+				ScriptPlugins.Add(GeneratorInterface);
+			}
 		}
 	}
 
diff --git a/Engine/Source/Programs/UnrealHeaderTool/UnrealHeaderTool.Target.cs b/Engine/Source/Programs/UnrealHeaderTool/UnrealHeaderTool.Target.cs
--- a/Engine/Source/Programs/UnrealHeaderTool/UnrealHeaderTool.Target.cs
+++ b/Engine/Source/Programs/UnrealHeaderTool/UnrealHeaderTool.Target.cs
@@ -8,6 +8,7 @@ public class UnrealHeaderToolTarget : TargetRules
 	public UnrealHeaderToolTarget(TargetInfo Target)
 	{
 		Type = TargetType.Program;
+		AdditionalPlugins.Add("ScriptGeneratorPlugin");
 	}
 
 	//
diff --git a/Engine/Source/Programs/UnrealLightmass/UnrealLightmass.Target.cs b/Engine/Source/Programs/UnrealLightmass/UnrealLightmass.Target.cs
--- a/Engine/Source/Programs/UnrealLightmass/UnrealLightmass.Target.cs
+++ b/Engine/Source/Programs/UnrealLightmass/UnrealLightmass.Target.cs
@@ -8,6 +8,7 @@ public class UnrealLightmassTarget : TargetRules
 	public UnrealLightmassTarget(TargetInfo Target)
 	{
 		Type = TargetType.Program;
+		AdditionalPlugins.Add("UdpMessaging");
 	}
 
 	//
