Hot-reload: support for structural changes for hot-reload (adding/removing properties/classes/structs/enums)

#change Classes will now be defer registered instead of being created during static initialization
#change Added class re-instancing after hot-reload
#change Changes to UHT to keep track of generated code changes
#change Refactored a few CoreUObject delegates to a common namespace struct.

diff --git a/Engine/Source/Editor/UnrealEd/Private/Editor.cpp b/Engine/Source/Editor/UnrealEd/Private/Editor.cpp
--- a/Engine/Source/Editor/UnrealEd/Private/Editor.cpp
+++ b/Engine/Source/Editor/UnrealEd/Private/Editor.cpp
@@ -422,8 +422,8 @@ void UEditorEngine::InitEditor(IEngineLoop* InEngineLoop)
 	FBSPOps::GFastRebuild = 0;
 
 	// Setup delegate callbacks for SavePackage()
-	GIsPackageOKToSaveDelegate.BindUObject( this, &UEditorEngine::IsPackageOKToSave );
-	GAutoPackageBackupDelegate.BindStatic( &FAutoPackageBackup::BackupPackage );
+	FCoreUObjectDelegates::IsPackageOKToSaveDelegate.BindUObject(this, &UEditorEngine::IsPackageOKToSave);
+	FCoreUObjectDelegates::AutoPackageBackupDelegate.BindStatic(&FAutoPackageBackup::BackupPackage);
 
 	// Update recents
 	UpdateRecentlyLoadedProjectFiles();
diff --git a/Engine/Source/Editor/UnrealEd/Private/Kismet2/KismetReinstanceUtilities.cpp b/Engine/Source/Editor/UnrealEd/Private/Kismet2/KismetReinstanceUtilities.cpp
--- a/Engine/Source/Editor/UnrealEd/Private/Kismet2/KismetReinstanceUtilities.cpp
+++ b/Engine/Source/Editor/UnrealEd/Private/Kismet2/KismetReinstanceUtilities.cpp
@@ -10,6 +10,42 @@
 /////////////////////////////////////////////////////////////////////////////////
 // FBlueprintCompileReinstancer
 
+FBlueprintCompileReinstancer::FBlueprintCompileReinstancer(UClass* InNewClass, UClass* InOldClass)
+: ClassToReinstance(InNewClass)
+, DuplicatedClass(InOldClass)
+, OriginalCDO(InOldClass->GetDefaultObject())
+, bHasReinstanced(false)
+, bSkipGarbageCollection(false)
+{
+	SaveClassFieldMapping(InOldClass);
+
+	TArray<UClass*> ChildrenOfClass;
+	GetDerivedClasses(InOldClass, ChildrenOfClass);
+	for (auto ClassIt = ChildrenOfClass.CreateConstIterator(); ClassIt; ++ClassIt)
+	{
+		UClass* ChildClass = *ClassIt;
+		UBlueprint* ChildBP = Cast<UBlueprint>(ChildClass->ClassGeneratedBy);
+		if (ChildBP && !ChildBP->HasAnyFlags(RF_BeingRegenerated))
+		{
+			// If this is a direct child, change the parent and relink so the property chain is valid for reinstancing
+			if (!ChildBP->HasAnyFlags(RF_NeedLoad))
+			{
+				if (ChildClass->GetSuperClass() == InOldClass)
+				{
+					ReparentChild(ChildBP);
+				}
+
+				Children.AddUnique(ChildBP);
+			}
+			else
+			{
+				// If this is a child that caused the load of their parent, relink to the REINST class so that we can still serialize in the CDO, but do not add to later processing
+				ReparentChild(ChildClass);
+			}
+		}
+	}
+}
+
 FBlueprintCompileReinstancer::FBlueprintCompileReinstancer(UClass* InClassToReinstance, bool bIsBytecodeOnly, bool bSkipGC)
 	: ClassToReinstance(InClassToReinstance)
 	, DuplicatedClass(NULL)
@@ -19,7 +55,7 @@ FBlueprintCompileReinstancer::FBlueprintCompileReinstancer(UClass* InClassToRein
 {
 	if( InClassToReinstance != NULL )
 	{
-		SaveClassFieldMapping();
+		SaveClassFieldMapping(InClassToReinstance);
 
 		// Remember the initial CDO for the class being resinstanced
 		OriginalCDO = ClassToReinstance->GetDefaultObject();
@@ -97,16 +133,16 @@ FBlueprintCompileReinstancer::FBlueprintCompileReinstancer(UClass* InClassToRein
 	}
 }
 
-void FBlueprintCompileReinstancer::SaveClassFieldMapping()
+void FBlueprintCompileReinstancer::SaveClassFieldMapping(UClass* InClassToReinstance)
 {
-	check(ClassToReinstance);
+	check(InClassToReinstance);
 
-	for(UProperty* Prop = ClassToReinstance->PropertyLink; Prop && (Prop->GetOuter() == ClassToReinstance); Prop = Prop->PropertyLinkNext)
+	for (UProperty* Prop = InClassToReinstance->PropertyLink; Prop && (Prop->GetOuter() == InClassToReinstance); Prop = Prop->PropertyLinkNext)
 	{
 		PropertyMap.Add(Prop->GetFName(), Prop);
 	}
 
-	for(auto Function : TFieldRange<UFunction>(ClassToReinstance, EFieldIteratorFlags::ExcludeSuper))
+	for (auto Function : TFieldRange<UFunction>(InClassToReinstance, EFieldIteratorFlags::ExcludeSuper))
 	{
 		FunctionMap.Add(Function->GetFName(),Function);
 	}
diff --git a/Engine/Source/Editor/UnrealEd/Private/Tests/BlueprintAutomationTests.cpp b/Engine/Source/Editor/UnrealEd/Private/Tests/BlueprintAutomationTests.cpp
--- a/Engine/Source/Editor/UnrealEd/Private/Tests/BlueprintAutomationTests.cpp
+++ b/Engine/Source/Editor/UnrealEd/Private/Tests/BlueprintAutomationTests.cpp
@@ -261,7 +261,7 @@ public:
 	{
 		if (DontSavePackagesList.Num() == 0)
 		{
-			GIsPackageOKToSaveDelegate.BindStatic(&IsPackageOKToSave);
+			FCoreUObjectDelegates::IsPackageOKToSaveDelegate.BindStatic(&IsPackageOKToSave);
 		}
 		DontSavePackagesList.AddUnique(Package->GetFName());
 	}
diff --git a/Engine/Source/Editor/UnrealEd/Private/UnrealEd.cpp b/Engine/Source/Editor/UnrealEd/Private/UnrealEd.cpp
--- a/Engine/Source/Editor/UnrealEd/Private/UnrealEd.cpp
+++ b/Engine/Source/Editor/UnrealEd/Private/UnrealEd.cpp
@@ -22,6 +22,7 @@
 #include "AssetEditorManager.h"
 
 #include "EditorActorFolders.h"
+#include "Kismet2/KismetReinstanceUtilities.h"
 
 UUnrealEdEngine* GUnrealEd;
 
@@ -197,7 +198,22 @@ void EditorExit()
 
 }
 
-IMPLEMENT_MODULE( FDefaultModuleImpl, UnrealEd );
+/** Singleton used for class re-instancing during hot-reload */
+struct FUClassHotReloadReinstancer
+{
+	FUClassHotReloadReinstancer()
+	{
+		FCoreUObjectDelegates::ReplaceHotReloadClassDelegate.BindStatic(&FUClassHotReloadReinstancer::ReinstanceClass);
+	}
 
+	static void ReinstanceClass(UClass* OldClass, UClass* NewClass)
+	{
+		UE_LOG(LogUnrealEd, Display, TEXT("Re-instancing %s after hot-reload."), *NewClass->GetName());
+		FBlueprintCompileReinstancer ReinstanceHelper(NewClass, OldClass);
+		ReinstanceHelper.ReinstanceObjects();
+	}
+};
+static FUClassHotReloadReinstancer GHotReloadReinstancer;
 
 
+IMPLEMENT_MODULE( FDefaultModuleImpl, UnrealEd );
diff --git a/Engine/Source/Editor/UnrealEd/Public/Kismet2/KismetReinstanceUtilities.h b/Engine/Source/Editor/UnrealEd/Public/Kismet2/KismetReinstanceUtilities.h
--- a/Engine/Source/Editor/UnrealEd/Public/Kismet2/KismetReinstanceUtilities.h
+++ b/Engine/Source/Editor/UnrealEd/Public/Kismet2/KismetReinstanceUtilities.h
@@ -42,9 +42,11 @@ public:
 
 	/** Sets the reinstancer up to work on every object of the specified class */
 	FBlueprintCompileReinstancer(UClass* InClassToReinstance, bool bIsBytecodeOnly = false, bool bSkipGC = false);
+	/** Sets the reinstancer up to reinstance native classes */
+	FBlueprintCompileReinstancer(UClass* InNewClass, UClass* InOldClass);
 
 	/** Saves a mapping of field names to their UField equivalents, so we can remap any bytecode that references them later */
-	void SaveClassFieldMapping();
+	void SaveClassFieldMapping(UClass* InClassToReinstance);
 
 	/** Helper to gather mappings from the old class's fields to the new class's version */
 	void GenerateFieldMappings(TMap<UObject*, UObject*>& FieldMapping);
diff --git a/Engine/Source/Programs/UnrealHeaderTool/Private/ClassMaps.cpp b/Engine/Source/Programs/UnrealHeaderTool/Private/ClassMaps.cpp
--- a/Engine/Source/Programs/UnrealHeaderTool/Private/ClassMaps.cpp
+++ b/Engine/Source/Programs/UnrealHeaderTool/Private/ClassMaps.cpp
@@ -16,3 +16,4 @@ TSet<UClass*>                            GPublicClassSet;
 TSet<UClass*>                            GExportedClasses;
 TMap<UProperty*, FString>                GArrayDimensions;
 TMap<UPackage*,  const FManifestModule*> GPackageToManifestModuleMap;
+TMap<UField*, uint32>                    GGeneratedCodeCRCs;
\ No newline at end of file
diff --git a/Engine/Source/Programs/UnrealHeaderTool/Private/ClassMaps.h b/Engine/Source/Programs/UnrealHeaderTool/Private/ClassMaps.h
--- a/Engine/Source/Programs/UnrealHeaderTool/Private/ClassMaps.h
+++ b/Engine/Source/Programs/UnrealHeaderTool/Private/ClassMaps.h
@@ -37,3 +37,4 @@ extern TSet<UClass*>                            GExportedClasses;
 extern TSet<UClass*>                            GPublicClassSet;
 extern TMap<UProperty*, FString>                GArrayDimensions;
 extern TMap<UPackage*,  const FManifestModule*> GPackageToManifestModuleMap;
+extern TMap<UField*, uint32>                    GGeneratedCodeCRCs;
diff --git a/Engine/Source/Programs/UnrealHeaderTool/Private/CodeGenerator.cpp b/Engine/Source/Programs/UnrealHeaderTool/Private/CodeGenerator.cpp
--- a/Engine/Source/Programs/UnrealHeaderTool/Private/CodeGenerator.cpp
+++ b/Engine/Source/Programs/UnrealHeaderTool/Private/CodeGenerator.cpp
@@ -666,6 +666,7 @@ FString FNativeClassHeaderGenerator::GetSingletonName(FClass* Item, bool bRequir
 FString FNativeClassHeaderGenerator::PropertyNew(FString& Meta, UProperty* Prop, const FString& OuterString, const FString& PropMacro, const TCHAR* NameSuffix, const TCHAR* Spaces, const TCHAR* SourceStruct)
 {
 	FString ExtraArgs;
+	FString GeneratedCrc;
 
 	FString PropNameDep = Prop->GetName();
 	if (Prop->HasAllPropertyFlags(CPF_Deprecated))
@@ -696,12 +697,22 @@ FString FNativeClassHeaderGenerator::PropertyNew(FString& Meta, UProperty* Prop,
 		UScriptStruct* Struct = StructProperty->Struct;
 		check(Struct);
 		ExtraArgs = FString::Printf(TEXT(", %s"), *GetSingletonName(Struct));
+		auto StructCrc = GGeneratedCodeCRCs.Find(Struct);
+		if (StructCrc)
+		{
+			GeneratedCrc = FString::Printf(TEXT(" // %u"), *StructCrc);
+		}
 	}
 	else if (UByteProperty* ByteProperty = Cast<UByteProperty>(Prop))
 	{
 		if (ByteProperty->Enum)
 		{
 			ExtraArgs = FString::Printf(TEXT(", %s"), *GetSingletonName(ByteProperty->Enum));
+			auto EnumCrc = GGeneratedCodeCRCs.Find(ByteProperty->Enum);
+			if (EnumCrc)
+			{
+				GeneratedCrc = FString::Printf(TEXT(" // %u"), *EnumCrc);
+			}
 		}
 	}
 	else if (UBoolProperty* BoolProperty = Cast<UBoolProperty>(Prop))
@@ -738,10 +749,11 @@ FString FNativeClassHeaderGenerator::PropertyNew(FString& Meta, UProperty* Prop,
 	TheFlagAudit.Add(Prop, TEXT("PropertyFlags"), Prop->PropertyFlags);
 
 	FString Symbol = FString::Printf(TEXT("NewProp_%s%s"), *Prop->GetName(), NameSuffix);
-	FString Lines  = FString::Printf(TEXT("%sUProperty* %s = %s\r\n"), 
+	FString Lines  = FString::Printf(TEXT("%sUProperty* %s = %s%s\r\n"), 
 		Spaces, 
 		*Symbol, 
-		*Constructor);
+		*Constructor,
+		*GeneratedCrc);
 
 	if (Prop->ArrayDim != 1)
 	{
@@ -962,12 +974,6 @@ void FNativeClassHeaderGenerator::ExportGeneratedPackageInitCode(UPackage* Packa
 
 void FNativeClassHeaderGenerator::ExportNativeGeneratedInitCode(FClass* Class)
 {
-	// Emit the IMPLEMENT_CLASS macro to go in the generated cpp file (unless it's the temporary UHT class).
-	if (!Class->HasAnyClassFlags(CLASS_Temporary))
-	{
-		GeneratedPackageCPP.Logf(TEXT("    IMPLEMENT_CLASS(%s);\r\n"),NameLookupCPP.GetNameCPP(Class));
-	}
-
 	int32 MaxLinesPerCpp = 30000;
 
 	if ((GeneratedFunctionBodyTextSplit.Num() == 0) || (GeneratedFunctionBodyTextSplit[GeneratedFunctionBodyTextSplit.Num()-1].GetLineCount() > MaxLinesPerCpp))
@@ -988,49 +994,55 @@ void FNativeClassHeaderGenerator::ExportNativeGeneratedInitCode(FClass* Class)
 	{
 		UEnum* Enum = *It;
 
-		FString SingletonName(GetSingletonName(Enum));
-
-		CallSingletons.Logf(TEXT("            OuterClass->LinkChild(%s);\r\n"), *SingletonName);
+		const FString SingletonName(GetSingletonName(Enum));		
 
 		FString Extern = FString::Printf(TEXT("    %sclass UEnum* %s;\r\n"), *ApiString, *SingletonName);
 		SingletonNameToExternDecl.Add(SingletonName, Extern);
 		GeneratedFunctionDeclarations.Log(*Extern);
 			
-		GeneratedFunctionText.Logf(TEXT("    UEnum* %s\r\n"), *SingletonName);
-		GeneratedFunctionText.Logf(TEXT("    {\r\n"));
+		FStringOutputDevice GeneratedEnumRegisterFunctionText;
+
+		GeneratedEnumRegisterFunctionText.Logf(TEXT("    UEnum* %s\r\n"), *SingletonName);
+		GeneratedEnumRegisterFunctionText.Logf(TEXT("    {\r\n"));
 		// Enums can either have a UClass or UPackage as outer (if declared in non-UClass header).
 		if (Enum->GetOuter()->IsA(UStruct::StaticClass()))
 		{
-			GeneratedFunctionText.Logf(TEXT("        UClass* Outer=%s;\r\n"), *GetSingletonName(CastChecked<UStruct>(Enum->GetOuter())));
+			GeneratedEnumRegisterFunctionText.Logf(TEXT("        UClass* Outer=%s;\r\n"), *GetSingletonName(CastChecked<UStruct>(Enum->GetOuter())));
 		}
 		else
 		{
-			GeneratedFunctionText.Logf(TEXT("        UPackage* Outer=%s;\r\n"), *GetPackageSingletonName(CastChecked<UPackage>(Enum->GetOuter())));
+			GeneratedEnumRegisterFunctionText.Logf(TEXT("        UPackage* Outer=%s;\r\n"), *GetPackageSingletonName(CastChecked<UPackage>(Enum->GetOuter())));
 		}
-		GeneratedFunctionText.Logf(TEXT("        static UEnum* ReturnEnum = NULL;\r\n"));
-		GeneratedFunctionText.Logf(TEXT("        if (!ReturnEnum)\r\n"));
-		GeneratedFunctionText.Logf(TEXT("        {\r\n"));
+		GeneratedEnumRegisterFunctionText.Logf(TEXT("        static UEnum* ReturnEnum = NULL;\r\n"));
+		GeneratedEnumRegisterFunctionText.Logf(TEXT("        if (!ReturnEnum)\r\n"));
+		GeneratedEnumRegisterFunctionText.Logf(TEXT("        {\r\n"));
 
-		GeneratedFunctionText.Logf(TEXT("            ReturnEnum = new(Outer, TEXT(\"%s\"), RF_Public|RF_Transient|RF_Native) UEnum(FPostConstructInitializeProperties());\r\n"), *Enum->GetName());
-		GeneratedFunctionText.Logf(TEXT("            TArray<FName> EnumNames;\r\n"));
+		GeneratedEnumRegisterFunctionText.Logf(TEXT("            ReturnEnum = new(Outer, TEXT(\"%s\"), RF_Public|RF_Transient|RF_Native) UEnum(FPostConstructInitializeProperties());\r\n"), *Enum->GetName());
+		GeneratedEnumRegisterFunctionText.Logf(TEXT("            TArray<FName> EnumNames;\r\n"));
 		for (int32 Index = 0; Index < Enum->NumEnums(); Index++)
 		{
-			GeneratedFunctionText.Logf(TEXT("            EnumNames.Add(FName(TEXT(\"%s\")));\r\n"), *Enum->GetEnum(Index).ToString());
+			GeneratedEnumRegisterFunctionText.Logf(TEXT("            EnumNames.Add(FName(TEXT(\"%s\")));\r\n"), *Enum->GetEnum(Index).ToString());
 		}
-		GeneratedFunctionText.Logf(TEXT("            ReturnEnum->SetEnums(EnumNames, %s);\r\n"), Enum->IsNamespaceEnum() ? TEXT("true") : TEXT("false"));
+		GeneratedEnumRegisterFunctionText.Logf(TEXT("            ReturnEnum->SetEnums(EnumNames, %s);\r\n"), Enum->IsNamespaceEnum() ? TEXT("true") : TEXT("false"));
 
 		FString Meta = GetMetaDataCodeForObject(Enum, TEXT("ReturnEnum"), TEXT("            "));
 		if (Meta.Len())
 		{
-			GeneratedFunctionText.Logf(TEXT("#if WITH_METADATA\r\n"));
-			GeneratedFunctionText.Logf(TEXT("            UMetaData* MetaData = ReturnEnum->GetOutermost()->GetMetaData();\r\n"));
-			GeneratedFunctionText.Log(*Meta);
-			GeneratedFunctionText.Logf(TEXT("#endif\r\n"));
+			GeneratedEnumRegisterFunctionText.Logf(TEXT("#if WITH_METADATA\r\n"));
+			GeneratedEnumRegisterFunctionText.Logf(TEXT("            UMetaData* MetaData = ReturnEnum->GetOutermost()->GetMetaData();\r\n"));
+			GeneratedEnumRegisterFunctionText.Log(*Meta);
+			GeneratedEnumRegisterFunctionText.Logf(TEXT("#endif\r\n"));
 		}
 
-		GeneratedFunctionText.Logf(TEXT("        }\r\n"));
-		GeneratedFunctionText.Logf(TEXT("        return ReturnEnum;\r\n"));
-		GeneratedFunctionText.Logf(TEXT("    }\r\n"));
+		GeneratedEnumRegisterFunctionText.Logf(TEXT("        }\r\n"));
+		GeneratedEnumRegisterFunctionText.Logf(TEXT("        return ReturnEnum;\r\n"));
+		GeneratedEnumRegisterFunctionText.Logf(TEXT("    }\r\n"));
+
+		GeneratedFunctionText += GeneratedEnumRegisterFunctionText;
+
+		uint32 EnumCrc = FCrc::MemCrc32(*GeneratedEnumRegisterFunctionText, GeneratedEnumRegisterFunctionText.Len() * sizeof(TCHAR));		
+		GGeneratedCodeCRCs.Add(Enum, EnumCrc);
+		CallSingletons.Logf(TEXT("            OuterClass->LinkChild(%s); // %u\r\n"), *SingletonName, EnumCrc);
 	}
 
 	// structs
@@ -1039,15 +1051,16 @@ void FNativeClassHeaderGenerator::ExportNativeGeneratedInitCode(FClass* Class)
 		UScriptStruct* ScriptStruct = *It;
 		UStruct* BaseStruct = ScriptStruct->GetSuperStruct();
 
-		FString SingletonName(GetSingletonName(ScriptStruct));
-
-		CallSingletons.Logf(TEXT("            OuterClass->LinkChild(%s);\r\n"), *SingletonName);
+		const FString SingletonName(GetSingletonName(ScriptStruct));
+		
 		FString Extern = FString::Printf(TEXT("    %sclass UScriptStruct* %s;\r\n"), *ApiString, *SingletonName);
 		SingletonNameToExternDecl.Add(SingletonName, Extern);
 		GeneratedFunctionDeclarations.Log(*Extern);
 
-		GeneratedFunctionText.Logf(TEXT("    UScriptStruct* %s\r\n"), *SingletonName);
-		GeneratedFunctionText.Logf(TEXT("    {\r\n"));
+		FStringOutputDevice GeneratedStructRegisterFunctionText;
+
+		GeneratedStructRegisterFunctionText.Logf(TEXT("    UScriptStruct* %s\r\n"), *SingletonName);
+		GeneratedStructRegisterFunctionText.Logf(TEXT("    {\r\n"));
 
 		// if this is a no export struct, we will put a local struct here for offset determination
 		TArray<UScriptStruct*> Structs;
@@ -1055,21 +1068,21 @@ void FNativeClassHeaderGenerator::ExportNativeGeneratedInitCode(FClass* Class)
 		if (Structs.Num())
 		{
 			TGuardValue<bool> Guard(bIsExportingForOffsetDeterminationOnly,true);
-			ExportMirrorsForNoexportStructs(Structs, 8, GeneratedFunctionText);
+			ExportMirrorsForNoexportStructs(Structs, 8, GeneratedStructRegisterFunctionText);
 		}
 
 		// Structs can either have a UClass or UPackage as outer (if delcared in non-UClass header).
 		if (ScriptStruct->GetOuter()->IsA(UStruct::StaticClass()))
 		{
-			GeneratedFunctionText.Logf(TEXT("        UStruct* Outer=%s;\r\n"), *GetSingletonName(CastChecked<UStruct>(ScriptStruct->GetOuter())));
+			GeneratedStructRegisterFunctionText.Logf(TEXT("        UStruct* Outer=%s;\r\n"), *GetSingletonName(CastChecked<UStruct>(ScriptStruct->GetOuter())));
 		}
 		else
 		{
-			GeneratedFunctionText.Logf(TEXT("        UPackage* Outer=%s;\r\n"), *GetPackageSingletonName(CastChecked<UPackage>(ScriptStruct->GetOuter())));
+			GeneratedStructRegisterFunctionText.Logf(TEXT("        UPackage* Outer=%s;\r\n"), *GetPackageSingletonName(CastChecked<UPackage>(ScriptStruct->GetOuter())));
 		}
-		GeneratedFunctionText.Logf(TEXT("        static UScriptStruct* ReturnStruct = NULL;\r\n"));
-		GeneratedFunctionText.Logf(TEXT("        if (!ReturnStruct)\r\n"));
-		GeneratedFunctionText.Logf(TEXT("        {\r\n"));
+		GeneratedStructRegisterFunctionText.Logf(TEXT("        static UScriptStruct* ReturnStruct = NULL;\r\n"));
+		GeneratedStructRegisterFunctionText.Logf(TEXT("        if (!ReturnStruct)\r\n"));
+		GeneratedStructRegisterFunctionText.Logf(TEXT("        {\r\n"));
 		FString BaseStructString(TEXT("NULL"));
 		if (BaseStruct)
 		{
@@ -1089,7 +1102,7 @@ void FNativeClassHeaderGenerator::ExportNativeGeneratedInitCode(FClass* Class)
 			ExplicitSizeString = FString::Printf( TEXT(", sizeof(%s), ALIGNOF(%s)"), NameLookupCPP.GetNameCPP(ScriptStruct), NameLookupCPP.GetNameCPP(ScriptStruct) );
 		}
 
-		GeneratedFunctionText.Logf(TEXT("            ReturnStruct = new(Outer, TEXT(\"%s\"), RF_Public|RF_Transient|RF_Native) UScriptStruct(FPostConstructInitializeProperties(), %s, %s, EStructFlags(0x%08X)%s);\r\n"), 
+		GeneratedStructRegisterFunctionText.Logf(TEXT("            ReturnStruct = new(Outer, TEXT(\"%s\"), RF_Public|RF_Transient|RF_Native) UScriptStruct(FPostConstructInitializeProperties(), %s, %s, EStructFlags(0x%08X)%s);\r\n"),
 			*ScriptStruct->GetName(),
 			*BaseStructString,
 			*CppStructOpsString,
@@ -1105,21 +1118,28 @@ void FNativeClassHeaderGenerator::ExportNativeGeneratedInitCode(FClass* Class)
 		}
 		FString OuterString = FString(TEXT("ReturnStruct"));
 		FString Meta = GetMetaDataCodeForObject(ScriptStruct, *OuterString, TEXT("            "));
-		OutputProperties(Meta, GeneratedFunctionText, OuterString, Props, TEXT("            "));
-		GeneratedFunctionText.Logf(TEXT("            ReturnStruct->StaticLink();\r\n"));
+		OutputProperties(Meta, GeneratedStructRegisterFunctionText, OuterString, Props, TEXT("            "));
+		GeneratedStructRegisterFunctionText.Logf(TEXT("            ReturnStruct->StaticLink();\r\n"));
 
 		if (Meta.Len())
 		{
-			GeneratedFunctionText.Logf(TEXT("#if WITH_METADATA\r\n"));
-			GeneratedFunctionText.Logf(TEXT("            UMetaData* MetaData = ReturnStruct->GetOutermost()->GetMetaData();\r\n"));
-			GeneratedFunctionText.Log(*Meta);
-			GeneratedFunctionText.Logf(TEXT("#endif\r\n"));
+			GeneratedStructRegisterFunctionText.Logf(TEXT("#if WITH_METADATA\r\n"));
+			GeneratedStructRegisterFunctionText.Logf(TEXT("            UMetaData* MetaData = ReturnStruct->GetOutermost()->GetMetaData();\r\n"));
+			GeneratedStructRegisterFunctionText.Log(*Meta);
+			GeneratedStructRegisterFunctionText.Logf(TEXT("#endif\r\n"));
 		}
 
-		GeneratedFunctionText.Logf(TEXT("        }\r\n"));
-		GeneratedFunctionText.Logf(TEXT("        return ReturnStruct;\r\n"));
-		GeneratedFunctionText.Logf(TEXT("    }\r\n"));
+		GeneratedStructRegisterFunctionText.Logf(TEXT("        }\r\n"));
+		GeneratedStructRegisterFunctionText.Logf(TEXT("        return ReturnStruct;\r\n"));
+		GeneratedStructRegisterFunctionText.Logf(TEXT("    }\r\n"));
+		
+		uint32 StructCrc = FCrc::MemCrc32(*GeneratedStructRegisterFunctionText, GeneratedStructRegisterFunctionText.Len() * sizeof(TCHAR));
+		GGeneratedCodeCRCs.Add(ScriptStruct, StructCrc);
 
+		GeneratedFunctionText += GeneratedStructRegisterFunctionText;
+		GeneratedFunctionText.Logf(TEXT("    uint32 Get_%s_CRC() { return %uU; }\r\n"), *SingletonName.Replace(TEXT("()"), TEXT("")), StructCrc);
+
+		CallSingletons.Logf(TEXT("            OuterClass->LinkChild(%s); // %u\r\n"), *SingletonName, StructCrc);
 	}
 
 	TArray<UFunction*> FunctionsToExport;
@@ -1145,7 +1165,7 @@ void FNativeClassHeaderGenerator::ExportNativeGeneratedInitCode(FClass* Class)
 
 		const FFuncInfo& FunctionData = CompilerInfo->GetFunctionData();
 
-		FString SingletonName(GetSingletonName(Function));
+		const FString SingletonName(GetSingletonName(Function));
 
 		CallSingletons.Logf(TEXT("            OuterClass->LinkChild(%s);\r\n"), *SingletonName);
 		FString Extern = FString::Printf(TEXT("    %sclass UFunction* %s;\r\n"), *ApiString, *SingletonName);
@@ -1240,6 +1260,8 @@ void FNativeClassHeaderGenerator::ExportNativeGeneratedInitCode(FClass* Class)
 		GeneratedFunctionText.Logf(TEXT("    }\r\n"));
 	}
 
+	FStringOutputDevice GeneratedClassRegisterFunctionText;
+
 	// The class itself (unless it's a temporary UHT class).
 	if (!Class->HasAnyClassFlags(CLASS_Temporary))
 	{
@@ -1251,10 +1273,10 @@ void FNativeClassHeaderGenerator::ExportNativeGeneratedInitCode(FClass* Class)
 			SingletonNameToExternDecl.Add(SingletonNameNoRegister, Extern);
 			GeneratedFunctionDeclarations.Log(*Extern);
 
-			GeneratedFunctionText.Logf(TEXT("    UClass* %s\r\n"), *SingletonNameNoRegister);
-			GeneratedFunctionText.Logf(TEXT("    {\r\n"));
-			GeneratedFunctionText.Logf(TEXT("        return %s::StaticClass();\r\n"), NameLookupCPP.GetNameCPP(Class));
-			GeneratedFunctionText.Logf(TEXT("    }\r\n"));
+			GeneratedClassRegisterFunctionText.Logf(TEXT("    UClass* %s\r\n"), *SingletonNameNoRegister);
+			GeneratedClassRegisterFunctionText.Logf(TEXT("    {\r\n"));
+			GeneratedClassRegisterFunctionText.Logf(TEXT("        return %s::StaticClass();\r\n"), NameLookupCPP.GetNameCPP(Class));
+			GeneratedClassRegisterFunctionText.Logf(TEXT("    }\r\n"));
 
 
 		}
@@ -1265,28 +1287,28 @@ void FNativeClassHeaderGenerator::ExportNativeGeneratedInitCode(FClass* Class)
 		SingletonNameToExternDecl.Add(SingletonName, Extern);
 		GeneratedFunctionDeclarations.Log(*Extern);
 
-		GeneratedFunctionText.Logf(TEXT("    UClass* %s\r\n"), *SingletonName);
-		GeneratedFunctionText.Logf(TEXT("    {\r\n"));
-		GeneratedFunctionText.Logf(TEXT("        static UClass* OuterClass = NULL;\r\n"));
-		GeneratedFunctionText.Logf(TEXT("        if (!OuterClass)\r\n"));
-		GeneratedFunctionText.Logf(TEXT("        {\r\n"));
+		GeneratedClassRegisterFunctionText.Logf(TEXT("    UClass* %s\r\n"), *SingletonName);
+		GeneratedClassRegisterFunctionText.Logf(TEXT("    {\r\n"));
+		GeneratedClassRegisterFunctionText.Logf(TEXT("        static UClass* OuterClass = NULL;\r\n"));
+		GeneratedClassRegisterFunctionText.Logf(TEXT("        if (!OuterClass)\r\n"));
+		GeneratedClassRegisterFunctionText.Logf(TEXT("        {\r\n"));
 		if (Class->GetSuperClass() && Class->GetSuperClass() != Class)
 		{
-			GeneratedFunctionText.Logf(TEXT("            %s;\r\n"), *GetSingletonName(Class->GetSuperClass()));
+			GeneratedClassRegisterFunctionText.Logf(TEXT("            %s;\r\n"), *GetSingletonName(Class->GetSuperClass()));
 		}
-		GeneratedFunctionText.Logf(TEXT("            %s;\r\n"), *GetPackageSingletonName(CastChecked<UPackage>(Class->GetOutermost())));
-		GeneratedFunctionText.Logf(TEXT("            OuterClass = %s::StaticClass();\r\n"), NameLookupCPP.GetNameCPP(Class));
-		GeneratedFunctionText.Logf(TEXT("            UObjectForceRegistration(OuterClass);\r\n"));
+		GeneratedClassRegisterFunctionText.Logf(TEXT("            %s;\r\n"), *GetPackageSingletonName(CastChecked<UPackage>(Class->GetOutermost())));
+		GeneratedClassRegisterFunctionText.Logf(TEXT("            OuterClass = %s::StaticClass();\r\n"), NameLookupCPP.GetNameCPP(Class));
+		GeneratedClassRegisterFunctionText.Logf(TEXT("            UObjectForceRegistration(OuterClass);\r\n"));
 		uint32 Flags = Class->ClassFlags & CLASS_SaveInCompiledInClasses;
-		GeneratedFunctionText.Logf(TEXT("            OuterClass->ClassFlags |= 0x%08X;\r\n"), Flags);
+		GeneratedClassRegisterFunctionText.Logf(TEXT("            OuterClass->ClassFlags |= 0x%08X;\r\n"), Flags);
 		//if (Class->IsChildOf(UObject::StaticClass()))
 		//{
-		//	GeneratedFunctionText.Logf(TEXT("            OuterClass->ClassAddReferencedObjects = &%s::AddReferencedObjects;\r\n"), NameLookupCPP.GetNameCPP(Class));
+		//	GeneratedClassRegisterFunctionText.Logf(TEXT("            OuterClass->ClassAddReferencedObjects = &%s::AddReferencedObjects;\r\n"), NameLookupCPP.GetNameCPP(Class));
 		//}
 		TheFlagAudit.Add(Class, TEXT("ClassFlags"), Flags);
-		GeneratedFunctionText.Logf(TEXT("\r\n"));
-		GeneratedFunctionText.Log(CallSingletons);
-		GeneratedFunctionText.Logf(TEXT("\r\n"));
+		GeneratedClassRegisterFunctionText.Logf(TEXT("\r\n"));
+		GeneratedClassRegisterFunctionText.Log(CallSingletons);
+		GeneratedClassRegisterFunctionText.Logf(TEXT("\r\n"));
 
 		FString OuterString = FString(TEXT("OuterClass"));
 		FString Meta = GetMetaDataCodeForObject(Class, *OuterString, TEXT("            "));
@@ -1297,7 +1319,7 @@ void FNativeClassHeaderGenerator::ExportNativeGeneratedInitCode(FClass* Class)
 			{
 				Props.Add(*ItInner);
 			}
-			OutputProperties(Meta, GeneratedFunctionText, OuterString, Props, TEXT("            "));
+			OutputProperties(Meta, GeneratedClassRegisterFunctionText, OuterString, Props, TEXT("            "));
 		}
 		// function table
 		{
@@ -1313,15 +1335,15 @@ void FNativeClassHeaderGenerator::ExportNativeGeneratedInitCode(FClass* Class)
 			// Emit code to construct each UFunction and rebuild the function map at runtime
 			for (UFunction* Function : FunctionsInMap)
 			{
-				GeneratedFunctionText.Logf(TEXT("            OuterClass->AddFunctionToFunctionMap(%s);\r\n"), *GetSingletonName(Function));
+				GeneratedClassRegisterFunctionText.Logf(TEXT("            OuterClass->AddFunctionToFunctionMap(%s);\r\n"), *GetSingletonName(Function));
 			}
 		}
 
 		// class flags are handled by the intrinsic bootstrap code
-		//GeneratedFunctionText.Logf(TEXT("            OuterClass->ClassFlags = 0x%08X;\r\n"), Class->ClassFlags);
+		//GeneratedClassRegisterFunctionText.Logf(TEXT("            OuterClass->ClassFlags = 0x%08X;\r\n"), Class->ClassFlags);
 		if (Class->ClassConfigName != NAME_None)
 		{
-			GeneratedFunctionText.Logf(TEXT("            OuterClass->ClassConfigName = FName(TEXT(\"%s\"));\r\n"), *Class->ClassConfigName.ToString());
+			GeneratedClassRegisterFunctionText.Logf(TEXT("            OuterClass->ClassConfigName = FName(TEXT(\"%s\"));\r\n"), *Class->ClassConfigName.ToString());
 		}
 
 		for (auto& Inter : Class->Interfaces)
@@ -1332,7 +1354,7 @@ void FNativeClassHeaderGenerator::ExportNativeGeneratedInitCode(FClass* Class)
 			{
 				OffsetString = FString::Printf(TEXT("VTABLE_OFFSET(%s, %s)"), NameLookupCPP.GetNameCPP(Class), NameLookupCPP.GetNameCPP(Inter.Class, true));
 			}
-			GeneratedFunctionText.Logf(TEXT("            OuterClass->Interfaces.Add(FImplementedInterface(%s, %s, %s ));\r\n"), 
+			GeneratedClassRegisterFunctionText.Logf(TEXT("            OuterClass->Interfaces.Add(FImplementedInterface(%s, %s, %s ));\r\n"), 
 				*GetSingletonName(Inter.Class, false),
 				*OffsetString,
 				Inter.bImplementedByK2 ? TEXT("true") : TEXT("false")
@@ -1341,23 +1363,25 @@ void FNativeClassHeaderGenerator::ExportNativeGeneratedInitCode(FClass* Class)
 		if (Class->ClassGeneratedBy)
 		{
 			UE_LOG(LogCompile, Fatal, TEXT("For intrinsic and compiled-in classes, ClassGeneratedBy should always be NULL"));
-			GeneratedFunctionText.Logf(TEXT("            OuterClass->ClassGeneratedBy = %s;\r\n"), *GetSingletonName(CastChecked<UClass>(Class->ClassGeneratedBy), false));
+			GeneratedClassRegisterFunctionText.Logf(TEXT("            OuterClass->ClassGeneratedBy = %s;\r\n"), *GetSingletonName(CastChecked<UClass>(Class->ClassGeneratedBy), false));
 		}
 
-		GeneratedFunctionText.Logf(TEXT("            OuterClass->StaticLink();\r\n"));
+		GeneratedClassRegisterFunctionText.Logf(TEXT("            OuterClass->StaticLink();\r\n"));
 
 		if (Meta.Len())
 		{
-			GeneratedFunctionText.Logf(TEXT("#if WITH_METADATA\r\n"));
-			GeneratedFunctionText.Logf(TEXT("            UMetaData* MetaData = OuterClass->GetOutermost()->GetMetaData();\r\n"));
-			GeneratedFunctionText.Log(*Meta);
-			GeneratedFunctionText.Logf(TEXT("#endif\r\n"));
+			GeneratedClassRegisterFunctionText.Logf(TEXT("#if WITH_METADATA\r\n"));
+			GeneratedClassRegisterFunctionText.Logf(TEXT("            UMetaData* MetaData = OuterClass->GetOutermost()->GetMetaData();\r\n"));
+			GeneratedClassRegisterFunctionText.Log(*Meta);
+			GeneratedClassRegisterFunctionText.Logf(TEXT("#endif\r\n"));
 		}
 
-		GeneratedFunctionText.Logf(TEXT("        }\r\n"));
-		GeneratedFunctionText.Logf(TEXT("        check(OuterClass->GetClass());\r\n"));
-		GeneratedFunctionText.Logf(TEXT("        return OuterClass;\r\n"));
-		GeneratedFunctionText.Logf(TEXT("    }\r\n"));
+		GeneratedClassRegisterFunctionText.Logf(TEXT("        }\r\n"));
+		GeneratedClassRegisterFunctionText.Logf(TEXT("        check(OuterClass->GetClass());\r\n"));
+		GeneratedClassRegisterFunctionText.Logf(TEXT("        return OuterClass;\r\n"));
+		GeneratedClassRegisterFunctionText.Logf(TEXT("    }\r\n"));
+
+		GeneratedFunctionText += GeneratedClassRegisterFunctionText;
 	}
 	if (FriendText.Len())
 	{
@@ -1374,8 +1398,15 @@ void FNativeClassHeaderGenerator::ExportNativeGeneratedInitCode(FClass* Class)
 
 	// Skip temporary UHT classes.
 	if (!Class->HasAnyClassFlags(CLASS_Temporary))
-	{
-		GeneratedFunctionText.Logf(TEXT("    static FCompiledInDefer Z_CompiledInDefer_UClass_%s(%s);\r\n"), NameLookupCPP.GetNameCPP(Class), *SingletonName);
+	{	
+		auto ClassNameCPP = NameLookupCPP.GetNameCPP(Class);		
+		GeneratedFunctionText.Logf(TEXT("    static FCompiledInDefer Z_CompiledInDefer_UClass_%s(%s, TEXT(\"%s\"));\r\n"), 
+			ClassNameCPP, *SingletonName, ClassNameCPP);
+		
+		// Calculate generated class initialization code CRC so that we know when it changes after hot-reload
+		uint32 ClassCrc = FCrc::MemCrc32(*GeneratedClassRegisterFunctionText, GeneratedClassRegisterFunctionText.Len() * sizeof(TCHAR));
+		// Emit the IMPLEMENT_CLASS macro to go in the generated cpp file.
+		GeneratedPackageCPP.Logf(TEXT("    IMPLEMENT_CLASS(%s, %u);\r\n"), ClassNameCPP, ClassCrc);
 	}
 }
 
@@ -2274,45 +2305,51 @@ void FNativeClassHeaderGenerator::ExportGeneratedStructBodyMacros(const TArray<U
 			const FString Macroized = Macroize(*MacroName, *CombinedLine);
 			GeneratedHeaderText.Log(*Macroized);
 
+			const TCHAR* StructNameCPP = NameLookupCPP.GetNameCPP(Struct);
 			FString SingletonName = StaticConstructionString.Replace(TEXT("()"), TEXT("")); // function address
+			FString GetCRCName = FString::Printf(TEXT("Get_%s_CRC"), *SingletonName);
 
-			GeneratedPackageCPP.Logf(TEXT("class UScriptStruct* %s::StaticStruct()\r\n"), NameLookupCPP.GetNameCPP(Struct));
+			GeneratedPackageCPP.Logf(TEXT("class UScriptStruct* %s::StaticStruct()\r\n"), StructNameCPP);
 			GeneratedPackageCPP.Logf(TEXT("{\r\n"));
 			GeneratedPackageCPP.Logf(TEXT("    static class UScriptStruct* Singleton = NULL;\r\n"));
 			GeneratedPackageCPP.Logf(TEXT("    if (!Singleton)\r\n"));
 			GeneratedPackageCPP.Logf(TEXT("    {\r\n"));
 			GeneratedPackageCPP.Logf(TEXT("        extern %sclass UScriptStruct* %s;\r\n"), *FriendApiString, *StaticConstructionString);
-
+			GeneratedPackageCPP.Logf(TEXT("        extern %suint32 %s();\r\n"), *FriendApiString, *GetCRCName);
+			
 			// UStructs can have UClass or UPackage outer (if declared in non-UClass headers).
 			if (Struct->GetOuter()->IsA(UStruct::StaticClass()))
 			{
-				GeneratedPackageCPP.Logf(TEXT("        Singleton = GetStaticStruct(%s, %s::StaticClass(), TEXT(\"%s\"));\r\n"), *SingletonName, NameLookupCPP.GetNameCPP(CastChecked<UStruct>(Struct->GetOuter())), *Struct->GetName());
+				GeneratedPackageCPP.Logf(TEXT("        Singleton = GetStaticStruct(%s, %s::StaticClass(), TEXT(\"%s\"), sizeof(%s), %s());\r\n"), 
+					*SingletonName, NameLookupCPP.GetNameCPP(CastChecked<UStruct>(Struct->GetOuter())), *Struct->GetName(), StructNameCPP, *GetCRCName);
 			}
 			else
 			{
 				FString PackageSingletonName = GetPackageSingletonName(CastChecked<UPackage>(Struct->GetOuter()));
 				GeneratedPackageCPP.Logf(TEXT("        extern %sclass UPackage* %s;\r\n"), *FriendApiString, *PackageSingletonName);
-				GeneratedPackageCPP.Logf(TEXT("        Singleton = GetStaticStruct(%s, %s, TEXT(\"%s\"));\r\n"), *SingletonName, *PackageSingletonName, *Struct->GetName());
+				GeneratedPackageCPP.Logf(TEXT("        Singleton = GetStaticStruct(%s, %s, TEXT(\"%s\"), sizeof(%s), %s());\r\n"), 
+					*SingletonName, *PackageSingletonName, *Struct->GetName(), StructNameCPP, *GetCRCName);
 			}
 			
 			GeneratedPackageCPP.Logf(TEXT("    }\r\n"));
 			GeneratedPackageCPP.Logf(TEXT("    return Singleton;\r\n"));
 			GeneratedPackageCPP.Logf(TEXT("}\r\n"));
-			GeneratedPackageCPP.Logf(TEXT("static FCompiledInDeferStruct Z_CompiledInDeferStruct_UScriptStruct_%s(%s::StaticStruct);\r\n"), NameLookupCPP.GetNameCPP(Struct), NameLookupCPP.GetNameCPP(Struct));
+			GeneratedPackageCPP.Logf(TEXT("static FCompiledInDeferStruct Z_CompiledInDeferStruct_UScriptStruct_%s(%s::StaticStruct);\r\n"), 
+				StructNameCPP, StructNameCPP);
 
 			// Generate StaticRegisterNatives equivalent for structs without classes.
 			if (!Struct->GetOuter()->IsA(UStruct::StaticClass()))
 			{
 				const FString ShortPackageName = FPackageName::GetShortName(Struct->GetOuter()->GetName());
-				GeneratedPackageCPP.Logf(TEXT("static struct FScriptStruct_%s_StaticRegisterNatives%s\r\n"), *ShortPackageName, NameLookupCPP.GetNameCPP(Struct));
+				GeneratedPackageCPP.Logf(TEXT("static struct FScriptStruct_%s_StaticRegisterNatives%s\r\n"), *ShortPackageName, StructNameCPP);
 				GeneratedPackageCPP.Logf(TEXT("{\r\n"));
-				GeneratedPackageCPP.Logf(TEXT("\tFScriptStruct_%s_StaticRegisterNatives%s()\r\n"), *ShortPackageName, NameLookupCPP.GetNameCPP(Struct));
+				GeneratedPackageCPP.Logf(TEXT("\tFScriptStruct_%s_StaticRegisterNatives%s()\r\n"), *ShortPackageName, StructNameCPP);
 				GeneratedPackageCPP.Logf(TEXT("\t{\r\n"));
 
 				GeneratedPackageCPP.Logf( TEXT("\t\tUScriptStruct::DeferCppStructOps(FName(TEXT(\"%s\")),new UScriptStruct::TCppStructOps<%s%s>);\r\n"), *Struct->GetName(), Struct->GetPrefixCPP(), *Struct->GetName() );
 
 				GeneratedPackageCPP.Logf(TEXT("\t}\r\n"));
-				GeneratedPackageCPP.Logf(TEXT("} ScriptStruct_%s_StaticRegisterNatives%s;\r\n"), *ShortPackageName, NameLookupCPP.GetNameCPP(Struct));
+				GeneratedPackageCPP.Logf(TEXT("} ScriptStruct_%s_StaticRegisterNatives%s;\r\n"), *ShortPackageName, StructNameCPP);
 			}
 		}
 	}
diff --git a/Engine/Source/Programs/UnrealHeaderTool/Private/HeaderParser.cpp b/Engine/Source/Programs/UnrealHeaderTool/Private/HeaderParser.cpp
--- a/Engine/Source/Programs/UnrealHeaderTool/Private/HeaderParser.cpp
+++ b/Engine/Source/Programs/UnrealHeaderTool/Private/HeaderParser.cpp
@@ -6329,10 +6329,6 @@ ECompilationResult::Type FHeaderParser::ParseHeaders(FClasses& AllClasses, FHead
 		if (DependsOnClass->HasAnyClassFlags(CLASS_Parsed))
 			continue;
 
-		// Treat manual dependency on a base class as an error to detect bad habits early on.
-		if (Class->IsChildOf(DependsOnClass))
-			FError::Throwf(TEXT("%s is derived from %s - please remove the DependsOn"),*Class->GetName(),*DependsOnClass->GetName() );
-
 		// Check for circular dependency. If the DependsOnClass is dependent on the SubClass, there is one.
 		if (DependsOnClass != Class && AllClasses.IsDependentOn(DependsOnClass, Class))
 		{
diff --git a/Engine/Source/Runtime/CoreUObject/Private/UObject/Class.cpp b/Engine/Source/Runtime/CoreUObject/Private/UObject/Class.cpp
--- a/Engine/Source/Runtime/CoreUObject/Private/UObject/Class.cpp
+++ b/Engine/Source/Runtime/CoreUObject/Private/UObject/Class.cpp
@@ -3228,6 +3228,7 @@ UClass::UClass(const class FPostConstructInitializeProperties& PCIP, UClass* InB
 UClass::UClass
 (
 	EStaticConstructor,
+	FName			InName,
 	uint32			InSize,
 	uint32			InClassFlags,
 	EClassCastFlags	InClassCastFlags,
@@ -3250,7 +3251,6 @@ UClass::UClass
 ,	bCooked( false )
 {
 	// If you add properties here, please update the other constructors and PurgeClass()
-
 	*(const TCHAR**)&ClassConfigName = InConfigName;
 }
 
diff --git a/Engine/Source/Runtime/CoreUObject/Private/UObject/CoreNative.cpp b/Engine/Source/Runtime/CoreUObject/Private/UObject/CoreNative.cpp
--- a/Engine/Source/Runtime/CoreUObject/Private/UObject/CoreNative.cpp
+++ b/Engine/Source/Runtime/CoreUObject/Private/UObject/CoreNative.cpp
@@ -22,7 +22,7 @@ IMPLEMENT_MODULE( FCoreUObjectModule, CoreUObject );
 
 // if we are not using compiled in natives, we still need this as a base class for intrinsics
 #if !USE_COMPILED_IN_NATIVES
-IMPLEMENT_CLASS(UObject);
+IMPLEMENT_CLASS(UObject, 0);
 COREUOBJECT_API class UClass* Z_Construct_UClass_UObject();
 UClass* Z_Construct_UClass_UObject()
 {
diff --git a/Engine/Source/Runtime/CoreUObject/Private/UObject/Enum.cpp b/Engine/Source/Runtime/CoreUObject/Private/UObject/Enum.cpp
--- a/Engine/Source/Runtime/CoreUObject/Private/UObject/Enum.cpp
+++ b/Engine/Source/Runtime/CoreUObject/Private/UObject/Enum.cpp
@@ -23,7 +23,10 @@ void UEnum::Serialize( FArchive& Ar )
 	Ar << Names;
 	Ar << bIsNamespace;
 
-	AddNamesToMasterList();
+	if (Ar.IsLoading() || Ar.IsSaving())
+	{	
+		AddNamesToMasterList();
+	}
 }
 
 int32 UEnum::ResolveEnumerator(FArchive& Ar, int32 EnumeratorIndex) const
diff --git a/Engine/Source/Runtime/CoreUObject/Private/UObject/SavePackage.cpp b/Engine/Source/Runtime/CoreUObject/Private/UObject/SavePackage.cpp
--- a/Engine/Source/Runtime/CoreUObject/Private/UObject/SavePackage.cpp
+++ b/Engine/Source/Runtime/CoreUObject/Private/UObject/SavePackage.cpp
@@ -13,9 +13,6 @@ static const int32 MAX_MERGED_COMPRESSION_CHUNKSIZE = 1024 * 1024;
 
 static const FName WorldClassName = FName("World");
 
-// Delegates used by SavePackage()
-FIsPackageOKToSaveDelegate GIsPackageOKToSaveDelegate;
-FAutoPackageBackupDelegate GAutoPackageBackupDelegate;
 
 static bool HasDeprecatedOrPendingKillOuter(UObject* InObj, UPackage* InSavingPackage)
 {
@@ -2351,9 +2348,9 @@ bool UPackage::SavePackage( UPackage* InOuter, UObject* Base, EObjectFlags TopLe
 
 	#if WITH_EDITOR
 		// Attempt to create a backup of this package before it is saved, if applicable
-		if (GAutoPackageBackupDelegate.IsBound())
+		if (FCoreUObjectDelegates::AutoPackageBackupDelegate.IsBound())
 		{
-			GAutoPackageBackupDelegate.Execute(*InOuter);
+			FCoreUObjectDelegates::AutoPackageBackupDelegate.Execute(*InOuter);
 		}
 	#endif	// #if WITH_EDITOR
 
@@ -2402,9 +2399,9 @@ bool UPackage::SavePackage( UPackage* InOuter, UObject* Base, EObjectFlags TopLe
 		}
 
 		// Make sure package is allowed to be saved.
-		if( !TargetPlatform && GIsPackageOKToSaveDelegate.IsBound() )
+		if( !TargetPlatform && FCoreUObjectDelegates::IsPackageOKToSaveDelegate.IsBound() )
 		{
-			bool bIsOKToSave = GIsPackageOKToSaveDelegate.Execute( InOuter, Filename, Error );
+			bool bIsOKToSave = FCoreUObjectDelegates::IsPackageOKToSaveDelegate.Execute(InOuter, Filename, Error);
 			if (!bIsOKToSave)
 			{
 				if (!(SaveFlags & SAVE_NoError))
diff --git a/Engine/Source/Runtime/CoreUObject/Private/UObject/UObjectBase.cpp b/Engine/Source/Runtime/CoreUObject/Private/UObject/UObjectBase.cpp
--- a/Engine/Source/Runtime/CoreUObject/Private/UObject/UObjectBase.cpp
+++ b/Engine/Source/Runtime/CoreUObject/Private/UObject/UObjectBase.cpp
@@ -427,18 +427,71 @@ void UObjectCompiledInDeferStruct(class UScriptStruct *(*InRegister)())
 	GetDeferredCompiledInStructRegistration().Add(InRegister);
 }
 
-class UScriptStruct *GetStaticStruct(class UScriptStruct *(*InRegister)(), UObject* StructOuter, const TCHAR* StructName)
+#if !IS_MONOLITHIC
+struct FStructOrEnumCompiledInfo : FFieldCompiledInInfo
+{
+	FStructOrEnumCompiledInfo(SIZE_T InClassSize, uint32 InCrc)
+		: FFieldCompiledInInfo(InClassSize, InCrc)
+	{}
+	virtual UField* Register() const { return NULL; }
+};
+
+/** Registered struct info (including size and reflection info) */
+static TMap<FName, FStructOrEnumCompiledInfo*>& GetStructOrEnumGeneratedCodeInfo()
+{
+	static TMap<FName, FStructOrEnumCompiledInfo*> StructOrEnumCompiledInfoMap;
+	return StructOrEnumCompiledInfoMap;
+}
+#endif
+
+class UScriptStruct *GetStaticStruct(class UScriptStruct *(*InRegister)(), UObject* StructOuter, const TCHAR* StructName, SIZE_T Size, uint32 Crc)
 {
 #if !IS_MONOLITHIC
-	if (GIsHotReload)
+	// Try to get the existing info for this struct
+	const FName StructFName(StructName);
+	FStructOrEnumCompiledInfo* Info = nullptr;
+	auto ExistingInfo = GetStructOrEnumGeneratedCodeInfo().Find(StructFName);
+	if (!ExistingInfo)
+	{
+		// New struct
+		Info = new FStructOrEnumCompiledInfo(Size, Crc);
+		Info->bHasChanged = true;
+		GetStructOrEnumGeneratedCodeInfo().Add(StructFName, Info);
+	}
+	else
 	{
-		UScriptStruct* ReturnStruct = FindObject<UScriptStruct>(StructOuter, StructName);
-		if (ReturnStruct)
+		// Hot-relaoded struct, check if it has changes
+		Info = *ExistingInfo;
+		Info->bHasChanged = (*ExistingInfo)->Size != Size || (*ExistingInfo)->Crc != Crc;
+		Info->Size = Size;
+		Info->Crc = Crc;
+	}
+
+	if (GIsHotReload)
+	{		
+		if (!Info->bHasChanged)
 		{
-			UE_LOG(LogClass, Log, TEXT( "%s HotReload."), StructName);
-			return ReturnStruct;
+			// New struct added during hot-reload
+			UScriptStruct* ReturnStruct = FindObject<UScriptStruct>(StructOuter, StructName);
+			if (ReturnStruct)
+			{
+				UE_LOG(LogClass, Log, TEXT( "%s HotReload."), StructName);
+				return ReturnStruct;
+			}
+			UE_LOG(LogClass, Log, TEXT("Could not find existing script struct %s for HotReload. Assuming new"), StructName);
+		}
+		else
+		{
+			// Existing struct, make sure we destroy the old one
+			UScriptStruct* ExistingStruct = FindObject<UScriptStruct>(StructOuter, StructName);
+			if (ExistingStruct)
+			{
+				// Make sure the old struct is not used by anything
+				ExistingStruct->ClearFlags(RF_RootSet | RF_Standalone | RF_Public);
+				const FName OldStructRename = MakeUniqueObjectName(GetTransientPackage(), ExistingStruct->GetClass(), *FString::Printf(TEXT("HOTRELOADED_%s"), StructName));
+				ExistingStruct->Rename(*OldStructRename.ToString(), GetTransientPackage());
+			}
 		}
-		UE_LOG(LogClass, Log, TEXT("Could not find existing script struct %s for HotReload. Assuming new"), StructName);
 	}
 #endif
 	return (*InRegister)();
@@ -481,10 +534,84 @@ static TArray<class UClass *(*)()>& GetDeferredCompiledInRegistration()
 	return DeferredCompiledInRegistration;
 }
 
-void UObjectCompiledInDefer(class UClass *(*InRegister)())
+/** Classes loaded with a module, deferred until we register them all in one go */
+static TArray<FFieldCompiledInInfo*>& GetDeferredClassRegistration()
 {
+	static TArray<FFieldCompiledInInfo*> DeferredClassRegistration;
+	return DeferredClassRegistration;
+}
+
 #if !IS_MONOLITHIC
-	if (!GIsHotReload) // in hot reload, we don't reregister anything
+/** Map of deferred class registration info (including size and reflection info) */
+static TMap<FName, FFieldCompiledInInfo*>& GetDeferRegisterClassMap()
+{
+	static TMap<FName, FFieldCompiledInInfo*> DeferRegisterClassMap;
+	return DeferRegisterClassMap;
+}
+
+/** Classes that changed during hot-reload and need to be re-instanced */
+static TArray<FFieldCompiledInInfo*>& GetHotReloadClasses()
+{
+	static TArray<FFieldCompiledInInfo*> HotReloadClasses;
+	return HotReloadClasses;
+}
+#endif
+
+void UClassCompiledInDefer(FFieldCompiledInInfo* ClassInfo, const TCHAR* Name, SIZE_T ClassSize, uint32 Crc)
+{
+	const FName CPPClassName(Name);
+#if !IS_MONOLITHIC
+	// Check for existing classes
+	auto ExistingClassInfo = GetDeferRegisterClassMap().Find(CPPClassName);
+	ClassInfo->bHasChanged = !ExistingClassInfo || (*ExistingClassInfo)->Size != ClassInfo->Size || (*ExistingClassInfo)->Crc != ClassInfo->Crc;
+	if (ExistingClassInfo)
+	{
+		// Class exists, this can only happen during hot-reload
+		check(GIsHotReload);
+		if (ClassInfo->bHasChanged)
+		{
+			// Get the native name
+			FString NameWithoutPrefix(Name);
+			NameWithoutPrefix = NameWithoutPrefix.Mid(1);
+
+			// Rename the old class and move it to transient package
+			auto ExistingClass = FindObjectChecked<UClass>(ANY_PACKAGE, *NameWithoutPrefix);
+			ExistingClass->ClearFlags(RF_RootSet | RF_Standalone | RF_Public);
+			ExistingClass->GetDefaultObject()->ClearFlags(RF_RootSet | RF_Standalone | RF_Public);
+			const FName OldClassRename = MakeUniqueObjectName(GetTransientPackage(), ExistingClass->GetClass(), *FString::Printf(TEXT("HOTRELOADED_%s"), *NameWithoutPrefix));
+			ExistingClass->Rename(*OldClassRename.ToString(), GetTransientPackage());
+			ExistingClass->AddToRoot();
+			ClassInfo->OldField = ExistingClass;
+
+			// Make sure enums de-register their names BEFORE we create the new class, otherwise there will be name conflicts
+			TArray<UObject*> ClassSubobjects;
+			GetObjectsWithOuter(ExistingClass, ClassSubobjects);
+			for (auto ClassSubobject : ClassSubobjects)
+			{
+				if (auto Enum = Cast<UEnum>(ClassSubobject))
+				{
+					Enum->RemoveNamesFromMasterList();
+				}
+			}
+
+			GetHotReloadClasses().Add(ClassInfo);
+		}
+		*ExistingClassInfo = ClassInfo;
+	}
+	else
+	{
+		GetDeferRegisterClassMap().Add(CPPClassName, ClassInfo);
+	}
+#endif
+	// We will either create a new class or update the static class pointer of the existing one
+	GetDeferredClassRegistration().Add(ClassInfo);
+}
+
+void UObjectCompiledInDefer(class UClass *(*InRegister)(), const TCHAR* Name)
+{
+#if !IS_MONOLITHIC
+	// Either add all classes if not hot-reloading, or those which have changed
+	if (!GIsHotReload || GetDeferRegisterClassMap().FindChecked(Name)->bHasChanged)
 #endif
 	{
 		checkSlow(!GetDeferredCompiledInRegistration().Contains(InRegister));
@@ -492,6 +619,31 @@ void UObjectCompiledInDefer(class UClass *(*InRegister)())
 	}
 }
 
+/** Register all loaded classes */
+void UClassRegisterAllCompiledInClasses()
+{
+	for (auto Class : GetDeferredClassRegistration())
+	{
+		Class->Register();
+	}
+	GetDeferredClassRegistration().Empty();
+}
+
+#if !IS_MONOLITHIC
+/** Re-instance all existing classes that have changed during hot-reload */
+void UClassReplaceHotReloadClasses()
+{
+	if (FCoreUObjectDelegates::ReplaceHotReloadClassDelegate.IsBound())
+	{	
+		for (auto Class : GetHotReloadClasses())
+		{
+			FCoreUObjectDelegates::ReplaceHotReloadClassDelegate.Execute(CastChecked<UClass>(Class->OldField), CastChecked<UClass>(Class->Register()));
+		}
+	}
+	GetHotReloadClasses().Empty();
+}
+#endif
+
 /**
  * Load any outstanding compiled in default properties
  */
@@ -576,12 +728,17 @@ bool AnyNewlyLoadedUObjects()
 
 void ProcessNewlyLoadedUObjects()
 {
+	UClassRegisterAllCompiledInClasses();
+
 	while( AnyNewlyLoadedUObjects() )
 	{
 		UObjectProcessRegistrants();
 		UObjectLoadAllCompiledInStructs();
 		UObjectLoadAllCompiledInDefaultProperties();		
 	}
+#if !IS_MONOLITHIC
+	UClassReplaceHotReloadClasses();
+#endif
 }
 
 
diff --git a/Engine/Source/Runtime/CoreUObject/Private/UObject/UObjectGlobals.cpp b/Engine/Source/Runtime/CoreUObject/Private/UObject/UObjectGlobals.cpp
--- a/Engine/Source/Runtime/CoreUObject/Private/UObject/UObjectGlobals.cpp
+++ b/Engine/Source/Runtime/CoreUObject/Private/UObject/UObjectGlobals.cpp
@@ -49,6 +49,12 @@ DEFINE_STAT(STAT_NameTableWideEntries);
 DEFINE_STAT(STAT_NameTableMemorySize);
 DEFINE_STAT(STAT_DestroyObject);
 
+/** CoreUObject delegates */
+FCoreUObjectDelegates::FReplaceHotReloadClassDelegate FCoreUObjectDelegates::ReplaceHotReloadClassDelegate;
+// Delegates used by SavePackage()
+FCoreUObjectDelegates::FIsPackageOKToSaveDelegate FCoreUObjectDelegates::IsPackageOKToSaveDelegate;
+FCoreUObjectDelegates::FAutoPackageBackupDelegate FCoreUObjectDelegates::AutoPackageBackupDelegate;
+
 /**
  * Fast version of StaticFindObject that relies on the passed in FName being the object name
  * without any group/ package qualifiers.
diff --git a/Engine/Source/Runtime/CoreUObject/Public/UObject/Class.h b/Engine/Source/Runtime/CoreUObject/Public/UObject/Class.h
--- a/Engine/Source/Runtime/CoreUObject/Public/UObject/Class.h
+++ b/Engine/Source/Runtime/CoreUObject/Public/UObject/Class.h
@@ -1302,12 +1302,15 @@ protected:
 	COREUOBJECT_API static TMap<FName, UEnum*> AllEnumNames;
 
 protected: 
+	
 	/** adds the Names in this enum to the master AllEnumNames list */
 	void AddNamesToMasterList();
+
+public:
+
 	/** removes the Names in this enum from the master AllEnumNames list */
 	void RemoveNamesFromMasterList();
 
-public:
 	// UObject interface.
 	COREUOBJECT_API virtual void Serialize(FArchive& Ar) override;
 	// End of UObject interface.
@@ -1805,7 +1808,7 @@ public:
 	// Constructors
 	UClass(const class FPostConstructInitializeProperties& PCIP);
 	explicit UClass(const class FPostConstructInitializeProperties& PCIP, UClass* InSuperClass);
-	UClass( EStaticConstructor, uint32 InSize, uint32 InClassFlags, EClassCastFlags InClassCastFlags,
+	UClass( EStaticConstructor, FName InName, uint32 InSize, uint32 InClassFlags, EClassCastFlags InClassCastFlags,
 		const TCHAR* InClassConfigName, EObjectFlags InFlags, void(*InClassConstructor)(const class FPostConstructInitializeProperties&),
 		void(*InClassAddReferencedObjects)(UObject*, class FReferenceCollector&));
 
@@ -2185,7 +2188,7 @@ void GetPrivateStaticClassBody( const TCHAR* PackageName, const TCHAR* Name, UCl
 			UE_LOG(LogClass, Log, TEXT("Could not find existing package %s for HotReload."),PackageName);
 			return;
 		}
-		ReturnClass = FindObjectChecked<UClass>((UObject *)Package, Name);
+		ReturnClass = FindObject<UClass>((UObject *)Package, Name);
 		if (ReturnClass)
 		{
 			if (ReturnClass->HotReloadPrivateStaticClass(
@@ -2215,6 +2218,7 @@ void GetPrivateStaticClassBody( const TCHAR* PackageName, const TCHAR* Name, UCl
 		UClass
 		(
 		EC_StaticConstructor,
+		Name,
 		sizeof(TClass),
 		TClass::StaticClassFlags,
 		TClass::StaticClassCastFlags(),
diff --git a/Engine/Source/Runtime/CoreUObject/Public/UObject/CoreObject.h b/Engine/Source/Runtime/CoreUObject/Public/UObject/CoreObject.h
--- a/Engine/Source/Runtime/CoreUObject/Public/UObject/CoreObject.h
+++ b/Engine/Source/Runtime/CoreUObject/Public/UObject/CoreObject.h
@@ -394,12 +394,6 @@ public:
 
 };
 
-/** Delegate type for saving check */
-DECLARE_DELEGATE_RetVal_ThreeParams(bool, FIsPackageOKToSaveDelegate, UPackage*, const FString&, FOutputDevice*);
-
-/** Delegate used by SavePackage() to check whether a package should be saved */
-extern COREUOBJECT_API FIsPackageOKToSaveDelegate GIsPackageOKToSaveDelegate;
-
 /*-----------------------------------------------------------------------------
 	UMetaData.
 -----------------------------------------------------------------------------*/
diff --git a/Engine/Source/Runtime/CoreUObject/Public/UObject/ObjectBase.h b/Engine/Source/Runtime/CoreUObject/Public/UObject/ObjectBase.h
--- a/Engine/Source/Runtime/CoreUObject/Public/UObject/ObjectBase.h
+++ b/Engine/Source/Runtime/CoreUObject/Public/UObject/ObjectBase.h
@@ -1173,8 +1173,8 @@ public: \
 	TWithinClass* GetOuter##TWithinClass() const { return (TWithinClass*)GetOuter(); }
 
 // Register a class at startup time.
-#define IMPLEMENT_CLASS(TClass) \
-	static UClass* AutoInitialize##TClass = TClass::StaticClass(); \
+#define IMPLEMENT_CLASS(TClass, TClassCrc) \
+	static TClassCompiledInDefer<TClass> AutoInitialize##TClass(TEXT(#TClass), sizeof(TClass), TClassCrc); \
 	UClass* TClass::GetPrivateStaticClass(const TCHAR* Package) \
 	{ \
 		static UClass* PrivateStaticClass = NULL; \
@@ -1193,7 +1193,7 @@ public: \
 
 // Used for intrinsics, this sets up the boiler plate, plus an initialization singleton, which can create properties and GC tokens
 #define IMPLEMENT_INTRINSIC_CLASS(TClass, TRequiredAPI, TSuperClass, TSuperRequiredAPI, InitCode) \
-	IMPLEMENT_CLASS(TClass) \
+	IMPLEMENT_CLASS(TClass, 0) \
 	TRequiredAPI UClass* Z_Construct_UClass_##TClass(); \
 	UClass* Z_Construct_UClass_##TClass() \
 	{ \
@@ -1211,7 +1211,7 @@ public: \
 		check(Class->GetClass()); \
 		return Class; \
 	} \
-	static FCompiledInDefer Z_CompiledInDefer_UClass_##TClass(Z_Construct_UClass_##TClass);
+	static FCompiledInDefer Z_CompiledInDefer_UClass_##TClass(Z_Construct_UClass_##TClass, TEXT(#TClass));
 
 #define IMPLEMENT_CORE_INTRINSIC_CLASS(TClass, TSuperClass, InitCode) \
 	IMPLEMENT_INTRINSIC_CLASS(TClass, COREUOBJECT_API, TSuperClass, COREUOBJECT_API, InitCode)
diff --git a/Engine/Source/Runtime/CoreUObject/Public/UObject/UObjectBase.h b/Engine/Source/Runtime/CoreUObject/Public/UObject/UObjectBase.h
--- a/Engine/Source/Runtime/CoreUObject/Public/UObject/UObjectBase.h
+++ b/Engine/Source/Runtime/CoreUObject/Public/UObject/UObjectBase.h
@@ -248,16 +248,69 @@ FORCEINLINE bool UObjectInitialized() { return Internal::GObjInitialized; }
  */
 COREUOBJECT_API void UObjectForceRegistration(UObjectBase* Object);
 
+/**
+ * Registers all recently loaded native classes
+ */
+COREUOBJECT_API void UClassRegisterAllCompiledInClasses();
+
+/** 
+ * Base class for deferred native class registration
+ */
+struct FFieldCompiledInInfo
+{
+	FFieldCompiledInInfo(SIZE_T InClassSize, uint32 InCrc)
+		: Size(InClassSize)
+		, Crc(InCrc)
+		, bHasChanged(false)
+		, OldField(NULL)
+	{
+	}
+
+	/** Registers the native class (constructs a UClass object) */
+	virtual UField* Register() const = 0;
+
+	/** Size of the class */
+	SIZE_T Size;
+	/** CRC of the generated code for this class */
+	uint32 Crc;
+	/** Old UClass object */
+	UField* OldField;
+	/** True if this class has changed after hot-reload (or new class) */
+	bool bHasChanged;
+};
+
+/**
+* Adds a class to deferred registration queue.
+*/
+COREUOBJECT_API void UClassCompiledInDefer(FFieldCompiledInInfo* Class, const TCHAR* Name, SIZE_T ClassSize, uint32 Crc);
+
+/**
+ * Specialized version of the deferred class registration structure.
+ */
+template <typename TClass>
+struct TClassCompiledInDefer : public FFieldCompiledInInfo
+{
+	TClassCompiledInDefer(const TCHAR* InName, SIZE_T InClassSize, uint32 InCrc)
+	: FFieldCompiledInInfo(InClassSize, InCrc)
+	{
+		UClassCompiledInDefer(this, InName, InClassSize, InCrc);
+	}
+	virtual UField* Register() const override
+	{
+		return TClass::StaticClass();
+	}
+};
+
 /**
  * Stashes the singleton function that builds a compiled in class. Later, this is executed.
  */
-COREUOBJECT_API void UObjectCompiledInDefer(class UClass *(*InRegister)());
+COREUOBJECT_API void UObjectCompiledInDefer(class UClass *(*InRegister)(), const TCHAR* Name);
 
 struct FCompiledInDefer
 {
-	FCompiledInDefer(class UClass *(*InRegister)())
+	FCompiledInDefer(class UClass *(*InRegister)(), const TCHAR* Name)
 	{
-		UObjectCompiledInDefer(InRegister);
+		UObjectCompiledInDefer(InRegister, Name);
 	}
 };
 
@@ -277,7 +330,7 @@ struct FCompiledInDeferStruct
 /**
  * Either call the passed in singleton, or if this is hot reload, find the existing struct
  */
-COREUOBJECT_API class UScriptStruct *GetStaticStruct(class UScriptStruct *(*InRegister)(), UObject* StructOuter, const TCHAR* StructName);
+COREUOBJECT_API class UScriptStruct *GetStaticStruct(class UScriptStruct *(*InRegister)(), UObject* StructOuter, const TCHAR* StructName, SIZE_T Size, uint32 Crc);
 
 /**
  * Stashes the singleton function that builds a compiled in enum. Later, this is executed.
diff --git a/Engine/Source/Runtime/CoreUObject/Public/UObject/UObjectGlobals.h b/Engine/Source/Runtime/CoreUObject/Public/UObject/UObjectGlobals.h
--- a/Engine/Source/Runtime/CoreUObject/Public/UObject/UObjectGlobals.h
+++ b/Engine/Source/Runtime/CoreUObject/Public/UObject/UObjectGlobals.h
@@ -1532,11 +1532,27 @@ protected:
 DECLARE_DELEGATE_RetVal_TwoParams( bool, FCheckForAutoAddDelegate, UPackage*, const FString& );
 DECLARE_DELEGATE_OneParam( FAddPackageToDefaultChangelistDelegate, const TCHAR* );
 
-/** Delegate type for making auto backup of package */
-DECLARE_DELEGATE_RetVal_OneParam( bool, FAutoPackageBackupDelegate, const UPackage& );
+/**
+ * Global CoreUObject delegates
+ */
+struct COREUOBJECT_API FCoreUObjectDelegates
+{
+	/** Delegate type for making auto backup of package */
+	DECLARE_DELEGATE_RetVal_OneParam(bool, FAutoPackageBackupDelegate, const UPackage&);
+
+	/** Delegate used by SavePackage() to create the package backup */
+	static FAutoPackageBackupDelegate AutoPackageBackupDelegate;
 
-/** Delegate used by SavePackage() to create the package backup */
-extern COREUOBJECT_API FAutoPackageBackupDelegate GAutoPackageBackupDelegate;
+	/** Delegate type for saving check */
+	DECLARE_DELEGATE_RetVal_ThreeParams(bool, FIsPackageOKToSaveDelegate, UPackage*, const FString&, FOutputDevice*);
+
+	/** Delegate used by SavePackage() to check whether a package should be saved */
+	static FIsPackageOKToSaveDelegate IsPackageOKToSaveDelegate;
+
+	DECLARE_DELEGATE_TwoParams(FReplaceHotReloadClassDelegate, UClass*, UClass*);
+
+	static FReplaceHotReloadClassDelegate ReplaceHotReloadClassDelegate;
+};
 
 /** Allows release builds to override not verifying GC assumptions. Useful for profiling as it's hitchy. */
 extern COREUOBJECT_API bool GShouldVerifyGCAssumptions;
diff --git a/Engine/Source/Runtime/Launch/Private/LaunchEngineLoop.cpp b/Engine/Source/Runtime/Launch/Private/LaunchEngineLoop.cpp
--- a/Engine/Source/Runtime/Launch/Private/LaunchEngineLoop.cpp
+++ b/Engine/Source/Runtime/Launch/Private/LaunchEngineLoop.cpp
@@ -618,6 +618,10 @@ int32 FEngineLoop::PreInit( const TCHAR* CmdLine )
 
 	// Switch into executable's directory.
 	FPlatformProcess::SetCurrentWorkingDirectoryToBaseDir();
+#if WITH_COREUOBJECT
+	// Register all classes that have been loaded so far. This is required for CVars to work.
+	UClassRegisterAllCompiledInClasses();
+#endif
 
 #if WITH_ENGINE
 	extern ENGINE_API void InitializeRenderingCVarsCaching();
