Linux: do not compile UHT monolithically anymore.

- Fixes it after being broken by a recent change, which requires all plugins to be loaded before proceeding.

diff --git a/Engine/Source/Programs/UnrealHeaderTool/UnrealHeaderTool.Target.cs b/Engine/Source/Programs/UnrealHeaderTool/UnrealHeaderTool.Target.cs
--- a/Engine/Source/Programs/UnrealHeaderTool/UnrealHeaderTool.Target.cs
+++ b/Engine/Source/Programs/UnrealHeaderTool/UnrealHeaderTool.Target.cs
@@ -29,7 +29,7 @@ public class UnrealHeaderToolTarget : TargetRules
 
 	public override bool ShouldCompileMonolithic(UnrealTargetPlatform InPlatform, UnrealTargetConfiguration InConfiguration)
 	{
-		if (UnrealBuildTool.UnrealBuildTool.CommandLineContains("-monolithic") == true || InPlatform == UnrealTargetPlatform.Linux)
+		if (UnrealBuildTool.UnrealBuildTool.CommandLineContains("-monolithic") == true)
 		{
 			return true;
 		}
