#UE4 Change blueprint access to static array compile error to be a compile display text. Warnings get converted to errors so we can't use warnings.

diff --git a/Engine/Source/Programs/UnrealHeaderTool/Private/HeaderParser.cpp b/Engine/Source/Programs/UnrealHeaderTool/Private/HeaderParser.cpp
--- a/Engine/Source/Programs/UnrealHeaderTool/Private/HeaderParser.cpp
+++ b/Engine/Source/Programs/UnrealHeaderTool/Private/HeaderParser.cpp
@@ -5555,7 +5555,7 @@ Found:
 			UProperty const* const Param = *It;
 			if (Param->ArrayDim > 1)
 			{
-				FError::Throwf(TEXT("Static array cannot be exposed to blueprint. Function: %s Parameter %s\n"), *TopFunction->GetName(), *Param->GetName());
+				UE_LOG(LogCompile, Display, TEXT("Static array cannot be exposed to blueprint. Function: %s Parameter %s\n"), *TopFunction->GetName(), *Param->GetName());
 			}
 		}
 	}
@@ -5904,7 +5904,7 @@ void FHeaderParser::CompileVariableDeclaration(FClasses& AllClasses, UStruct* St
 
 		if (NewProperty->HasAnyPropertyFlags(CPF_BlueprintVisible) && (NewProperty->ArrayDim > 1))
 		{
-			UE_LOG(LogCompile, Warning, TEXT("Static array cannot be exposed to blueprint %s.%s\n"), *Struct->GetName(), *NewProperty->GetName());
+			UE_LOG(LogCompile, Display, TEXT("Static array cannot be exposed to blueprint %s.%s\n"), *Struct->GetName(), *NewProperty->GetName());
 		}
 
 	} while( MatchSymbol(TEXT(",")) );
