[GitHub] 341 : Enabled class member initializer-awareness to UnrealHeaderTool

Integrated changes, added a test.

diff --git a/Engine/Source/Programs/UnrealHeaderTool/Private/HeaderParser.cpp b/Engine/Source/Programs/UnrealHeaderTool/Private/HeaderParser.cpp
--- a/Engine/Source/Programs/UnrealHeaderTool/Private/HeaderParser.cpp
+++ b/Engine/Source/Programs/UnrealHeaderTool/Private/HeaderParser.cpp
@@ -5909,6 +5909,22 @@ void FHeaderParser::CompileVariableDeclaration(FClasses& AllClasses, UStruct* St
 
 	} while( MatchSymbol(TEXT(",")) );
 
+	// Optional member initializer.
+	if (MatchSymbol(TEXT("=")))
+	{
+		// Skip past the specified member initializer; we make no attempt to parse it
+		FToken SkipToken;
+		while (GetToken(SkipToken))
+		{
+			if (SkipToken.Matches(TEXT(";")))
+			{
+				// went too far
+				UngetToken(SkipToken);
+				break;
+			}
+		}
+	}
+
 	// Expect a semicolon.
 	RequireSymbol( TEXT(";"), TEXT("'variable declaration'") );
 }
diff --git a/Engine/Source/Programs/UnrealHeaderTool/Resources/UHTDebugging/TestObject.h b/Engine/Source/Programs/UnrealHeaderTool/Resources/UHTDebugging/TestObject.h
--- a/Engine/Source/Programs/UnrealHeaderTool/Resources/UHTDebugging/TestObject.h
+++ b/Engine/Source/Programs/UnrealHeaderTool/Resources/UHTDebugging/TestObject.h
@@ -10,6 +10,9 @@ class UTestObject : public UObject
 	UFUNCTION(BlueprintCallable, Category="Random")
 	void TestForNullPtrDefaults(UObject* Obj1 = NULL, UObject* Obj2 = nullptr, UObject* Obj3 = 0);
 
+	UPROPERTY()
+	int32 Cpp11Init = 123;
+
 #if BLAH
 	UFUNCTION() int x; // This should not compile if UHT parses it, which it shouldn't
 #elif BLAH2
