enum class support for UENUMs.  Currently only supports enum classes of underlying type uint8, because the property system needs to be updated to handle more than that.

EPropertyReferenceFlags removed too, because it is unused.

#ttp 329394: UHT: Support the C++ "enum class" syntax for UENUMs

diff --git a/Engine/Source/Editor/LandscapeEditor/Classes/LandscapeEditorObject.h b/Engine/Source/Editor/LandscapeEditor/Classes/LandscapeEditorObject.h
--- a/Engine/Source/Editor/LandscapeEditor/Classes/LandscapeEditorObject.h
+++ b/Engine/Source/Editor/LandscapeEditor/Classes/LandscapeEditorObject.h
@@ -162,18 +162,15 @@ namespace ELandscapeImportHeightmapError
 }
 
 UENUM()
-namespace ELandscapeImportLayerError
+enum class ELandscapeImportLayerError : uint8
 {
-	enum Type
-	{
-		None,
-		MissingLayerInfo,
-		FileNotFound,
-		FileSizeMismatch,
-		CorruptFile,
-		ColorPng,
-	};
-}
+	None,
+	MissingLayerInfo,
+	FileNotFound,
+	FileSizeMismatch,
+	CorruptFile,
+	ColorPng,
+};
 
 USTRUCT()
 struct FLandscapeImportLayer : public FLandscapeImportLayerInfo
@@ -181,7 +178,7 @@ struct FLandscapeImportLayer : public FLandscapeImportLayerInfo
 	GENERATED_USTRUCT_BODY()
 
 	UPROPERTY(Category="Import", VisibleAnywhere)
-	TEnumAsByte<ELandscapeImportLayerError::Type> ImportError;
+	ELandscapeImportLayerError ImportError;
 
 	FLandscapeImportLayer()
 		: FLandscapeImportLayerInfo()
diff --git a/Engine/Source/Editor/UnrealEd/Private/Kismet2/EnumEditorUtils.cpp b/Engine/Source/Editor/UnrealEd/Private/Kismet2/EnumEditorUtils.cpp
--- a/Engine/Source/Editor/UnrealEd/Private/Kismet2/EnumEditorUtils.cpp
+++ b/Engine/Source/Editor/UnrealEd/Private/Kismet2/EnumEditorUtils.cpp
@@ -37,7 +37,7 @@ UEnum* FEnumEditorUtils::CreateUserDefinedEnum(UObject* InParent, FName EnumName
 	if (NULL != Enum)
 	{
 		TArray<FName> EmptyNames;
-		Enum->SetEnums(EmptyNames, true);
+		Enum->SetEnums(EmptyNames, UEnum::ECppForm::Namespaced);
 		Enum->SetMetaData(TEXT("BlueprintType"), TEXT("true"));
 	}
 
@@ -62,7 +62,7 @@ void FEnumEditorUtils::UpdateAfterPathChanged(UEnum* Enum)
 		NewEnumeratorsNames.Add(FName(*NewFullName));
 	}
 
-	Enum->SetEnums(NewEnumeratorsNames, true);
+	Enum->SetEnums(NewEnumeratorsNames, UEnum::ECppForm::Namespaced);
 }
 
 //////////////////////////////////////////////////////////////////////////
@@ -95,8 +95,8 @@ void FEnumEditorUtils::AddNewEnumeratorForUserDefinedEnum(UUserDefinedEnum* Enum
 		const FString FullNameStr = Enum->GenerateFullEnumName(*EnumNameString);
 		Names.Add(FName(*FullNameStr));
 
-		const bool bIsNamespace = Enum->IsNamespaceEnum();
-		Enum->SetEnums(Names, bIsNamespace);
+		const UEnum::ECppForm EnumType = Enum->GetCppForm();
+		Enum->SetEnums(Names, EnumType);
 		EnsureAllDisplayNamesExist(Enum);
 
 		BroadcastChanges(Enum, OldNames);
@@ -119,8 +119,8 @@ void FEnumEditorUtils::RemoveEnumeratorFromUserDefinedEnum(UUserDefinedEnum* Enu
 
 		Enum->RemoveMetaData(FEnumEditorUtilsHelper::DisplayName(), EnumeratorIndex);
 
-		const bool bIsNamespace = Enum->IsNamespaceEnum();
-		Enum->SetEnums(Names, bIsNamespace);
+		const UEnum::ECppForm EnumType = Enum->GetCppForm();
+		Enum->SetEnums(Names, EnumType);
 		EnsureAllDisplayNamesExist(Enum);
 		BroadcastChanges(Enum, OldNames);
 
@@ -151,8 +151,8 @@ void FEnumEditorUtils::MoveEnumeratorInUserDefinedEnum(UUserDefinedEnum* Enum, i
 			Names.Swap(EnumeratorIndex, EnumeratorIndex + 1);
 		}
 
-		const bool bIsNamespace = Enum->IsNamespaceEnum();
-		Enum->SetEnums(Names, bIsNamespace);
+		const UEnum::ECppForm EnumType = Enum->GetCppForm();
+		Enum->SetEnums(Names, EnumType);
 		EnsureAllDisplayNamesExist(Enum);
 		BroadcastChanges(Enum, OldNames);
 
diff --git a/Engine/Source/Programs/UnrealHeaderTool/Private/ClassMaps.cpp b/Engine/Source/Programs/UnrealHeaderTool/Private/ClassMaps.cpp
--- a/Engine/Source/Programs/UnrealHeaderTool/Private/ClassMaps.cpp
+++ b/Engine/Source/Programs/UnrealHeaderTool/Private/ClassMaps.cpp
@@ -16,4 +16,5 @@ TSet<UClass*>                            GPublicClassSet;
 TSet<UClass*>                            GExportedClasses;
 TMap<UProperty*, FString>                GArrayDimensions;
 TMap<UPackage*,  const FManifestModule*> GPackageToManifestModuleMap;
-TMap<UField*, uint32>                    GGeneratedCodeCRCs;
\ No newline at end of file
+TMap<UField*, uint32>                    GGeneratedCodeCRCs;
+TMap<UEnum*,  EPropertyType>             GEnumUnderlyingTypes;
diff --git a/Engine/Source/Programs/UnrealHeaderTool/Private/ClassMaps.h b/Engine/Source/Programs/UnrealHeaderTool/Private/ClassMaps.h
--- a/Engine/Source/Programs/UnrealHeaderTool/Private/ClassMaps.h
+++ b/Engine/Source/Programs/UnrealHeaderTool/Private/ClassMaps.h
@@ -38,3 +38,4 @@ extern TSet<UClass*>                            GPublicClassSet;
 extern TMap<UProperty*, FString>                GArrayDimensions;
 extern TMap<UPackage*,  const FManifestModule*> GPackageToManifestModuleMap;
 extern TMap<UField*, uint32>                    GGeneratedCodeCRCs;
+extern TMap<UEnum*,  EPropertyType>             GEnumUnderlyingTypes;
diff --git a/Engine/Source/Programs/UnrealHeaderTool/Private/CodeGenerator.cpp b/Engine/Source/Programs/UnrealHeaderTool/Private/CodeGenerator.cpp
--- a/Engine/Source/Programs/UnrealHeaderTool/Private/CodeGenerator.cpp
+++ b/Engine/Source/Programs/UnrealHeaderTool/Private/CodeGenerator.cpp
@@ -1023,7 +1023,15 @@ void FNativeClassHeaderGenerator::ExportNativeGeneratedInitCode(FClass* Class)
 		{
 			GeneratedEnumRegisterFunctionText.Logf(TEXT("            EnumNames.Add(FName(TEXT(\"%s\")));\r\n"), *Enum->GetEnum(Index).ToString());
 		}
-		GeneratedEnumRegisterFunctionText.Logf(TEXT("            ReturnEnum->SetEnums(EnumNames, %s);\r\n"), Enum->IsNamespaceEnum() ? TEXT("true") : TEXT("false"));
+
+		FString EnumTypeStr;
+		switch (Enum->GetCppForm())
+		{
+			case UEnum::ECppForm::Regular:    EnumTypeStr = TEXT("UEnum::ECppForm::Regular");    break;
+			case UEnum::ECppForm::Namespaced: EnumTypeStr = TEXT("UEnum::ECppForm::Namespaced"); break;
+			case UEnum::ECppForm::EnumClass:  EnumTypeStr = TEXT("UEnum::ECppForm::EnumClass");  break;
+		}
+		GeneratedEnumRegisterFunctionText.Logf(TEXT("            ReturnEnum->SetEnums(EnumNames, %s);\r\n"), *EnumTypeStr);
 
 		FString Meta = GetMetaDataCodeForObject(Enum, TEXT("ReturnEnum"), TEXT("            "));
 		if (Meta.Len())
@@ -2256,12 +2264,6 @@ void FNativeClassHeaderGenerator::ExportEnums( const TArray<UEnum*>& Enums )
 	{
 		UEnum* Enum = Enums[EnumIdx];
 
-		FString QualifierPrefix; 
-		if (!Enum->ActualEnumNameInsideNamespace.IsEmpty())
-		{
-			QualifierPrefix = Enum->GetName() + TEXT("::");
-		}
-
 		// Export FOREACH macro
 		EnumForeachText.Logf( TEXT("#define FOREACH_ENUM_%s(op) "), *Enum->GetName().ToUpper() );
 		for (int32 i = 0; i < Enum->NumEnums() - 1; i++)
@@ -3393,26 +3395,17 @@ void FNativeClassHeaderGenerator::ExportFunctionThunk(FStringOutputDevice& RPCWr
 			}
 		}
 
-		UByteProperty* ByteProp = Cast< UByteProperty >( Param );
-		if ((ByteProp != NULL) && (ByteProp->Enum != NULL))
+		UByteProperty* ByteProp = Cast<UByteProperty>(Param);
+		if (ByteProp && ByteProp->Enum)
 		{
-			UEnum* Enum = ByteProp->Enum;
-
-			FString FullyQualifiedEnumName = Enum->GetName();
-			if (!Enum->ActualEnumNameInsideNamespace.IsEmpty())
-			{
-				FullyQualifiedEnumName += TEXT("::");
-				FullyQualifiedEnumName += Enum->ActualEnumNameInsideNamespace;
-			}
-
 			// For enums, add an explicit conversion 
 			if (!(ByteProp->PropertyFlags & CPF_OutParm))
 			{
-				ParamName = FString::Printf( TEXT( "%s(%s)" ), *FullyQualifiedEnumName, *ParamName );
+				ParamName = FString::Printf( TEXT( "%s(%s)" ), *ByteProp->Enum->CppType, *ParamName );
 			}
 			else
 			{
-				ParamName = FString::Printf( TEXT( "(TEnumAsByte<%s>&)(%s)" ), *FullyQualifiedEnumName, *ParamName );
+				ParamName = FString::Printf( TEXT( "(TEnumAsByte<%s>&)(%s)" ), *ByteProp->Enum->CppType, *ParamName );
 			}
 		}
 
diff --git a/Engine/Source/Programs/UnrealHeaderTool/Private/HeaderParser.cpp b/Engine/Source/Programs/UnrealHeaderTool/Private/HeaderParser.cpp
--- a/Engine/Source/Programs/UnrealHeaderTool/Private/HeaderParser.cpp
+++ b/Engine/Source/Programs/UnrealHeaderTool/Private/HeaderParser.cpp
@@ -545,34 +545,55 @@ UEnum* FHeaderParser::CompileEnum(UClass* Scope)
 	TArray<FPropertySpecifier> SpecifiersFound;
 	ReadSpecifierSetInsideMacro(SpecifiersFound, TEXT("Enum"), EnumToken.MetaData);
 
-	// Check enum type. This can be global 'enum' or 'namespace' enum.
-	FToken EnumType;
-	if (!GetIdentifier(EnumType) || (!EnumType.Matches(TEXT("enum"), ESearchCase::CaseSensitive) && !EnumType.Matches(TEXT("namespace"), ESearchCase::CaseSensitive)))
-	{
-		FError::Throwf(TEXT("UENUM() should be followed by \'enum\' or \'namespace\' keywords.") );
-	}
-	bool bInNamespace = EnumType.Matches(TEXT("namespace"));
-
 	FScriptLocation DeclarationPosition;
 
-	// Get enumeration name.
+	// Check enum type. This can be global 'enum', 'namespace' or 'enum class' enums.
+	bool            bReadEnumName = false;
+	UEnum::ECppForm CppForm       = UEnum::ECppForm::Regular;
 	if (!GetIdentifier(EnumToken))
 	{
-		if (bInNamespace)
+		FError::Throwf(TEXT("Missing identifier after UENUM()") );
+	}
+
+	if (EnumToken.Matches(TEXT("namespace"), ESearchCase::CaseSensitive))
+	{
+		CppForm      = UEnum::ECppForm::Namespaced;
+		bReadEnumName = GetIdentifier(EnumToken);
+	}
+	else if (EnumToken.Matches(TEXT("enum"), ESearchCase::CaseSensitive))
+	{
+		if (!GetIdentifier(EnumToken))
 		{
-			FError::Throwf(TEXT("Missing enumeration namespace name") );
+			FError::Throwf(TEXT("Missing identifier after enum") );
+		}
+
+		if (EnumToken.Matches(TEXT("class"), ESearchCase::CaseSensitive) || EnumToken.Matches(TEXT("struct"), ESearchCase::CaseSensitive))
+		{
+			CppForm       = UEnum::ECppForm::EnumClass;
+			bReadEnumName = GetIdentifier(EnumToken);
 		}
 		else
 		{
-			FError::Throwf(TEXT("Missing enumeration name") );
+			CppForm       = UEnum::ECppForm::Regular;
+			bReadEnumName = true;
 		}
 	}
+	else
+	{
+		FError::Throwf(TEXT("UENUM() should be followed by \'enum\' or \'namespace\' keywords.") );
+	}
+
+	// Get enumeration name.
+	if (!bReadEnumName)
+	{
+		FError::Throwf(TEXT("Missing enumeration name") );
+	}
 
 	// Verify that the enumeration definition is unique within this scope.
-	UField* Existing = FindField( Scope, EnumToken.Identifier );
-	if ((Existing != NULL) && (Existing->GetOuter() == Scope))
+	UField* Existing = FindField(Scope, EnumToken.Identifier);
+	if (Existing && Existing->GetOuter() == Scope)
 	{
-		FError::Throwf(TEXT("enum: '%s' already defined here"), *EnumToken.TokenName.ToString() );
+		FError::Throwf(TEXT("enum: '%s' already defined here"), *EnumToken.TokenName.ToString());
 	}
 
 	CheckObscures(Scope, EnumToken.Identifier, EnumToken.Identifier);
@@ -586,25 +607,52 @@ UEnum* FHeaderParser::CompileEnum(UClass* Scope)
 	// Validate the metadata for the enum
 	ValidateMetaDataFormat(Enum, EnumToken.MetaData);
 
+	// Read optional base for enum class
+	if (CppForm == UEnum::ECppForm::EnumClass && MatchSymbol(TEXT(":")))
+	{
+		FToken BaseToken;
+		if (!GetIdentifier(BaseToken))
+		{
+			FError::Throwf(TEXT("Missing enum base") );
+		}
+
+		// We only support uint8 at the moment, until the properties get updated
+		if (FCString::Strcmp(BaseToken.Identifier, TEXT("uint8")))
+		{
+			FError::Throwf(TEXT("Only enum bases of type uint8 are currently supported"));
+		}
+
+		GEnumUnderlyingTypes.Add(Enum, CPT_Byte);
+	}
+
 	// Get opening brace.
 	RequireSymbol( TEXT("{"), TEXT("'Enum'") );
 
-	if (bInNamespace)
+	switch (CppForm)
 	{
-		// Now handle the inner true enum portion
-		RequireIdentifier(TEXT("enum"), TEXT("'Enum'"));
-
-		FToken InnerEnumToken;
-		if (!GetIdentifier(InnerEnumToken))
+		case UEnum::ECppForm::Namespaced:
 		{
-			FError::Throwf(TEXT("Missing enumeration name") );
+			// Now handle the inner true enum portion
+			RequireIdentifier(TEXT("enum"), TEXT("'Enum'"));
+
+			FToken InnerEnumToken;
+			if (!GetIdentifier(InnerEnumToken))
+			{
+				FError::Throwf(TEXT("Missing enumeration name") );
+			}
+
+			Enum->CppType = FString::Printf(TEXT("%s::%s"), EnumToken.Identifier, InnerEnumToken.Identifier);
+
+			RequireSymbol( TEXT("{"), TEXT("'Enum'") );
 		}
-		else
+		break;
+
+		case UEnum::ECppForm::Regular:
+		case UEnum::ECppForm::EnumClass:
 		{
-			Enum->ActualEnumNameInsideNamespace = InnerEnumToken.Identifier;
+			Enum->CppType = EnumToken.Identifier;
 		}
-
-		RequireSymbol( TEXT("{"), TEXT("'Enum'") );
+		break;
 	}
 
 	// List of all metadata generated for this enum
@@ -649,51 +697,58 @@ UEnum* FHeaderParser::CompileEnum(UClass* Scope)
 
 		int32 iFound;
 		FName NewTag;
-		if (bInNamespace)
+		switch (CppForm)
 		{
-			NewTag = FName(*FString::Printf(TEXT("%s::%s"), EnumToken.Identifier, TagToken.Identifier), FNAME_Add, true);
-		}
-		else
-		{
-			NewTag = FName(TagToken.Identifier, FNAME_Add, true);
+			case UEnum::ECppForm::Namespaced:
+			case UEnum::ECppForm::EnumClass:
+			{
+				NewTag = FName(*FString::Printf(TEXT("%s::%s"), EnumToken.Identifier, TagToken.Identifier), FNAME_Add, true);
+			}
+			break;
+
+			case UEnum::ECppForm::Regular:
+			{
+				NewTag = FName(TagToken.Identifier, FNAME_Add, true);
+			}
+			break;
 		}
+
 		if (EnumNames.Find(NewTag, iFound))
 		{
 			FError::Throwf(TEXT("Duplicate enumeration tag %s"), TagToken.Identifier );
 		}
-		else if (CurrentEnumValue > 255)
+
+		if (CurrentEnumValue > 255)
 		{
 			FError::Throwf(TEXT("Exceeded maximum of 255 enumerators") );
 		}
-		else
+
+		UEnum* FoundEnum = NULL;
+		if (UEnum::LookupEnumName(NewTag, &FoundEnum) != INDEX_NONE)
 		{
-			UEnum* FoundEnum = NULL;
-			if (UEnum::LookupEnumName(NewTag, &FoundEnum) != INDEX_NONE)
-			{
-				FError::Throwf(TEXT("Enumeration tag '%s' already in use by enum '%s'"), TagToken.Identifier, *FoundEnum->GetPathName());
-			}
+			FError::Throwf(TEXT("Enumeration tag '%s' already in use by enum '%s'"), TagToken.Identifier, *FoundEnum->GetPathName());
+		}
 
-			// Make sure the enum names array is tightly packed by inserting dummies
-			//@TODO: UCREMOVAL: Improve the UEnum system so we can have loosely packed values for e.g., bitfields
-			for (int32 DummyIndex = EnumNames.Num(); DummyIndex < CurrentEnumValue; ++DummyIndex)
-			{
-				FString DummyName              = FString::Printf(TEXT("UnusedSpacer_%d"), DummyIndex);
-				FString DummyNameWithNamespace = FString::Printf(TEXT("%s::%s"), EnumToken.Identifier, *DummyName);
-				EnumNames.Add(FName(*DummyNameWithNamespace));
+		// Make sure the enum names array is tightly packed by inserting dummies
+		//@TODO: UCREMOVAL: Improve the UEnum system so we can have loosely packed values for e.g., bitfields
+		for (int32 DummyIndex = EnumNames.Num(); DummyIndex < CurrentEnumValue; ++DummyIndex)
+		{
+			FString DummyName              = FString::Printf(TEXT("UnusedSpacer_%d"), DummyIndex);
+			FString DummyNameWithQualifier = FString::Printf(TEXT("%s::%s"), EnumToken.Identifier, *DummyName);
+			EnumNames.Add(FName(*DummyNameWithQualifier));
 
-				// These ternary operators are the correct way around, believe it or not.
-				// Spacers are qualified with the ETheEnum:: when they're not namespaced in order to prevent spacer name clashes.
-				// They're not qualified when they're actually in a namespace.
-				InsertMetaDataPair(EnumValueMetaData, (bInNamespace ? DummyName : DummyNameWithNamespace) + TEXT(".Hidden"), TEXT(""));
-				InsertMetaDataPair(EnumValueMetaData, (bInNamespace ? DummyName : DummyNameWithNamespace) + TEXT(".Spacer"), TEXT(""));
-			}
+			// These ternary operators are the correct way around, believe it or not.
+			// Spacers are qualified with the ETheEnum:: when they're regular enums in order to prevent spacer name clashes.
+			// They're not qualified when they're actually in a namespace or are enum classes.
+			InsertMetaDataPair(EnumValueMetaData, ((CppForm != UEnum::ECppForm::Regular) ? DummyName : DummyNameWithQualifier) + TEXT(".Hidden"), TEXT(""));
+			InsertMetaDataPair(EnumValueMetaData, ((CppForm != UEnum::ECppForm::Regular) ? DummyName : DummyNameWithQualifier) + TEXT(".Spacer"), TEXT(""));
+		}
 
-			// Save the new tag
-			EnumNames.Add( NewTag );
+		// Save the new tag
+		EnumNames.Add( NewTag );
 
-			// Autoincrement the current enumerant value
-			CurrentEnumValue++;
-		}
+		// Autoincrement the current enumerant value
+		CurrentEnumValue++;
 
 		// check for metadata on this enum value
 		ParseFieldMetaData(TagToken.MetaData, TagToken.Identifier);
@@ -735,27 +790,24 @@ UEnum* FHeaderParser::CompileEnum(UClass* Scope)
 	RequireSymbol( TEXT("}"), TEXT("'Enum'") );
 	MatchSemi();
 
-	if (bInNamespace)
+	if (CppForm == UEnum::ECppForm::Namespaced)
 	{
 		// Trailing brace for the namespace.
 		RequireSymbol( TEXT("}"), TEXT("'Enum'") );
 	}
 
 	// Register the list of enum names.
-	if (!Enum->SetEnums(EnumNames, bInNamespace))
+	if (!Enum->SetEnums(EnumNames, CppForm))
 	{
-		const FName MaxEnumItem = *(Enum->GenerateEnumPrefix() + TEXT("_MAX"));
+		const FName MaxEnumItem      = *(Enum->GenerateEnumPrefix() + TEXT("_MAX"));
 		const int32 MaxEnumItemIndex = Enum->FindEnumIndex(MaxEnumItem);
 		if (MaxEnumItemIndex != INDEX_NONE)
 		{
 			ReturnToLocation(EnumTagLocations[MaxEnumItemIndex], false, true);
 			FError::Throwf(TEXT("Illegal enumeration tag specified.  Conflicts with auto-generated tag '%s'"), *MaxEnumItem.ToString());
 		}
-		else
-		{
-			FError::Throwf(TEXT("Unable to generate enum MAX entry '%s' due to name collision"), *MaxEnumItem.ToString());
-		}
-		///return NULL;
+
+		FError::Throwf(TEXT("Unable to generate enum MAX entry '%s' due to name collision"), *MaxEnumItem.ToString());
 	}
 
 	return Enum;
@@ -1122,18 +1174,17 @@ UScriptStruct* FHeaderParser::CompileStructDeclaration(FClasses& AllClasses, FCl
 	if (!Scope->HasAnyClassFlags(CLASS_Temporary))
 	{
 		UField* Existing = FindField(Scope, *EffectiveStructName);
-		if ((Existing != NULL) && (Existing->GetOuter() == Scope))
+		if (Existing && Existing->GetOuter() == Scope)
 		{
 			FError::Throwf(TEXT("struct: '%s' already defined here"), *EffectiveStructName);
 		}
-		else if (FindObject<UClass>(ANY_PACKAGE, *EffectiveStructName) != NULL)
+
+		if (FindObject<UClass>(ANY_PACKAGE, *EffectiveStructName) != NULL)
 		{
 			FError::Throwf(TEXT("struct: '%s' conflicts with class name"), *EffectiveStructName);
 		}
-		else
-		{
-			CheckObscures(Scope, StructNameInScript, EffectiveStructName);
-		}
+
+		CheckObscures(Scope, StructNameInScript, EffectiveStructName);
 	}
 
 	// Get optional superstruct.
@@ -2654,7 +2705,8 @@ bool FHeaderParser::GetVarType
 		}
 		FError::Throwf(TEXT("%s: Missing variable type"), Thing );
 	}
-	else if ( VarType.Matches(TEXT("int8")) )
+
+	if ( VarType.Matches(TEXT("int8")) )
 	{
 		VarProperty = FPropertyBase(CPT_Int8);
 	}
@@ -2794,7 +2846,7 @@ bool FHeaderParser::GetVarType
 		// Eat the forward declaration enum text if present
 		MatchIdentifier(TEXT("enum"));
 
-		bool bFailedToFindEnum = true;
+		bool bFoundEnum = false;
 
 		FToken InnerEnumType;
 		if (GetIdentifier(InnerEnumType, true))
@@ -2802,9 +2854,9 @@ bool FHeaderParser::GetVarType
 			if (UEnum* Enum = FindObject<UEnum>(ANY_PACKAGE, InnerEnumType.Identifier))
 			{
 				// In-scope enumeration.
-				VarProperty = FPropertyBase(Enum);
-				bFailedToFindEnum = false;
-			}	
+				VarProperty = FPropertyBase(Enum, CPT_Byte);
+				bFoundEnum  = true;
+			}
 		}
 
 		// Try to handle namespaced enums
@@ -2818,18 +2870,24 @@ bool FHeaderParser::GetVarType
 			}
 		}
 
-		if (bFailedToFindEnum)
+		if (!bFoundEnum)
 		{
 			FError::Throwf(TEXT("Expected the name of a previously defined enum"));
 		}
 
 		RequireSymbol(TEXT(">"), VarType.Identifier, ESymbolParseOption::CloseTemplateBracket);
 	}
-	else if (UEnum* Enum = FindObject<UEnum>( ANY_PACKAGE, VarType.Identifier ))
+	else if (UEnum* Enum = FindObject<UEnum>(ANY_PACKAGE, VarType.Identifier))
 	{
+		EPropertyType UnderlyingType = CPT_Byte;
+
 		if (VariableCategory == EVariableCategory::Member)
 		{
-			FError::Throwf(TEXT("You cannot use the raw enum name as a type for member variables, instead use TEnumAsByte<%s>."), VarType.Identifier);
+			auto* UnderlyingType = GEnumUnderlyingTypes.Find(Enum);
+			if (!UnderlyingType || *UnderlyingType != CPT_Byte)
+			{
+				FError::Throwf(TEXT("You cannot use the raw enum name as a type for member variables, instead use TEnumAsByte or a C++11 enum class with an explicit underlying type (currently only uint8 supported)."), *Enum->CppType);
+			}
 		}
 
 		// Try to handle namespaced enums
@@ -2844,7 +2902,7 @@ bool FHeaderParser::GetVarType
 		}
 
 		// In-scope enumeration.
-		VarProperty = FPropertyBase(Enum);
+		VarProperty            = FPropertyBase(Enum, UnderlyingType);
 		bUnconsumedEnumKeyword = false;
 	}
 	else
@@ -3005,7 +3063,7 @@ bool FHeaderParser::GetVarType
 				CheckInScope( TempClass );
 
 				bool bAllowWeak = !(Disallow & CPF_AutoWeak); // if it is not allowing anything, force it strong. this is probably a function arg
-				VarProperty = FPropertyBase( TempClass, NULL, CPRT_None, bAllowWeak, bIsWeak, bWeakIsAuto, bIsLazy, bIsAsset );
+				VarProperty = FPropertyBase( TempClass, NULL, bAllowWeak, bIsWeak, bWeakIsAuto, bIsLazy, bIsAsset );
 				if (TempClass->IsChildOf(UClass::StaticClass()))
 				{
 					if ( MatchSymbol(TEXT("<")) )
diff --git a/Engine/Source/Programs/UnrealHeaderTool/Private/ParserHelper.cpp b/Engine/Source/Programs/UnrealHeaderTool/Private/ParserHelper.cpp
--- a/Engine/Source/Programs/UnrealHeaderTool/Private/ParserHelper.cpp
+++ b/Engine/Source/Programs/UnrealHeaderTool/Private/ParserHelper.cpp
@@ -218,22 +218,6 @@ const TCHAR* FPropertyBase::GetPropertyTypeText( EPropertyType Type )
 	return TEXT("");
 }
 
-const TCHAR* FPropertyBase::GetPropertyRefText( EPropertyReferenceType Type )
-{
-#if UE_BUILD_DEBUG
-	switch ( Type )
-	{
-		CASE_TEXT(CPRT_None);
-		CASE_TEXT(CPRT_AssignValue);
-		CASE_TEXT(CPRT_SimpleReference);
-		CASE_TEXT(CPRT_AssignmentReference);
-		CASE_TEXT(CPRT_DualReference);
-	}
-#endif
-
-	return TEXT("");
-}
-
 /////////////////////////////////////////////////////
 // FToken
 
diff --git a/Engine/Source/Programs/UnrealHeaderTool/Private/ParserHelper.h b/Engine/Source/Programs/UnrealHeaderTool/Private/ParserHelper.h
--- a/Engine/Source/Programs/UnrealHeaderTool/Private/ParserHelper.h
+++ b/Engine/Source/Programs/UnrealHeaderTool/Private/ParserHelper.h
@@ -11,24 +11,6 @@ extern class FCompilerMetadataManager GScriptHelper;
 	FPropertyBase.
 -----------------------------------------------------------------------------*/
 
-enum EPropertyReferenceType
-{
-	/** no reference */
-	CPRT_None					=0,
-
-	/** we're referencing this property in order to assign a value to it */
-	CPRT_AssignValue			=1,
-
-	/** we're referencing this property's value in such a way that doesn't require it be initialized (if checks, iterators, etc.) */
-	CPRT_SimpleReference		=2,
-
-	/** we're referecning this property's value in such a away that requires it to be initialized (assigning to another property, using in a function call, etc) */
-	CPRT_AssignmentReference	=3,
-
-	/** we're referencing this property in a way that both changes its value and references the value (combination of CPRT_AssignmentReference and CPRT_AssignValue) */
-	CPRT_DualReference			=4,
-};
-
 enum EFunctionExportFlags
 {
 	FUNCEXPORT_Final			=0x00000001,	// function declaration included "final" keyword.  Used to differentiate between functions that have FUNC_Final only because they're private
@@ -121,14 +103,12 @@ public:
 	/** Map of key value pairs that will be added to the package's UMetaData for this property */
 	TMap<FName, FString> MetaData;
 
-	EPropertyReferenceType ReferenceType;
-
 	EPointerType::Type PointerType;
 
 public:
 	/** @name Constructors */
 	//@{
-	FPropertyBase( EPropertyType InType, EPropertyReferenceType InRefType = CPRT_None )
+	explicit FPropertyBase(EPropertyType InType)
 	: Type               (InType)
 	, ArrayType          (EArrayType::None)
 	, PropertyFlags      (0)
@@ -138,13 +118,12 @@ public:
 	, MetaClass          (NULL)
 	, DelegateName       (NAME_None)
 	, RepNotifyName      (NAME_None)
-	, ReferenceType      (InRefType)
 	, PointerType        (EPointerType::None)
 	{
 	}
 
-	FPropertyBase( UEnum* InEnum, EPropertyReferenceType InRefType = CPRT_None )
-	: Type               (CPT_Byte)
+	explicit FPropertyBase(UEnum* InEnum, EPropertyType InType)
+	: Type               (InType)
 	, ArrayType          (EArrayType::None)
 	, PropertyFlags      (0)
 	, RefQualifier       (ERefQualifier::None)
@@ -153,12 +132,11 @@ public:
 	, MetaClass          (NULL)
 	, DelegateName       (NAME_None)
 	, RepNotifyName      (NAME_None)
-	, ReferenceType      (InRefType)
 	, PointerType        (EPointerType::None)
 	{
 	}
 
-	FPropertyBase( UClass* InClass, UClass* InMetaClass=NULL, EPropertyReferenceType InRefType = CPRT_None, bool bAllowWeak = false, bool bIsWeak = false, bool bWeakIsAuto = false, bool bIsLazy = false, bool bIsAsset = false)
+	explicit FPropertyBase(UClass* InClass, UClass* InMetaClass=NULL, bool bAllowWeak = false, bool bIsWeak = false, bool bWeakIsAuto = false, bool bIsLazy = false, bool bIsAsset = false)
 	: Type               (CPT_ObjectReference)
 	, ArrayType          (EArrayType::None)
 	, PropertyFlags      (0)
@@ -168,7 +146,6 @@ public:
 	, MetaClass          (InMetaClass)
 	, DelegateName       (NAME_None)
 	, RepNotifyName      (NAME_None)
-	, ReferenceType      (InRefType)
 	, PointerType        (EPointerType::None)
 	{
 		// if this is an interface class, we use the UInterfaceProperty class instead of UObjectProperty
@@ -216,7 +193,7 @@ public:
 		}
 	}
 
-	FPropertyBase( UScriptStruct* InStruct, EPropertyReferenceType InRefType = CPRT_None )
+	explicit FPropertyBase(UScriptStruct* InStruct)
 	: Type               (CPT_Struct)
 	, ArrayType          (EArrayType::None)
 	, PropertyFlags      (0)
@@ -226,12 +203,11 @@ public:
 	, MetaClass          (NULL)
 	, DelegateName       (NAME_None)
 	, RepNotifyName      (NAME_None)
-	, ReferenceType      (InRefType)
 	, PointerType        (EPointerType::None)
 	{
 	}
 
-	FPropertyBase( UProperty* Property, EPropertyReferenceType InRefType = CPRT_None )
+	explicit FPropertyBase(UProperty* Property)
 	: PropertyExportFlags(PROPEXPORT_Public)
 	, DelegateName       (NAME_None)
 	, RepNotifyName      (NAME_None)
@@ -392,7 +368,6 @@ public:
 		ArrayType     = ArrType;
 		PropertyFlags = Property->PropertyFlags | PropagateFlags;
 		RefQualifier  = ERefQualifier::None;
-		ReferenceType = InRefType;
 		PointerType   = EPointerType::None;
 	}
 	//@}
@@ -470,8 +445,7 @@ public:
 				return false;
 			}
 		}
-		else if ( (Type == CPT_ObjectReference || Type == CPT_WeakObjectReference || Type == CPT_LazyObjectReference || Type == CPT_AssetObjectReference) && Other.Type != CPT_Interface
-			&&	(ReferenceType == CPRT_AssignmentReference || ReferenceType == CPRT_DualReference || (PropertyFlags & CPF_ReturnParm)))
+		else if ((Type == CPT_ObjectReference || Type == CPT_WeakObjectReference || Type == CPT_LazyObjectReference || Type == CPT_AssetObjectReference) && Other.Type != CPT_Interface && (PropertyFlags & CPF_ReturnParm))
 		{
 			bReverseClassChainCheck = false;
 		}
@@ -612,8 +586,6 @@ public:
 	//@}
 
 	static const TCHAR* GetPropertyTypeText( EPropertyType Type );
-
-	static const TCHAR* GetPropertyRefText( EPropertyReferenceType Type );
 };
 
 //
@@ -727,28 +699,34 @@ public:
 	}
 
 	// Constructors.
-	FToken() : FPropertyBase(CPT_None, CPRT_None), TokenProperty(NULL)
+	FToken()
+		: FPropertyBase(CPT_None)
+		, TokenProperty(NULL)
 	{
-		InitToken(CPT_None, CPRT_None);
+		InitToken(CPT_None);
 	}
-	FToken( EPropertyReferenceType InRefType )
-	: FPropertyBase( CPT_None, InRefType ), TokenProperty(NULL)
+
+	// copy constructors
+	explicit FToken(EPropertyType InType)
+		: FPropertyBase(InType)
+		, TokenProperty(NULL)
 	{
-		InitToken( CPT_None, InRefType );
+		InitToken(InType);
 	}
 
 	// copy constructors
-	FToken( const FPropertyBase& InType )
-	: FPropertyBase( CPT_None, InType.ReferenceType ), TokenProperty(NULL)
+	FToken(const FPropertyBase& InType)
+		: FPropertyBase(CPT_None)
+		, TokenProperty(NULL)
 	{
-		InitToken( CPT_None, InType.ReferenceType );
+		InitToken( CPT_None );
 		(FPropertyBase&)*this = InType;
 	}
 
 	// Inlines.
-	void InitToken( EPropertyType InType, EPropertyReferenceType InRefType = CPRT_None )
+	void InitToken( EPropertyType InType )
 	{
-		(FPropertyBase&)*this = FPropertyBase(InType, InRefType);
+		(FPropertyBase&)*this = FPropertyBase(InType);
 		TokenType		= TOKEN_None;
 		TokenName		= NAME_None;
 		StartPos		= 0;
@@ -847,11 +825,11 @@ public:
 	FString Describe()
 	{
 		return FString::Printf(
-			TEXT("Property:%s  Type:%s  TokenName:%s  ConstValue:%s  Struct:%s  Flags:%lli  RefType:%s"),
+			TEXT("Property:%s  Type:%s  TokenName:%s  ConstValue:%s  Struct:%s  Flags:%lli"),
 			TokenProperty!=NULL?*TokenProperty->GetName():TEXT("NULL"),
 			GetPropertyTypeText(Type), *TokenName.ToString(), *GetConstantValue(),
 			Struct!=NULL?*Struct->GetName():TEXT("NULL"),
-			PropertyFlags, GetPropertyRefText(ReferenceType)
+			PropertyFlags
 			);
 	}
 };
diff --git a/Engine/Source/Programs/UnrealHeaderTool/Resources/UHTDebugging/EnumOnlyHeader.h b/Engine/Source/Programs/UnrealHeaderTool/Resources/UHTDebugging/EnumOnlyHeader.h
--- a/Engine/Source/Programs/UnrealHeaderTool/Resources/UHTDebugging/EnumOnlyHeader.h
+++ b/Engine/Source/Programs/UnrealHeaderTool/Resources/UHTDebugging/EnumOnlyHeader.h
@@ -20,3 +20,11 @@ namespace ESomeNamespacedEnum
 		Six
 	};
 }
+
+UENUM()
+enum class ECppEnum : uint8
+{
+	Seven,
+	Eight,
+	Nine
+};
diff --git a/Engine/Source/Programs/UnrealHeaderTool/Resources/UHTDebugging/TestObject.h b/Engine/Source/Programs/UnrealHeaderTool/Resources/UHTDebugging/TestObject.h
--- a/Engine/Source/Programs/UnrealHeaderTool/Resources/UHTDebugging/TestObject.h
+++ b/Engine/Source/Programs/UnrealHeaderTool/Resources/UHTDebugging/TestObject.h
@@ -1,5 +1,6 @@
 #pragma once
 
+#include "EnumOnlyHeader.h"
 #include "TestObject.generated.h"
 
 UCLASS()
@@ -13,6 +14,9 @@ class UTestObject : public UObject
 	UPROPERTY()
 	int32 Cpp11Init = 123;
 
+	UPROPERTY()
+	ECppEnum EnumProperty;
+
 #if BLAH
 	UFUNCTION() int x; // This should not compile if UHT parses it, which it shouldn't
 #elif BLAH2
diff --git a/Engine/Source/Runtime/AIModule/Classes/BehaviorTree/BTCompositeNode.h b/Engine/Source/Runtime/AIModule/Classes/BehaviorTree/BTCompositeNode.h
--- a/Engine/Source/Runtime/AIModule/Classes/BehaviorTree/BTCompositeNode.h
+++ b/Engine/Source/Runtime/AIModule/Classes/BehaviorTree/BTCompositeNode.h
@@ -17,14 +17,11 @@ struct FBTCompositeMemory
 };
 
 UENUM()
-namespace EBTChildIndex
+enum class EBTChildIndex
 {
-	enum Type
-	{
-		FirstNode,
-		TaskNode,
-	};
-}
+	FirstNode,
+	TaskNode,
+};
 
 UENUM()
 namespace EBTDecoratorLogic
@@ -133,7 +130,7 @@ class AIMODULE_API UBTCompositeNode : public UBTNode
 	int32 GetChildrenNum() const;
 
 	/** @return execution index of child node */
-	uint16 GetChildExecutionIndex(int32 Index, EBTChildIndex::Type ChildMode = EBTChildIndex::TaskNode) const;
+	uint16 GetChildExecutionIndex(int32 Index, EBTChildIndex ChildMode = EBTChildIndex::TaskNode) const;
 
 	/** @return execution index of last node in child branches */
 	uint16 GetLastExecutionIndex() const;
diff --git a/Engine/Source/Runtime/AIModule/Private/BehaviorTree/BTCompositeNode.cpp b/Engine/Source/Runtime/AIModule/Private/BehaviorTree/BTCompositeNode.cpp
--- a/Engine/Source/Runtime/AIModule/Private/BehaviorTree/BTCompositeNode.cpp
+++ b/Engine/Source/Runtime/AIModule/Private/BehaviorTree/BTCompositeNode.cpp
@@ -505,7 +505,7 @@ void UBTCompositeNode::RequestDelayedExecution(class UBehaviorTreeComponent* Own
 	OwnerComp->RequestExecution(LastResult);
 }
 
-uint16 UBTCompositeNode::GetChildExecutionIndex(int32 Index, EBTChildIndex::Type ChildMode) const
+uint16 UBTCompositeNode::GetChildExecutionIndex(int32 Index, EBTChildIndex ChildMode) const
 {
 	const UBTNode* ChildNode = GetChildNode(Index);
 	if (ChildNode)
diff --git a/Engine/Source/Runtime/Core/Public/UObject/ObjectVersion.h b/Engine/Source/Runtime/Core/Public/UObject/ObjectVersion.h
--- a/Engine/Source/Runtime/Core/Public/UObject/ObjectVersion.h
+++ b/Engine/Source/Runtime/Core/Public/UObject/ObjectVersion.h
@@ -592,6 +592,8 @@ enum EUnrealEngineObjectUE4Version
 	VER_UE4_SKELETON_ADD_SMARTNAMES,
 	// Added the currency code field to FTextHistory_AsCurrency
 	VER_UE4_ADDED_CURRENCY_CODE_TO_FTEXT,
+	// Added support for C++11 enum classes
+	VER_UE4_ENUM_CLASS_SUPPORT,
 
 	// -----<new versions can be added before this line>-------------------------------------------------
 	// - this needs to be the last line (see note below)
diff --git a/Engine/Source/Runtime/CoreUObject/Private/UObject/Enum.cpp b/Engine/Source/Runtime/CoreUObject/Private/UObject/Enum.cpp
--- a/Engine/Source/Runtime/CoreUObject/Private/UObject/Enum.cpp
+++ b/Engine/Source/Runtime/CoreUObject/Private/UObject/Enum.cpp
@@ -21,7 +21,19 @@ void UEnum::Serialize( FArchive& Ar )
 {
 	Super::Serialize(Ar);
 	Ar << Names;
-	Ar << bIsNamespace;
+
+	if (Ar.UE4Ver() < VER_UE4_ENUM_CLASS_SUPPORT)
+	{
+		bool bIsNamespace;
+		Ar << bIsNamespace;
+		CppForm = bIsNamespace ? ECppForm::Namespaced : ECppForm::Regular;
+	}
+	else
+	{
+		int8 EnumTypeByte = (int8)CppForm;
+		Ar << EnumTypeByte;
+		CppForm = (ECppForm)EnumTypeByte;
+	}
 
 	if (Ar.IsLoading() || Ar.IsSaving())
 	{	
@@ -36,7 +48,7 @@ int32 UEnum::ResolveEnumerator(FArchive& Ar, int32 EnumeratorIndex) const
 
 FString UEnum::GenerateFullEnumName(const TCHAR* InEnumName) const
 {
-	return bIsNamespace ? GenerateFullEnumName(this, InEnumName) : InEnumName;
+	return (CppForm != ECppForm::Regular) ? GenerateFullEnumName(this, InEnumName) : InEnumName;
 }
 
 void UEnum::AddNamesToMasterList()
@@ -127,7 +139,7 @@ FString UEnum::GenerateEnumPrefix() const
 int32 UEnum::FindEnumIndex(FName InName) const
 {
 	int32 EnumIndex = Names.Find( InName );
-	if (EnumIndex == INDEX_NONE && bIsNamespace && !IsFullEnumName(*InName.ToString()))
+	if (EnumIndex == INDEX_NONE && CppForm != ECppForm::Regular && !IsFullEnumName(*InName.ToString()))
 	{
 		// Try the long enum name if InName doesn't have a namespace specified and this is a namespace enum.
 		FName LongName(*GenerateFullEnumName(*InName.ToString()));
@@ -217,15 +229,14 @@ int32 UEnum::FindEnumRedirects(const UEnum * Enum, FName EnumEntryName)
  *
  * @return	true unless the MAX enum already exists and isn't the last enum.
  */
-bool UEnum::SetEnums(TArray<FName>& InNames, bool bNamespace)
+bool UEnum::SetEnums(TArray<FName>& InNames, UEnum::ECppForm InCppForm)
 {
 	if (Names.Num() > 0)
 	{
 		RemoveNamesFromMasterList();
 	}
-	Names.Empty();
-	Names = InNames;
-	bIsNamespace = bNamespace;
+	Names   = InNames;
+	CppForm = InCppForm;
 	return GenerateMaxEnum();
 }
 
@@ -249,10 +260,8 @@ bool UEnum::GenerateMaxEnum()
 			// the MAX identifier is already being used by another enum
 			return false;
 		}
-		else
-		{
-			Names.Add(MaxEnumItem);
-		}
+
+		Names.Add(MaxEnumItem);
 	}
 
 	AddNamesToMasterList();
diff --git a/Engine/Source/Runtime/CoreUObject/Private/UObject/PropertyByte.cpp b/Engine/Source/Runtime/CoreUObject/Private/UObject/PropertyByte.cpp
--- a/Engine/Source/Runtime/CoreUObject/Private/UObject/PropertyByte.cpp
+++ b/Engine/Source/Runtime/CoreUObject/Private/UObject/PropertyByte.cpp
@@ -83,11 +83,16 @@ FString UByteProperty::GetCPPType( FString* ExtendedTypeText/*=NULL*/, uint32 CP
 {
 	if (Enum)
 	{
-		FString FullyQualifiedEnumName = Enum->GetName();
-		if (!Enum->ActualEnumNameInsideNamespace.IsEmpty())
+		FString FullyQualifiedEnumName;
+		if (!Enum->CppType.IsEmpty())
 		{
-			FullyQualifiedEnumName += TEXT("::");
-			FullyQualifiedEnumName += Enum->ActualEnumNameInsideNamespace;
+			FullyQualifiedEnumName = Enum->CppType;
+		}
+		else
+		{
+			// This would give the wrong result if it's a namespaced type and the CppType hasn't
+			// been set, but we do this here in case existing code relies on it... somehow.
+			FullyQualifiedEnumName = Enum->GetName();
 		}
 
 		if ((CPPExportFlags & CPPF_ArgumentOrReturnValue) && ((PropertyFlags & CPF_ReturnParm) || !(PropertyFlags & CPF_OutParm)))
diff --git a/Engine/Source/Runtime/CoreUObject/Public/UObject/Class.h b/Engine/Source/Runtime/CoreUObject/Public/UObject/Class.h
--- a/Engine/Source/Runtime/CoreUObject/Public/UObject/Class.h
+++ b/Engine/Source/Runtime/CoreUObject/Public/UObject/Class.h
@@ -1288,15 +1288,23 @@ class UEnum : public UField
 	DECLARE_CASTED_CLASS_INTRINSIC_WITH_API(UEnum,UField,0,CoreUObject,CASTCLASS_UEnum,COREUOBJECT_API)
 
 public:
-	// This will be the true name of the enum inside a namespace, if the enum wasn't in the global scope
-	FString ActualEnumNameInsideNamespace;
+	enum class ECppForm
+	{
+		Regular,
+		Namespaced,
+		EnumClass
+	};
+
+	// This will be the true type of the enum as a string, e.g. "ENamespacedEnum::InnerType" or "ERegularEnum" or "EEnumClass"
+	FString CppType;
 
 protected:
 	// Variables.
 	/** List of all enum names. */
 	TArray<FName> Names;
-	/** True if this enum is namespace enum, false if global. */
-	bool bIsNamespace;
+
+	/** How the enum was originally defined. */
+	ECppForm CppForm;
 
 	/** global list of all value names used by all enums in memory, used for property text import */
 	COREUOBJECT_API static TMap<FName, UEnum*> AllEnumNames;
@@ -1323,13 +1331,13 @@ public:
 	COREUOBJECT_API virtual int32 ResolveEnumerator(FArchive& Ar, int32 EnumeratorIndex) const;
 
 	/**
-	 * Checks if this enum is a nemespace declared enum.
+	 * Returns the type of enum: whether it's a regular enum, namespaced enum or C++11 enum class.
 	 *
-	 * @return true if this enum is a nemespace declared enum, false otherwise.
+	 * @return The enum type.
 	 */
-	bool IsNamespaceEnum() const
+	ECppForm GetCppForm() const
 	{
-		return bIsNamespace;
+		return CppForm;
 	}
 
 	/**
@@ -1352,7 +1360,12 @@ public:
 	 */
 	COREUOBJECT_API static FString GenerateFullEnumName(const UEnum * InEnum, const TCHAR* InEnumName)
 	{
-		return (InEnum->bIsNamespace && IsFullEnumName(InEnumName) == false) ? FString::Printf(TEXT("%s::%s"), *InEnum->GetName(), InEnumName) : InEnumName;
+		if (InEnum->GetCppForm() == ECppForm::Regular || IsFullEnumName(InEnumName))
+		{
+			return InEnumName;
+		}
+
+		return FString::Printf(TEXT("%s::%s"), *InEnum->GetName(), InEnumName);
 	}
 
 	/**
@@ -1412,10 +1425,10 @@ public:
 	 * Sets the array of enums.
 	 *
 	 * @param InNames List of enum names.
-	 * @param bNamespace True if this enum is namespace enum, false if global.
+	 * @param InCppForm The form of enum.
 	 * @return	true unless the MAX enum already exists and isn't the last enum.
 	 */
-	COREUOBJECT_API bool SetEnums(TArray<FName>& InNames, bool bNamespace);
+	COREUOBJECT_API bool SetEnums(TArray<FName>& InNames, ECppForm InCppForm);
 
 	/**
 	 * @return	The enum name at the specified Index.
@@ -1436,19 +1449,17 @@ public:
 	{
 		if (Names.IsValidIndex(InIndex))
 		{
-			if (bIsNamespace)
+			if (CppForm == ECppForm::Regular)
 			{
-				// Strip the namespace from the name.
-				FString EnumName(Names[InIndex].ToString());
-				int32 ScopeIndex = EnumName.Find(TEXT("::"));
-				if (ScopeIndex != INDEX_NONE)
-				{
-					return EnumName.Mid(ScopeIndex + 2);
-				}
+				return Names[InIndex].ToString();
 			}
-			else
+
+			// Strip the namespace from the name.
+			FString EnumName(Names[InIndex].ToString());
+			int32 ScopeIndex = EnumName.Find(TEXT("::"));
+			if (ScopeIndex != INDEX_NONE)
 			{
-				return Names[InIndex].ToString();
+				return EnumName.Mid(ScopeIndex + 2);
 			}
 		}
 		return FName(NAME_None).ToString();
diff --git a/Engine/Source/Runtime/Engine/Classes/Engine/UserDefinedEnum.h b/Engine/Source/Runtime/Engine/Classes/Engine/UserDefinedEnum.h
--- a/Engine/Source/Runtime/Engine/Classes/Engine/UserDefinedEnum.h
+++ b/Engine/Source/Runtime/Engine/Classes/Engine/UserDefinedEnum.h
@@ -4,7 +4,7 @@
 
 /** 
  *	User defined enumerations:
- *	- have bIsNamespace set always to true (to comfortable handle names collisions)
+ *	- EnumType is always UEnum::ECppForm::Namespaced (to comfortable handle names collisions)
  *	- always have the last '_MAX' enumerator, that cannot be changed by user
  *	- Full enumerator name has form: '<enumeration path>::<short, user defined enumerator name>'
  */
diff --git a/Engine/Source/Runtime/Engine/Private/UserDefinedEnum.cpp b/Engine/Source/Runtime/Engine/Private/UserDefinedEnum.cpp
--- a/Engine/Source/Runtime/Engine/Private/UserDefinedEnum.cpp
+++ b/Engine/Source/Runtime/Engine/Private/UserDefinedEnum.cpp
@@ -14,7 +14,7 @@ UUserDefinedEnum::UUserDefinedEnum(const class FPostConstructInitializePropertie
 
 FString UUserDefinedEnum::GenerateFullEnumName(const TCHAR* InEnumName) const
 {
-	check(bIsNamespace);
+	check(CppForm == ECppForm::Namespaced);
 
 	FString PathName;
 	GetPathName(NULL, PathName);
diff --git a/Engine/Source/Runtime/JsonUtilities/Private/JsonObjectConverter.cpp b/Engine/Source/Runtime/JsonUtilities/Private/JsonObjectConverter.cpp
--- a/Engine/Source/Runtime/JsonUtilities/Private/JsonObjectConverter.cpp
+++ b/Engine/Source/Runtime/JsonUtilities/Private/JsonObjectConverter.cpp
@@ -158,7 +158,7 @@ bool FJsonObjectConverter::JsonValueToUProperty(const TSharedPtr<FJsonValue> Jso
 			int32 IntValue = EnumProperty->FindEnumIndex(FName(*StrValue));
 			if (IntValue == INDEX_NONE)
 			{
-				UE_LOG(LogJson, Error, TEXT("JsonValueToUProperty - Unable import enum %s from string value %s"), *EnumProperty->ActualEnumNameInsideNamespace, *StrValue);
+				UE_LOG(LogJson, Error, TEXT("JsonValueToUProperty - Unable import enum %s from string value %s"), *EnumProperty->CppType, *StrValue);
 				return false;
 			}
 			NumericProperty->SetIntPropertyValue(OutValue, (int64)IntValue);
