Fixed UHT failing to compile on Windows when PCHs are disabled
- This was caused because of how MSVC resolves include paths (Class.h was ambiguous)

diff --git a/Engine/Source/Programs/UnrealHeaderTool/Private/Class.cpp b/Engine/Source/Programs/UnrealHeaderTool/Private/Class.cpp
deleted file mode 100644
diff --git a/Engine/Source/Programs/UnrealHeaderTool/Private/Class.h b/Engine/Source/Programs/UnrealHeaderTool/Private/Class.h
deleted file mode 100644
diff --git a/Engine/Source/Programs/UnrealHeaderTool/Private/Classes.cpp b/Engine/Source/Programs/UnrealHeaderTool/Private/Classes.cpp
--- a/Engine/Source/Programs/UnrealHeaderTool/Private/Classes.cpp
+++ b/Engine/Source/Programs/UnrealHeaderTool/Private/Classes.cpp
@@ -2,7 +2,7 @@
 
 #include "UnrealHeaderTool.h"
 #include "Classes.h"
-#include "Class.h"
+#include "ParserClass.h"
 #include "StringUtils.h"
 
 namespace
diff --git a/Engine/Source/Programs/UnrealHeaderTool/Private/NativeClassExporter.h b/Engine/Source/Programs/UnrealHeaderTool/Private/NativeClassExporter.h
--- a/Engine/Source/Programs/UnrealHeaderTool/Private/NativeClassExporter.h
+++ b/Engine/Source/Programs/UnrealHeaderTool/Private/NativeClassExporter.h
@@ -2,7 +2,7 @@
 
 #pragma once
 
-#include "Class.h"
+#include "ParserClass.h"
 #include "Classes.h"
 
 //
diff --git a/Engine/Source/Programs/UnrealHeaderTool/Private/ParserClass.cpp b/Engine/Source/Programs/UnrealHeaderTool/Private/ParserClass.cpp
new file mode 100644
--- /dev/null
+++ b/Engine/Source/Programs/UnrealHeaderTool/Private/ParserClass.cpp
@@ -0,0 +1,105 @@
+// Copyright 1998-2014 Epic Games, Inc. All Rights Reserved.
+
+#include "UnrealHeaderTool.h"
+#include "ParserClass.h"
+#include "Classes.h"
+#include "ClassMaps.h"
+
+bool FClass::Inherits(const FClass* SuspectBase) const
+{
+	return IsChildOf(SuspectBase);
+}
+
+FString FClass::GetNameWithPrefix(EEnforceInterfacePrefix::Type EnforceInterfacePrefix) const
+{
+	const TCHAR* Prefix = 0;
+
+	if (HasAnyClassFlags(CLASS_Interface))
+	{
+		// Grab the expected prefix for interfaces (U on the first one, I on the second one)
+		switch (EnforceInterfacePrefix)
+		{
+			case EEnforceInterfacePrefix::None:
+				// For old-style files: "I" for interfaces, unless it's the actual "Interface" class, which gets "U"
+				if (GetFName() == NAME_Interface)
+				{
+					Prefix = TEXT("U");
+				}
+				else
+				{
+					Prefix = TEXT("I");
+				}
+				break;
+
+			case EEnforceInterfacePrefix::I:
+				Prefix = TEXT("I");
+				break;
+
+			case EEnforceInterfacePrefix::U:
+				Prefix = TEXT("U");
+				break;
+
+			default:
+				check(false);
+		}
+	}
+	else
+	{
+		// Get the expected class name with prefix
+		Prefix = GetPrefixCPP();
+	}
+
+	return FString::Printf(TEXT("%s%s"), Prefix, *GetName());
+}
+
+FClass* FClass::GetSuperClass() const
+{
+	return static_cast<FClass*>(static_cast<const UClass*>(this)->GetSuperClass());
+}
+
+FClass* FClass::GetClassWithin() const
+{
+	return (FClass*)ClassWithin;
+}
+
+TArray<FName> FClass::GetDependentNames() const
+{
+	TArray<FName> Result;
+
+	for (auto Name : *GClassDependentOnMap.FindOrAdd(const_cast<FClass*>(this)))
+	{
+		Result.Add(Name);
+	}
+
+	return Result;
+}
+
+TArray<FClass*> FClass::GetInterfaceTypes() const
+{
+	TArray<FClass*> Result;
+	for (auto& i : Interfaces)
+	{
+		Result.Add((FClass*)i.Class);
+	}
+	return Result;
+}
+
+void FClass::GetHideCategories(TArray<FString>& OutHideCategories) const
+{
+	static const FName NAME_HideCategories(TEXT("HideCategories"));
+	if (HasMetaData(NAME_HideCategories))
+	{
+		const FString& HideCategories = GetMetaData(NAME_HideCategories);
+		HideCategories.ParseIntoArray(&OutHideCategories, TEXT(" "), true);
+	}
+}
+
+void FClass::GetShowCategories(TArray<FString>& OutShowCategories) const
+{
+	static const FName NAME_ShowCategories(TEXT("ShowCategories"));
+	if (HasMetaData(NAME_ShowCategories))
+	{
+		const FString& ShowCategories = GetMetaData(NAME_ShowCategories);
+		ShowCategories.ParseIntoArray(&OutShowCategories, TEXT(" "), true);
+	}
+}
diff --git a/Engine/Source/Programs/UnrealHeaderTool/Private/ParserClass.h b/Engine/Source/Programs/UnrealHeaderTool/Private/ParserClass.h
new file mode 100644
--- /dev/null
+++ b/Engine/Source/Programs/UnrealHeaderTool/Private/ParserClass.h
@@ -0,0 +1,57 @@
+// Copyright 1998-2014 Epic Games, Inc. All Rights Reserved.
+
+#pragma once
+
+struct EEnforceInterfacePrefix
+{
+	enum Type
+	{
+		None,
+		I,
+		U
+	};
+};
+
+class FClasses;
+
+class FClass : public UClass
+{
+public:
+	FClass();
+
+	/**
+	 * Tests if this class inherits another.
+	 *
+	 * @param SuspectBase The class to test if it's a base of this.
+	 * @return true if the SuspectBase is a base of this, false otherwise.
+	 */
+	bool Inherits(const FClass* SuspectBase) const;
+
+	/** 
+	 * Returns the name of the given class with a valid prefix.
+	 *
+	 * @param InClass Class used to create a valid class name with prefix
+	 */
+	FString GetNameWithPrefix(EEnforceInterfacePrefix::Type EnforceInterfacePrefix = EEnforceInterfacePrefix::None) const;
+
+	/**
+	 * Returns the super class of this class, or NULL if there is no superclass.
+	 *
+	 * @return The super class of this class.
+	 */
+	FClass* GetSuperClass() const;
+
+	/**
+	 * Returns the 'within' class of this class.
+	 *
+	 * @return The 'within' class of this class.
+	 */
+	FClass* GetClassWithin() const;
+
+	TArray<FName> GetDependentNames() const;
+
+	TArray<FClass*> GetInterfaceTypes() const;
+
+	void GetHideCategories(TArray<FString>& OutHideCategories) const;
+	void GetShowCategories(TArray<FString>& OutShowCategories) const;
+};
