Add error message when UHT fails to create generated.h.
TTP #339649 UHT fails to create generated.h with no error message.

diff --git a/Engine/Source/Programs/UnrealHeaderTool/Private/HeaderParser.cpp b/Engine/Source/Programs/UnrealHeaderTool/Private/HeaderParser.cpp
--- a/Engine/Source/Programs/UnrealHeaderTool/Private/HeaderParser.cpp
+++ b/Engine/Source/Programs/UnrealHeaderTool/Private/HeaderParser.cpp
@@ -3967,7 +3967,7 @@ bool FHeaderParser::SkipDeclaration(FToken& Token)
 
 			if (Nest < 0)
 			{
-				FError::Throwf(TEXT("Unexpected '}'. Didn't you miss a semi-colon?"));
+				FError::Throwf(TEXT("Unexpected '}'. Did you miss a semi-colon?"));
 			}
 		}
 		else if (bMacroDeclaration && Nest == 0)
@@ -3995,7 +3995,12 @@ bool FHeaderParser::SkipDeclaration(FToken& Token)
 				// Not a variable name.
 				UngetToken(VariableName);
 			}
+			else if (!SafeMatchSymbol(TEXT(";")))
+			{
+				FError::Throwf(*FString::Printf(TEXT("Unexpected '%s'. Did you miss a semi-colon?"), VariableName.Identifier));
+			}
 		}
+
 		// C++ allows any number of ';' after member declaration/definition.
 		while (SafeMatchSymbol(TEXT(";")));
 	}
