Fixed exception-related compile errors with Clang on Windows
- Exceptions are not yet supported in Clang on Windows platform

diff --git a/Engine/Source/Programs/UnrealHeaderTool/Private/CodeGenerator.cpp b/Engine/Source/Programs/UnrealHeaderTool/Private/CodeGenerator.cpp
--- a/Engine/Source/Programs/UnrealHeaderTool/Private/CodeGenerator.cpp
+++ b/Engine/Source/Programs/UnrealHeaderTool/Private/CodeGenerator.cpp
@@ -4464,15 +4464,19 @@ ECompilationResult::Type UnrealHeaderTool_Main(const FString& ModuleInfoFilename
 	FString ModuleInfoPath = FPaths::GetPath(ModuleInfoFilename);
 
 	// Load the manifest file, giving a list of all modules to be processed, pre-sorted by dependency ordering
+#if !PLATFORM_EXCEPTIONS_DISABLED
 	try
+#endif
 	{
 		GManifest = FManifest::LoadFromFile(ModuleInfoFilename);
 	}
+#if !PLATFORM_EXCEPTIONS_DISABLED
 	catch (const TCHAR* Ex)
 	{
 		UE_LOG(LogCompile, Error, TEXT("Failed to load manifest file '%s': %s"), *ModuleInfoFilename, Ex);
 		return GCompilationResult;
 	}
+#endif
 
 	// Load classes for editing.
 	int32 NumFailures = 0;
diff --git a/Engine/Source/Runtime/Core/Public/GenericPlatform/GenericPlatformMisc.h b/Engine/Source/Runtime/Core/Public/GenericPlatform/GenericPlatformMisc.h
--- a/Engine/Source/Runtime/Core/Public/GenericPlatform/GenericPlatformMisc.h
+++ b/Engine/Source/Runtime/Core/Public/GenericPlatform/GenericPlatformMisc.h
@@ -391,7 +391,7 @@ struct CORE_API FGenericPlatformMisc
 
 	FORCEINLINE static void RaiseException( uint32 ExceptionCode )
 	{
-#if HACK_HEADER_GENERATOR
+#if HACK_HEADER_GENERATOR && !PLATFORM_EXCEPTIONS_DISABLED
 		// We want Unreal Header Tool to throw an exception but in normal runtime code 
 		// we don't support exception handling
 		throw( ExceptionCode );
