Prevent declaring structures in the UCLASS macro when building documentation. Prevents Doxygen from being able to associate the class comment with the class definition.

diff --git a/Engine/Source/Programs/UnrealHeaderTool/Private/CodeGenerator.cpp b/Engine/Source/Programs/UnrealHeaderTool/Private/CodeGenerator.cpp
--- a/Engine/Source/Programs/UnrealHeaderTool/Private/CodeGenerator.cpp
+++ b/Engine/Source/Programs/UnrealHeaderTool/Private/CodeGenerator.cpp
@@ -1992,9 +1992,13 @@ void FNativeClassHeaderGenerator::ExportClassHeaderWrapper( FClass* Class, bool
 		{
 			GeneratedHeaderText.Logf(TEXT("#undef UCLASS\r\n"));
 			GeneratedHeaderText.Logf(TEXT("#undef UINTERFACE\r\n"));
+			GeneratedHeaderText.Logf(TEXT("#if UE_BUILD_DOCS\r\n"));
 			FString MacroName = FString::Printf(TEXT("%s(...)"), Class->HasAnyClassFlags(CLASS_Interface) ? TEXT("UINTERFACE") : TEXT("UCLASS"));
-			FString Macroized = Macroize(*MacroName, *PrologMacroCalls);
+			GeneratedHeaderText.Logf(TEXT("#define %s\r\n"), *MacroName);
+			GeneratedHeaderText.Logf(TEXT("#else\r\n"));
+			FString Macroized = Macroize(*MacroName, *PrologMacroCalls).Replace(TEXT("\r\n\r\n\r\n"), TEXT("\r\n"));
 			GeneratedHeaderText.Log(*Macroized);
+			GeneratedHeaderText.Logf(TEXT("#endif\r\n\r\n\r\n"));
 			PrologMacroCalls.Empty();
 		}
 
