Removed the error we raised when a const function was exposed as BlueprintPure, tested by marking QAPawn::GetWeapon() BlueprintPure, also ran BlueprintInfoDump commandlet, and no changes were observed in blueprint data

diff --git a/Engine/Source/Programs/UnrealHeaderTool/Private/HeaderParser.cpp b/Engine/Source/Programs/UnrealHeaderTool/Private/HeaderParser.cpp
--- a/Engine/Source/Programs/UnrealHeaderTool/Private/HeaderParser.cpp
+++ b/Engine/Source/Programs/UnrealHeaderTool/Private/HeaderParser.cpp
@@ -5433,13 +5433,13 @@ void FHeaderParser::CompileFunctionDeclaration(FClasses& AllClasses)
 			//FError::Throwf(TEXT("'const' may only be used for native functions"));
 		}
 
-		if( (FuncInfo.FunctionFlags & (FUNC_BlueprintPure)) != 0 )
-		{
-			FError::Throwf(TEXT("Cannot mark function '%s' as both 'BlueprintPure' and 'const'"), *TopFunction->GetName());
-		}
-
 		FuncInfo.FunctionFlags |= FUNC_Const;
 
+		// @todo: the presence of const and one or more outputs does not guarantee that there are
+		// no side effects. On GCC and clang we could use __attribure__((pure)) or __attribute__((const))
+		// or we could just rely on the use marking things BlueprintPure. Either way, checking the C++
+		// const identifier to determine purity is not desirable. We should remove the following logic:
+
 		// If its a const BlueprintCallable function with some sort of output, mark it as BlueprintPure as well
 		if ( bHasAnyOutputs && ((FuncInfo.FunctionFlags & FUNC_BlueprintCallable) != 0) )
 		{
