Resolved issues causing an ensure to fire when compiling a Blueprint with an overide function.

Having a private function with BlueprintImplementableEvent will now report as an error by the UHT.

Having a sealed function with BlueprintImplementableEvent will now report as an error by the UHT.

Having a 'final' function with BlueprintImplementableEvent will now report as an error by the UHT.

#ttp 343140 - BP:  LIVE CRASH:  ensure((Function->FunctionFlags & FUNC_FuncOverrideMatch) == (OverridenFunction->FunctionFlags & FUNC_FuncOverrideMatch))

diff --git a/Engine/Source/Editor/IntroTutorials/Public/EditorTutorial.h b/Engine/Source/Editor/IntroTutorials/Public/EditorTutorial.h
--- a/Engine/Source/Editor/IntroTutorials/Public/EditorTutorial.h
+++ b/Engine/Source/Editor/IntroTutorials/Public/EditorTutorial.h
@@ -231,7 +231,7 @@ public:
 	/** Called each tick so the Blueprint can optionally complete or skip stages */
 	void HandleTickCurrentStage(FName StageName);
 
-private:
+protected:
 	/** Event fired when a tutorial stage begins */
 	UFUNCTION(BlueprintImplementableEvent, Category="Tutorial")
 	void OnTutorialStageStarted(FName StageName) const;
diff --git a/Engine/Source/Programs/UnrealHeaderTool/Private/HeaderParser.cpp b/Engine/Source/Programs/UnrealHeaderTool/Private/HeaderParser.cpp
--- a/Engine/Source/Programs/UnrealHeaderTool/Private/HeaderParser.cpp
+++ b/Engine/Source/Programs/UnrealHeaderTool/Private/HeaderParser.cpp
@@ -337,6 +337,11 @@ namespace
 		{
 			FError::Throwf(TEXT("SealedEvent may only be used on events"));
 		}
+
+		if (FuncInfo.bSealedEvent && FuncInfo.FunctionFlags & FUNC_BlueprintEvent)
+		{
+			FError::Throwf(TEXT("SealedEvent cannot be used on Blueprint events"));
+		}
 	}
 }
 	
@@ -5141,6 +5146,41 @@ void FHeaderParser::CompileFunctionDeclaration(FClasses& AllClasses)
 
 	FFuncInfo FuncInfo;
 	FuncInfo.FunctionFlags = FUNC_Native;
+
+	// Infer the function's access level from the currently declared C++ access level
+	if (CurrentAccessSpecifier == ACCESS_Public)
+	{
+		FuncInfo.FunctionFlags |= FUNC_Public;
+	}
+	else if (CurrentAccessSpecifier == ACCESS_Protected)
+	{
+		FuncInfo.FunctionFlags |= FUNC_Protected;
+	}
+	else if (CurrentAccessSpecifier == ACCESS_Private)
+	{
+		FuncInfo.FunctionFlags |= FUNC_Private;
+		FuncInfo.FunctionFlags |= FUNC_Final;
+
+		// This is automatically final as well, but in a different way and for a different reason
+		bAutomaticallyFinal = false;
+	}
+	else
+	{
+		FError::Throwf(TEXT("Unknown access level"));
+	}
+
+	// non-static functions in a const class must be const themselves
+	if (Class->HasAnyClassFlags(CLASS_Const))
+	{
+		FuncInfo.FunctionFlags |= FUNC_Const;
+	}
+
+	if (MatchIdentifier(TEXT("static")))
+	{
+		FuncInfo.FunctionFlags |= FUNC_Static;
+		FuncInfo.FunctionExportFlags |= FUNCEXPORT_CppStatic;
+	}
+
 	ProcessFunctionSpecifiers(FuncInfo, SpecifiersFound);
 
 	if( (FuncInfo.FunctionFlags & FUNC_BlueprintPure) && Class->HasAnyClassFlags(CLASS_Interface) )
@@ -5205,41 +5245,6 @@ void FHeaderParser::CompileFunctionDeclaration(FClasses& AllClasses)
 			}
 		}
 	}
-	
-
-	// Infer the function's access level from the currently declared C++ access level
-	if (CurrentAccessSpecifier == ACCESS_Public)
-	{
-		FuncInfo.FunctionFlags |= FUNC_Public;
-	}
-	else if (CurrentAccessSpecifier == ACCESS_Protected)
-	{
-		FuncInfo.FunctionFlags |= FUNC_Protected;
-	}
-	else if (CurrentAccessSpecifier == ACCESS_Private)
-	{
-		FuncInfo.FunctionFlags |= FUNC_Private;
-		FuncInfo.FunctionFlags |= FUNC_Final;
-
-		// This is automatically final as well, but in a different way and for a different reason
-		bAutomaticallyFinal = false;
-	}
-	else
-	{
-		FError::Throwf(TEXT("Unknown access level"));
-	}
-
-	// non-static functions in a const class must be const themselves
-	if (Class->HasAnyClassFlags(CLASS_Const))
-	{
-		FuncInfo.FunctionFlags |= FUNC_Const;
-	}
-
-	if (MatchIdentifier(TEXT("static")))
-	{
-		FuncInfo.FunctionFlags |= FUNC_Static;
-		FuncInfo.FunctionExportFlags |= FUNCEXPORT_CppStatic;
-	}
 
 	// Peek ahead to look for a CORE_API style DLL import/export token if present
 	{
@@ -5485,6 +5490,10 @@ void FHeaderParser::CompileFunctionDeclaration(FClasses& AllClasses)
 		{
 			FError::Throwf(TEXT("Interface functions cannot be declared 'final'"));
 		}
+		else if (FuncInfo.FunctionFlags & FUNC_BlueprintEvent)
+		{
+			FError::Throwf(TEXT("Blueprint events cannot be declared 'final'"));
+		}
 	}
 
 	// Make sure any new flags made it to the function
