UnrealHeaderTool: Clarified a couple of error messages that sounded like keywords were limited to only multicast delegate properties (when they are also applicable to functions)

diff --git a/Engine/Source/Programs/UnrealHeaderTool/Private/HeaderParser.cpp b/Engine/Source/Programs/UnrealHeaderTool/Private/HeaderParser.cpp
--- a/Engine/Source/Programs/UnrealHeaderTool/Private/HeaderParser.cpp
+++ b/Engine/Source/Programs/UnrealHeaderTool/Private/HeaderParser.cpp
@@ -3295,12 +3295,12 @@ bool FHeaderParser::GetVarType
 
 	if ((VarProperty.PropertyFlags & CPF_BlueprintCallable) && VarProperty.Type != CPT_MulticastDelegate)
 	{
-		FError::Throwf(TEXT("'BlueprintCallable' is only allowed on multicast delegate properties"));
+		FError::Throwf(TEXT("'BlueprintCallable' is only allowed on a property when it is a multicast delegate"));
 	}
 
 	if ((VarProperty.PropertyFlags & CPF_BlueprintAuthorityOnly) && VarProperty.Type != CPT_MulticastDelegate)
 	{
-		FError::Throwf(TEXT("'BlueprintAuthorityOnly' is only allowed on multicast delegate properties"));
+		FError::Throwf(TEXT("'BlueprintAuthorityOnly' is only allowed on a property when it is a multicast delegate"));
 	}
 	
 	if (VariableCategory != EVariableCategory::Member)
