Removing references to MakeCommandlet from code

diff --git a/Engine/Source/Programs/UnrealHeaderTool/Private/CodeGenerator.cpp b/Engine/Source/Programs/UnrealHeaderTool/Private/CodeGenerator.cpp
--- a/Engine/Source/Programs/UnrealHeaderTool/Private/CodeGenerator.cpp
+++ b/Engine/Source/Programs/UnrealHeaderTool/Private/CodeGenerator.cpp
@@ -179,7 +179,7 @@ void ConvertToBuildIncludePath(UPackage* Package, FString& LocalPath)
  *
  *	@return	bool			true if found, false if not
  */
-bool MakeCommandlet_FindPackageLocation(const TCHAR* InPackage, FString& OutLocation, FString& OutHeaderLocation)
+bool FindPackageLocation(const TCHAR* InPackage, FString& OutLocation, FString& OutHeaderLocation)
 {
 	// Mapping of processed packages to their locations
 	// An empty location string means it was processed but not found
@@ -2019,7 +2019,7 @@ void FNativeClassHeaderGenerator::ExportClassHeaderWrapper( FClass* Class, bool
 
 	FString PkgDir;
 	FString GeneratedIncludeDirectory;
-	if (MakeCommandlet_FindPackageLocation(*PkgName, PkgDir, GeneratedIncludeDirectory) == false)
+	if (FindPackageLocation(*PkgName, PkgDir, GeneratedIncludeDirectory) == false)
 	{
 		UE_LOG(LogCompile, Error, TEXT("Failed to find path for package %s"), *PkgName);
 	}
@@ -3877,7 +3877,7 @@ FNativeClassHeaderGenerator::FNativeClassHeaderGenerator( UPackage* InPackage, F
 	{
 		FString PkgDir;
 		FString GeneratedIncludeDirectory;
-		if (MakeCommandlet_FindPackageLocation(*PackageName, PkgDir, GeneratedIncludeDirectory) == false)
+		if (FindPackageLocation(*PackageName, PkgDir, GeneratedIncludeDirectory) == false)
 		{
 			UE_LOG(LogCompile, Error, TEXT("Failed to find path for package %s"), *PackageName);
 		}
@@ -4221,7 +4221,7 @@ void FNativeClassHeaderGenerator::ExportGeneratedProto()
 		FString PkgName = FPackageName::GetShortName(Package);
 		FString PkgDir;
 		FString GeneratedIncludeDirectory;
-		if (MakeCommandlet_FindPackageLocation(*PkgName, PkgDir, GeneratedIncludeDirectory) == false)
+		if (FindPackageLocation(*PkgName, PkgDir, GeneratedIncludeDirectory) == false)
 		{
 			UE_LOG(LogCompile, Error, TEXT("Failed to find path for package %s"), *PkgName);
 		}
@@ -4254,7 +4254,7 @@ void FNativeClassHeaderGenerator::ExportGeneratedMCP()
 		FString PkgName = FPackageName::GetShortName(Package);
 		FString PkgDir;
 		FString GeneratedIncludeDirectory;
-		if (MakeCommandlet_FindPackageLocation(*PkgName, PkgDir, GeneratedIncludeDirectory) == false)
+		if (FindPackageLocation(*PkgName, PkgDir, GeneratedIncludeDirectory) == false)
 		{
 			UE_LOG(LogCompile, Error, TEXT("Failed to find path for package %s"), *PkgName);
 		}
diff --git a/Engine/Source/Programs/UnrealHeaderTool/Private/HeaderParser.cpp b/Engine/Source/Programs/UnrealHeaderTool/Private/HeaderParser.cpp
--- a/Engine/Source/Programs/UnrealHeaderTool/Private/HeaderParser.cpp
+++ b/Engine/Source/Programs/UnrealHeaderTool/Private/HeaderParser.cpp
@@ -4313,7 +4313,7 @@ void FHeaderParser::CompileClassDeclaration(FClasses& AllClasses)
 		{
 			FError::Throwf(TEXT("The dependsOn specifier is deprecated. Please use proper #include instead."));
 
-			// Make sure the syntax matches but don't do anything with it; that's handled in MakeCommandlet.cpp
+			// Make sure the syntax matches but don't do anything with it;
 			RequireSpecifierValue(PropSpecifier);
 		}
 		else if (Specifier == TEXT("MinimalAPI"))
@@ -4743,7 +4743,7 @@ void FHeaderParser::CompileInterfaceDeclaration(FClasses& AllClasses)
 
 		if (Specifier == TEXT("DependsOn"))
 		{
-			// Make sure the syntax matches but don't do anything with it; that's handled in MakeCommandlet.cpp
+			// Make sure the syntax matches but don't do anything with it
 			RequireSpecifierValue(*SpecifierIt);
 		}
 		else if (Specifier == TEXT("MinimalAPI"))
diff --git a/Engine/Source/Runtime/Engine/Classes/Commandlets/Commandlet.h b/Engine/Source/Runtime/Engine/Classes/Commandlets/Commandlet.h
--- a/Engine/Source/Runtime/Engine/Classes/Commandlets/Commandlet.h
+++ b/Engine/Source/Runtime/Engine/Classes/Commandlets/Commandlet.h
@@ -12,7 +12,7 @@
  * for example:
  *
  *     yourgame.exe Core.HelloWorldCommandlet
- *     yourgame.exe UnrealEd.MakeCommandlet
+ *     yourgame.exe UnrealEd.CookCommandlet
  *
  * As a convenience, if a user tries to run a commandlet and the exact
  * name he types isn't found, then ucc.exe appends the text "commandlet"
diff --git a/Engine/Source/Runtime/Launch/Private/LaunchEngineLoop.cpp b/Engine/Source/Runtime/Launch/Private/LaunchEngineLoop.cpp
--- a/Engine/Source/Runtime/Launch/Private/LaunchEngineLoop.cpp
+++ b/Engine/Source/Runtime/Launch/Private/LaunchEngineLoop.cpp
@@ -952,25 +952,6 @@ int32 FEngineLoop::PreInit( const TCHAR* CmdLine )
 		}
 	}
 
-	if (!bHasEditorToken)
-	{
-		if ((Token == TEXT("MAKE") || Token == TEXT("MAKECOMMANDLET")))
-		{
-			//@todo this is too early to localize
-			const FText Message = NSLOCTEXT("Engine", "MakeRunNoLongerUsed", "The make/run commandlet/arguments are no longer used.");
-			if (!FParse::Param(CmdLine,TEXT("BUILDMACHINE")))
-			{
-				FMessageDialog::Open( EAppMsgType::Ok, Message );
-			}
-			else
-			{
-				UE_LOG(LogInit, Fatal, TEXT("%s"), *Message.ToString());
-			}
-			GGameName[0] = 0; // this disables part of the crash reporter to avoid writing log files to a bogus directory
-			exit(0);
-		}
-	}
-
 	if (GIsUCCMakeStandaloneHeaderGenerator)
 	{
 		// Rebuilding script requires some hacks in the engine so we flag that.
