347339 BLOCKER: CRASH: EDITOR: ORION: Occurs when adding a skylight to a scene.
347343 BLOCKER: EDITOR: ORION: Unable to edit or view the light details properties.

diff --git a/Engine/Source/Programs/UnrealHeaderTool/Private/HeaderParser.cpp b/Engine/Source/Programs/UnrealHeaderTool/Private/HeaderParser.cpp
--- a/Engine/Source/Programs/UnrealHeaderTool/Private/HeaderParser.cpp
+++ b/Engine/Source/Programs/UnrealHeaderTool/Private/HeaderParser.cpp
@@ -3638,6 +3638,11 @@ UProperty* FHeaderParser::GetVarNameAndDim
 				if ( DoesAnythingInHierarchyHaveDefaultToInstanced( VarProperty.PropertyClass ) )
 				{
 					VarProperty.PropertyFlags |= CPF_InstancedReference;
+
+					if (0 != (VarProperty.PropertyFlags & CPF_Edit))
+					{
+						VarProperty.PropertyFlags |= CPF_EditInline;
+					}
 				}
 				NewProperty = new(NewScope,PropertyName,ObjectFlags)UObjectProperty(FPostConstructInitializeProperties());
 			}
