Better error message regarding the use of interfaces as properties, suggesting the use of TScriptInterface instead.

diff --git a/Engine/Source/Programs/UnrealHeaderTool/Private/HeaderParser.cpp b/Engine/Source/Programs/UnrealHeaderTool/Private/HeaderParser.cpp
--- a/Engine/Source/Programs/UnrealHeaderTool/Private/HeaderParser.cpp
+++ b/Engine/Source/Programs/UnrealHeaderTool/Private/HeaderParser.cpp
@@ -2303,7 +2303,8 @@ bool FHeaderParser::GetVarType
 	const TCHAR*                    Thing,
 	FToken*                         OuterPropertyType,
 	EPropertyDeclarationStyle::Type PropertyDeclarationStyle,
-	EVariableCategory::Type         VariableCategory
+	EVariableCategory::Type         VariableCategory,
+	FIndexRange*                    ParsedVarIndexRange
 )
 {
 	check(Scope);
@@ -2371,6 +2372,12 @@ bool FHeaderParser::GetVarType
 		Flags |= CPF_EditorOnly;
 	}
 
+	// Store the start and end positions of the parsed type
+	if (ParsedVarIndexRange)
+	{
+		ParsedVarIndexRange->StartIndex = InputPos;
+	}
+
 	// Process the list of specifiers
 	bool bSeenEditSpecifier = false;
 	bool bSeenBlueprintEditSpecifier = false;
@@ -3324,6 +3331,10 @@ bool FHeaderParser::GetVarType
 
 	VarProperty.MetaData = MetaDataFromNewStyle;
 
+	if (ParsedVarIndexRange)
+	{
+		ParsedVarIndexRange->Count = InputPos - ParsedVarIndexRange->StartIndex;
+	}
 	return 1;
 }
 
@@ -5892,7 +5903,9 @@ void FHeaderParser::CompileVariableDeclaration(FClasses& AllClasses, UStruct* St
 	// Get variable type.
 	FPropertyBase OriginalProperty(CPT_None);
 	EObjectFlags ObjectFlags = RF_NoFlags;
-	GetVarType( AllClasses, Struct, OriginalProperty, ObjectFlags, DisallowFlags, TEXT("Member variable declaration"), /*OuterPropertyType=*/ NULL, PropertyDeclarationStyle, EVariableCategory::Member );
+
+	FIndexRange TypeRange;
+	GetVarType( AllClasses, Struct, OriginalProperty, ObjectFlags, DisallowFlags, TEXT("Member variable declaration"), /*OuterPropertyType=*/ NULL, PropertyDeclarationStyle, EVariableCategory::Member, &TypeRange );
 	OriginalProperty.PropertyFlags |= EdFlags;
 
 	FString* Category = OriginalProperty.MetaData.Find("Category");
@@ -5908,7 +5921,9 @@ void FHeaderParser::CompileVariableDeclaration(FClasses& AllClasses, UStruct* St
 		// Validate that pointer properties are not interfaces (which are not GC'd and so will cause runtime errors)
 		if (OriginalProperty.PointerType == EPointerType::Native && OriginalProperty.Struct->IsChildOf(UInterface::StaticClass()))
 		{
-			FError::Throwf(TEXT("UPROPERTY pointers cannot be interfaces"));
+			// Get the name of the type, removing the asterisk representing the pointer
+			FString TypeName = FString(TypeRange.Count, Input + TypeRange.StartIndex).Trim().TrimTrailing().LeftChop(1).TrimTrailing();
+			FError::Throwf(TEXT("UPROPERTY pointers cannot be interfaces - did you mean TScriptInterface<%s>?"), *TypeName);
 		}
 	}
 
diff --git a/Engine/Source/Programs/UnrealHeaderTool/Private/HeaderParser.h b/Engine/Source/Programs/UnrealHeaderTool/Private/HeaderParser.h
--- a/Engine/Source/Programs/UnrealHeaderTool/Private/HeaderParser.h
+++ b/Engine/Source/Programs/UnrealHeaderTool/Private/HeaderParser.h
@@ -67,6 +67,12 @@ struct FNestInfo
 	int32 Allow;
 };
 
+struct FIndexRange
+{
+	int32 StartIndex;
+	int32 Count;
+};
+
 /////////////////////////////////////////////////////
 // FHeaderParser
 
@@ -392,6 +398,7 @@ protected:
 	 * @param   OuterPropertyType         only specified when compiling the inner properties for arrays or maps.  corresponds to the FToken for the outer property declaration.
 	 * @param   PropertyDeclarationStyle  if the variable is defined with a UPROPERTY
 	 * @param   VariableCategory          what kind of variable is being parsed
+	 * @param   ParsedVarIndexRange       The source text [Start, End) index range for the parsed type.
 	 *
 	 * @return  true if the variable type was parsed
 	 */
@@ -404,7 +411,8 @@ protected:
 		const TCHAR*                    Thing,
 		FToken*                         OuterPropertyType,
 		EPropertyDeclarationStyle::Type PropertyDeclarationStyle,
-		EVariableCategory::Type         VariableCategory);
+		EVariableCategory::Type         VariableCategory,
+		FIndexRange*                    ParsedVarIndexRange = nullptr);
 
 	/**
 	 * Parses a variable name declaration and creates a new UProperty object.
diff --git a/Engine/Source/Programs/UnrealHeaderTool/Resources/UHTDebugging/TestObject.h b/Engine/Source/Programs/UnrealHeaderTool/Resources/UHTDebugging/TestObject.h
--- a/Engine/Source/Programs/UnrealHeaderTool/Resources/UHTDebugging/TestObject.h
+++ b/Engine/Source/Programs/UnrealHeaderTool/Resources/UHTDebugging/TestObject.h
@@ -20,6 +20,11 @@ class UTestObject : public UObject
 	UFUNCTION()
 	void CodeGenTestForEnumClasses(ECppEnum Val);
 
+#if 0
+	UPROPERTY()
+	ITestInterface* ShouldGiveAMeaningfulErrorAboutTScriptInterface;
+#endif
+
 #if BLAH
 	UFUNCTION() int x; // This should not compile if UHT parses it, which it shouldn't
 #elif BLAH2
