Disable deprecation warnings in generated code.

diff --git a/Engine/Source/Programs/UnrealHeaderTool/Private/CodeGenerator.cpp b/Engine/Source/Programs/UnrealHeaderTool/Private/CodeGenerator.cpp
--- a/Engine/Source/Programs/UnrealHeaderTool/Private/CodeGenerator.cpp
+++ b/Engine/Source/Programs/UnrealHeaderTool/Private/CodeGenerator.cpp
@@ -1981,9 +1981,42 @@ void FNativeClassHeaderGenerator::ExportClassHeaderWrapper( FClass* Class, bool
 		{
 			GeneratedHeaderText.Logf(TEXT("#undef GENERATED_UCLASS_BODY\r\n"));
 			GeneratedHeaderText.Logf(TEXT("#undef GENERATED_IINTERFACE_BODY\r\n"));
+			auto ClangDeprecationPushString =
+				TEXT("_Pragma (\"clang diagnostic push\")") LINE_TERMINATOR
+				TEXT("_Pragma (\"clang diagnostic ignored \\\"-Wdeprecated-declarations\\\"\")") LINE_TERMINATOR;
+			auto MSVCDeprecationPushString =
+				TEXT("__pragma (warning(push))") LINE_TERMINATOR
+				TEXT("__pragma (warning(disable:4995))") LINE_TERMINATOR
+				TEXT("__pragma (warning(disable:4996))") LINE_TERMINATOR;
+			auto UnknownCompilerDeprecationPushString = 
+				TEXT("#error Unsupported compiler") LINE_TERMINATOR;
+
+			auto ClangDeprecationPopString =
+				TEXT("_Pragma(\"clang diagnostic pop\")") LINE_TERMINATOR;
+			auto MSVCDeprecationPopString =
+				TEXT("__pragma(warning(pop))") LINE_TERMINATOR;
+			auto UnknownCompilerDeprecationPopString =
+				TEXT("#error Unsupported compiler") LINE_TERMINATOR;
+
+			auto Public = TEXT("public:" LINE_TERMINATOR);
+			auto WrappedInClassMacroCalls = Public + InClassMacroCalls + Public;
+
 			FString MacroName = FString::Printf(TEXT("GENERATED_%s_BODY()"), Class->HasAnyClassFlags(CLASS_Interface) ? TEXT("IINTERFACE") : TEXT("UCLASS"));
-			FString Macroized = Macroize(*MacroName, *(FString(TEXT("public:\r\n")) + InClassMacroCalls + TEXT("public:\r\n")));
-			GeneratedHeaderText.Log(*Macroized);
+			FString MacroizedMSVC =
+				FString(TEXT("#if _MSC_VER") LINE_TERMINATOR) +
+				Macroize(*MacroName, *(FString(MSVCDeprecationPushString) + WrappedInClassMacroCalls + MSVCDeprecationPopString));
+			FString MacroizedClang = 
+				FString(TEXT("#elif __clang__") LINE_TERMINATOR) +
+				Macroize(*MacroName, *(FString(ClangDeprecationPushString) + WrappedInClassMacroCalls + ClangDeprecationPopString));
+			FString MacroizedUnknownCompiler =
+				FString(TEXT("#else") LINE_TERMINATOR) +
+				Macroize(*MacroName, *(FString(UnknownCompilerDeprecationPushString) + WrappedInClassMacroCalls + UnknownCompilerDeprecationPopString)) +
+				FString(TEXT("#endif") LINE_TERMINATOR);
+
+			GeneratedHeaderText.Log(*MacroizedMSVC);
+			GeneratedHeaderText.Log(*MacroizedClang);
+			GeneratedHeaderText.Log(*MacroizedUnknownCompiler);
+
 			InClassMacroCalls.Empty();
 		}
 	}
