#ue4 - Fix HeaderParser so it no longer removes valid uses of "*", " * ", etc from within comments when generating tooltips. Only asterisks from the beginning of lines following only whitespace are removed.

- For example: "A * B" within a java-doc comment used to become "A B".

diff --git a/Engine/Source/Programs/UnrealHeaderTool/Private/HeaderParser.cpp b/Engine/Source/Programs/UnrealHeaderTool/Private/HeaderParser.cpp
--- a/Engine/Source/Programs/UnrealHeaderTool/Private/HeaderParser.cpp
+++ b/Engine/Source/Programs/UnrealHeaderTool/Private/HeaderParser.cpp
@@ -896,16 +896,6 @@ FString FHeaderParser::FormatCommentForToolTip(const FString& Input)
 		Result = Result.Replace( TEXT("*/"), TEXT("") );
 	}
 
-	if ( bJavaDocStyle )
-	{
-		// Remove stars from left edge.
-		Result = Result.Replace( TEXT(" * "), TEXT(" ") );
-		Result = Result.Replace( TEXT(" *\t"), TEXT("\t") );
-		Result = Result.Replace( TEXT(" *"), TEXT("") );
-		//Result = Result.Replace( TEXT("* "), TEXT("") );
-		//Result = Result.Replace( TEXT("*"), TEXT("") );
-	}
-
 	if ( bCPPStyle )
 	{
 		// Remove c++-style comment markers.  Also handle javadoc-style comments by replacing
@@ -931,10 +921,33 @@ FString FHeaderParser::FormatCommentForToolTip(const FString& Input)
 	TArray<FString> Lines;
 	Result.ParseIntoArray(&Lines, TEXT("\n"), false);
 
-	// Remove trailing whitespace
 	for (auto& Line : Lines)
 	{
+		// Remove trailing whitespace
 		Line.TrimTrailing();
+
+		// Remove leading "*" and "* " in javadoc comments.
+		if (bJavaDocStyle)
+		{
+			// Find first non-whitespace character
+			int32 Pos = 0;
+			while (Pos < Line.Len() && FChar::IsWhitespace(Line[Pos]))
+			{
+				++Pos;
+			}
+
+			// Is it a *?
+			if (Pos < Line.Len() && Line[Pos] == '*')
+			{
+				// Eat next space as well
+				if (Pos+1 < Line.Len() && FChar::IsWhitespace(Line[Pos+1]))
+				{
+					++Pos;
+				}
+
+				Line = Line.RightChop(Pos + 1);
+			}
+		}
 	}
 
 	// Find first meaningful line
