#ue4 Fix for NotBlueprintType specifier not removing BlueprintType specifier, causing conflicting metadata

diff --git a/Engine/Source/Programs/UnrealHeaderTool/Private/BaseParser.cpp b/Engine/Source/Programs/UnrealHeaderTool/Private/BaseParser.cpp
--- a/Engine/Source/Programs/UnrealHeaderTool/Private/BaseParser.cpp
+++ b/Engine/Source/Programs/UnrealHeaderTool/Private/BaseParser.cpp
@@ -93,6 +93,7 @@ namespace
 
 			FMetadataKeyword& NotBlueprintType = Dictionary.Add(TEXT("NotBlueprintType"), EMetadataValueArgument::None);
 			NotBlueprintType.InsertAddAction(TEXT("NotBlueprintType"), TEXT("true"));
+			NotBlueprintType.InsertRemoveAction(TEXT("BlueprintType"));
 
 			FMetadataKeyword& Blueprintable = Dictionary.Add(TEXT("Blueprintable"), EMetadataValueArgument::None);
 			Blueprintable.InsertAddAction(TEXT("IsBlueprintBase"), TEXT("true"));
