UObject custom constructors.

In order to use new feature user have to change GENERATED_UCLASS_BODY() macro to the new GENERATED_BODY(). Then no constructor is implicitly declared. If there is no constructor declared in UCLASS then PCIP one is declared and defined that passes PCIP down to super-class. On the other hand if there is a constructor declared then UObject system expects to have one of the default or PCIP constructor to initialize a class during loading from disk, default object creation, etc. and it expects that user will declare and define one.

There is a possibility now to create UCLASS with the default constructor (i.e. no PCIP).

New macro is encouraged by standard classes and in-editor wizards templates.

diff --git a/Engine/Source/Editor/GameProjectGeneration/Private/GameProjectUtils.cpp b/Engine/Source/Editor/GameProjectGeneration/Private/GameProjectUtils.cpp
--- a/Engine/Source/Editor/GameProjectGeneration/Private/GameProjectUtils.cpp
+++ b/Engine/Source/Editor/GameProjectGeneration/Private/GameProjectUtils.cpp
@@ -1708,7 +1708,7 @@ bool GameProjectUtils::GenerateGameFrameworkSourceCode(const FString& NewProject
 		const FString NewClassName = NewProjectName + BaseClass->GetName();
 		const FString NewHeaderFilename = GameModulePath / NewClassName + TEXT(".h");
 		FString UnusedSyncLocation;
-		if ( GenerateClassHeaderFile(NewHeaderFilename, NewClassName, FNewClassInfo(BaseClass), TArray<FString>(), TEXT(""), TEXT(""), UnusedSyncLocation, NewModuleInfo, OutFailReason) )
+		if ( GenerateClassHeaderFile(NewHeaderFilename, NewClassName, FNewClassInfo(BaseClass), TArray<FString>(), TEXT(""), TEXT(""), UnusedSyncLocation, NewModuleInfo, false, OutFailReason) )
 		{
 			OutCreatedFiles.Add(NewHeaderFilename);
 		}
@@ -2173,7 +2173,53 @@ FString GameProjectUtils::MakeIncludeList(const TArray<FString>& InList)
 	return ReturnString;
 }
 
-bool GameProjectUtils::GenerateClassHeaderFile(const FString& NewHeaderFileName, const FString UnPrefixedClassName, const FNewClassInfo ParentClassInfo, const TArray<FString>& ClassSpecifierList, const FString& ClassProperties, const FString& ClassFunctionDeclarations, FString& OutSyncLocation, const FModuleContextInfo& ModuleInfo, FText& OutFailReason)
+/**
+ * Generates UObject class constructor definition with property overrides.
+ *
+ * @param Out String to assign generated constructor to.
+ * @param PrefixedClassName Prefixed class name for which we generate the constructor.
+ * @param PropertyOverridesStr String with property overrides in the constructor.
+ * @param OutFailReason Template read function failure reason.
+ *
+ * @returns True on success. False otherwise.
+ */
+bool GenerateConstructorDefinition(FString& Out, const FString& PrefixedClassName, const FString& PropertyOverridesStr, FText& OutFailReason)
+{
+	FString Template;
+	if (!GameProjectUtils::ReadTemplateFile(TEXT("UObjectClassConstructorDefinition.template"), Template, OutFailReason))
+	{
+		return false;
+	}
+
+	Out = Template.Replace(TEXT("%PREFIXED_CLASS_NAME%"), *PrefixedClassName, ESearchCase::CaseSensitive);
+	Out = Out.Replace(TEXT("%PROPERTY_OVERRIDES%"), *PropertyOverridesStr, ESearchCase::CaseSensitive);
+
+	return true;
+}
+
+/**
+ * Generates UObject class constructor declaration.
+ *
+ * @param Out String to assign generated constructor to.
+ * @param PrefixedClassName Prefixed class name for which we generate the constructor.
+ * @param OutFailReason Template read function failure reason.
+ *
+ * @returns True on success. False otherwise.
+ */
+bool GenerateConstructorDeclaration(FString& Out, const FString& PrefixedClassName, FText& OutFailReason)
+{
+	FString Template;
+	if (!GameProjectUtils::ReadTemplateFile(TEXT("UObjectClassConstructorDeclaration.template"), Template, OutFailReason))
+	{
+		return false;
+	}
+
+	Out = Template.Replace(TEXT("%PREFIXED_CLASS_NAME%"), *PrefixedClassName, ESearchCase::CaseSensitive);
+
+	return true;
+}
+
+bool GameProjectUtils::GenerateClassHeaderFile(const FString& NewHeaderFileName, const FString UnPrefixedClassName, const FNewClassInfo ParentClassInfo, const TArray<FString>& ClassSpecifierList, const FString& ClassProperties, const FString& ClassFunctionDeclarations, FString& OutSyncLocation, const FModuleContextInfo& ModuleInfo, bool bDeclareConstructor, FText& OutFailReason)
 {
 	FString Template;
 	if ( !ReadTemplateFile(ParentClassInfo.GetHeaderTemplateFilename(), Template, OutFailReason) )
@@ -2205,6 +2251,15 @@ bool GameProjectUtils::GenerateClassHeaderFile(const FString& NewHeaderFileName,
 		}
 	}
 
+	FString EventualConstructorDeclaration;
+	if (bDeclareConstructor)
+	{
+		if (!GenerateConstructorDeclaration(EventualConstructorDeclaration, PrefixedClassName, OutFailReason))
+		{
+			return false;
+		}
+	}
+
 	// Not all of these will exist in every class template
 	FString FinalOutput = Template.Replace(TEXT("%COPYRIGHT_LINE%"), *MakeCopyrightLine(), ESearchCase::CaseSensitive);
 	FinalOutput = FinalOutput.Replace(TEXT("%UNPREFIXED_CLASS_NAME%"), *UnPrefixedClassName, ESearchCase::CaseSensitive);
@@ -2212,6 +2267,7 @@ bool GameProjectUtils::GenerateClassHeaderFile(const FString& NewHeaderFileName,
 	FinalOutput = FinalOutput.Replace(TEXT("%UCLASS_SPECIFIER_LIST%"), *MakeCommaDelimitedList(ClassSpecifierList, false), ESearchCase::CaseSensitive);
 	FinalOutput = FinalOutput.Replace(TEXT("%PREFIXED_CLASS_NAME%"), *PrefixedClassName, ESearchCase::CaseSensitive);
 	FinalOutput = FinalOutput.Replace(TEXT("%PREFIXED_BASE_CLASS_NAME%"), *PrefixedBaseClassName, ESearchCase::CaseSensitive);
+	FinalOutput = FinalOutput.Replace(TEXT("%EVENTUAL_CONSTRUCTOR_DECLARATION%"), *EventualConstructorDeclaration, ESearchCase::CaseSensitive);
 	FinalOutput = FinalOutput.Replace(TEXT("%CLASS_PROPERTIES%"), *ClassProperties, ESearchCase::CaseSensitive);
 	FinalOutput = FinalOutput.Replace(TEXT("%CLASS_FUNCTION_DECLARATIONS%"), *ClassFunctionDeclarations, ESearchCase::CaseSensitive);
 	FinalOutput = FinalOutput.Replace(TEXT("%BASE_CLASS_INCLUDE_DIRECTIVE%"), *BaseClassIncludeDirective, ESearchCase::CaseSensitive);
@@ -2331,13 +2387,22 @@ bool GameProjectUtils::GenerateClassCPPFile(const FString& NewCPPFileName, const
 		ModuleIncludePath = ModuleInfo.ModuleName + ".h";
 	}
 
+	FString EventualConstructorDefinition;
+	if (PropertyOverrides.Num() != 0)
+	{
+		if (!GenerateConstructorDefinition(EventualConstructorDefinition, PrefixedClassName, PropertyOverridesStr, OutFailReason))
+		{
+			return false;
+		}
+	}
+
 	// Not all of these will exist in every class template
 	FString FinalOutput = Template.Replace(TEXT("%COPYRIGHT_LINE%"), *MakeCopyrightLine(), ESearchCase::CaseSensitive);
 	FinalOutput = FinalOutput.Replace(TEXT("%UNPREFIXED_CLASS_NAME%"), *UnPrefixedClassName, ESearchCase::CaseSensitive);
 	FinalOutput = FinalOutput.Replace(TEXT("%MODULE_NAME%"), *ModuleInfo.ModuleName, ESearchCase::CaseSensitive);
 	FinalOutput = FinalOutput.Replace(TEXT("%MODULE_INCLUDE_PATH%"), *ModuleIncludePath, ESearchCase::CaseSensitive);
 	FinalOutput = FinalOutput.Replace(TEXT("%PREFIXED_CLASS_NAME%"), *PrefixedClassName, ESearchCase::CaseSensitive);
-	FinalOutput = FinalOutput.Replace(TEXT("%PROPERTY_OVERRIDES%"), *PropertyOverridesStr, ESearchCase::CaseSensitive);
+	FinalOutput = FinalOutput.Replace(TEXT("%EVENTUAL_CONSTRUCTOR_DEFINITION%"), *EventualConstructorDefinition, ESearchCase::CaseSensitive);
 	FinalOutput = FinalOutput.Replace(TEXT("%ADDITIONAL_MEMBER_DEFINITIONS%"), *AdditionalMemberDefinitions, ESearchCase::CaseSensitive);
 	FinalOutput = FinalOutput.Replace(TEXT("%ADDITIONAL_INCLUDE_DIRECTIVES%"), *AdditionalIncludesStr, ESearchCase::CaseSensitive);
 
@@ -2739,7 +2804,7 @@ bool GameProjectUtils::AddCodeToProject_Internal(const FString& NewClassName, co
 	FString SyncLocation;
 	const FString NewHeaderFilename = NewHeaderPath / ParentClassInfo.GetHeaderFilename(NewClassName);
 	{
-		if ( GenerateClassHeaderFile(NewHeaderFilename, CleanClassName, ParentClassInfo, TArray<FString>(), TEXT(""), TEXT(""), SyncLocation, ModuleInfo, OutFailReason) )
+		if ( GenerateClassHeaderFile(NewHeaderFilename, CleanClassName, ParentClassInfo, TArray<FString>(), TEXT(""), TEXT(""), SyncLocation, ModuleInfo, false, OutFailReason) )
 		{
 			CreatedFiles.Add(NewHeaderFilename);
 		}
diff --git a/Engine/Source/Editor/GameProjectGeneration/Private/GameProjectUtils.h b/Engine/Source/Editor/GameProjectGeneration/Private/GameProjectUtils.h
--- a/Engine/Source/Editor/GameProjectGeneration/Private/GameProjectUtils.h
+++ b/Engine/Source/Editor/GameProjectGeneration/Private/GameProjectUtils.h
@@ -273,6 +273,8 @@ public:
 	/** Clear the list of supported target platforms */
 	static void ClearSupportedTargetPlatforms();
 
+	/** Returns the contents of the specified template file */
+	static bool ReadTemplateFile(const FString& TemplateFileName, FString& OutFileContents, FText& OutFailReason);
 private:
 
 	static FString GetHardwareConfigString(const FProjectInformation& InProjectInfo);
@@ -331,9 +333,6 @@ private:
 	/** Creates the batch file for launching the editor or game */
 	static bool GenerateLaunchBatchFile(const FString& ProjectName, const FString& ProjectFolder, bool bLaunchEditor, TArray<FString>& OutCreatedFiles, FText& OutFailReason);
 
-	/** Returns the contents of the specified template file */
-	static bool ReadTemplateFile(const FString& TemplateFileName, FString& OutFileContents, FText& OutFailReason);
-
 	/** Writes an output file. OutputFilename includes a path */
 	static bool WriteOutputFile(const FString& OutputFilename, const FString& OutputFileContents, FText& OutFailReason);
 
@@ -347,7 +346,7 @@ private:
 	static FString MakeIncludeList(const TArray<FString>& InList);
 
 	/** Generates a header file for a UObject class. OutSyncLocation is a string representing the preferred cursor sync location for this file after creation. */
-	static bool GenerateClassHeaderFile(const FString& NewHeaderFileName, const FString UnPrefixedClassName, const FNewClassInfo ParentClassInfo, const TArray<FString>& ClassSpecifierList, const FString& ClassProperties, const FString& ClassFunctionDeclarations, FString& OutSyncLocation, const FModuleContextInfo& ModuleInfo, FText& OutFailReason);
+	static bool GenerateClassHeaderFile(const FString& NewHeaderFileName, const FString UnPrefixedClassName, const FNewClassInfo ParentClassInfo, const TArray<FString>& ClassSpecifierList, const FString& ClassProperties, const FString& ClassFunctionDeclarations, FString& OutSyncLocation, const FModuleContextInfo& ModuleInfo, bool bDeclareConstructor, FText& OutFailReason);
 
 	/** Generates a cpp file for a UObject class */
 	static bool GenerateClassCPPFile(const FString& NewCPPFileName, const FString UnPrefixedClassName, const FNewClassInfo ParentClassInfo, const TArray<FString>& AdditionalIncludes, const TArray<FString>& PropertyOverrides, const FString& AdditionalMemberDefinitions, const FModuleContextInfo& ModuleInfo, FText& OutFailReason);
diff --git a/Engine/Source/Programs/UnrealHeaderTool/Private/CodeGenerator.cpp b/Engine/Source/Programs/UnrealHeaderTool/Private/CodeGenerator.cpp
--- a/Engine/Source/Programs/UnrealHeaderTool/Private/CodeGenerator.cpp
+++ b/Engine/Source/Programs/UnrealHeaderTool/Private/CodeGenerator.cpp
@@ -1633,6 +1633,20 @@ void FNativeClassHeaderGenerator::ExportNames()
 
 }
 
+/**
+ * Gets preprocessor string to emit GENERATED_U*_BODY() macro is deprecated.
+ *
+ * @param MacroName Name of the macro to be deprecated.
+ *
+ * @returns Preprocessor string to emit the message.
+ */
+FString GetGeneratedMacroDeprecationWarning(const TCHAR* MacroName)
+{
+	// Deprecation warning is disabled right now. After people get familiar with the new macro it should be re-enabled.
+	//return FString() + TEXT("EMIT_DEPRECATED_WARNING_MESSAGE(\"") + MacroName + TEXT("() macro is deprecated. Please use GENERATED_BODY() macro instead.\")") LINE_TERMINATOR;
+	return TEXT("");
+}
+
 /**
  * Exports the C++ class declarations for a native interface class.
  */
@@ -1707,10 +1721,7 @@ void FNativeClassHeaderGenerator::ExportInterfaceClassDeclaration( FClass* Class
 			bCastedClass ? *FString::Printf(TEXT("CASTCLASS_%s"), ClassCPPName) : TEXT("0"),
 			*FPackageName::GetShortName(Class->GetOuter()->GetName()), 
 			*APIArg );
-		if (!Class->HasAnyClassFlags(CLASS_CustomConstructor))
-		{
-			UInterfaceBoilerplate.Logf(TEXT("    %s_API %s(const class FPostConstructInitializeProperties& PCIP);\r\n"), *APIArg, ClassCPPName);
-		}
+		
 		UInterfaceBoilerplate.Logf(TEXT("    DECLARE_SERIALIZER(%s)\r\n"), ClassCPPName);
 		UInterfaceBoilerplate.Log(TEXT("    enum {IsIntrinsic=COMPILED_IN_INTRINSIC};\r\n"));
 
@@ -1719,10 +1730,31 @@ void FNativeClassHeaderGenerator::ExportInterfaceClassDeclaration( FClass* Class
 			UInterfaceBoilerplate.Logf(TEXT("    DECLARE_WITHIN(%s)\r\n"), NameLookupCPP.GetNameCPP( Class->GetClassWithin() ) );
 		}
 
-		FString MacroName = TEXT("GENERATED_UINTERFACE_BODY()");
-		FString Macroized = Macroize(*MacroName, *UInterfaceBoilerplate);
+		ExportConstructorsMacros(Class);
+
+		GeneratedHeaderText.Log(TEXT("#undef GENERATED_UINTERFACE_BODY_COMMON\r\n"));
+		GeneratedHeaderText.Log(Macroize(TEXT("GENERATED_UINTERFACE_BODY_COMMON()"), *UInterfaceBoilerplate));
+
+		auto DeprecationWarning = GetGeneratedMacroDeprecationWarning(TEXT("GENERATED_UINTERFACE_BODY"));
+		
+		auto DeprecationPushString = TEXT("PRAGMA_DISABLE_DEPRECATION_WARNINGS") LINE_TERMINATOR;
+		auto DeprecationPopString = TEXT("PRAGMA_ENABLE_DEPRECATION_WARNINGS") LINE_TERMINATOR;
+		auto Offset = FCString::Spc(4);
+
 		GeneratedHeaderText.Log(TEXT("#undef GENERATED_UINTERFACE_BODY\r\n"));
-		GeneratedHeaderText.Log(*Macroized);
+		GeneratedHeaderText.Log(Macroize(TEXT("GENERATED_UINTERFACE_BODY()"), *(FString() +
+			Offset + DeprecationWarning +
+			Offset + DeprecationPushString +
+			Offset + TEXT("GENERATED_UINTERFACE_BODY_COMMON()") LINE_TERMINATOR +
+			StandardUObjectConstructorsMacroCall +
+			Offset + DeprecationPopString)));
+
+		GeneratedHeaderText.Log(TEXT("#undef GENERATED_BODY\r\n"));
+		GeneratedHeaderText.Log(Macroize(TEXT("GENERATED_BODY()"), *(FString() +
+			Offset + DeprecationPushString +
+			Offset + TEXT("GENERATED_UINTERFACE_BODY_COMMON()") LINE_TERMINATOR +
+			EnhancedUObjectConstructorsMacroCall +
+			Offset + DeprecationPopString)));
 	}
 
 	// =============================================
@@ -1898,11 +1930,7 @@ void FNativeClassHeaderGenerator::ExportClassHeaderInner(FClass* Class, bool bVa
 					bCastedClass ? *FString::Printf(TEXT("CASTCLASS_%s"), ClassCPPName) : TEXT("0"),
 					*FPackageName::GetShortName(*Class->GetOuter()->GetName()), 
 					*APIArg );
-				if (!Class->HasAnyClassFlags(CLASS_CustomConstructor))
-				{
-					ClassBoilerplate.Logf(TEXT("    /** Standard constructor, called after all reflected properties have been initialized */"));
-					ClassBoilerplate.Logf(TEXT("    %s_API %s(const class FPostConstructInitializeProperties& PCIP);\r\n"), *APIArg, ClassCPPName);
-				}
+
 				ClassBoilerplate.Logf(TEXT("    DECLARE_SERIALIZER(%s)\r\n"), ClassCPPName);
 				ClassBoilerplate.Log(TEXT("    /** Indicates whether the class is compiled into the engine */"));
 				ClassBoilerplate.Log(TEXT("    enum {IsIntrinsic=COMPILED_IN_INTRINSIC};\r\n"));
@@ -1935,7 +1963,9 @@ void FNativeClassHeaderGenerator::ExportClassHeaderInner(FClass* Class, bool bVa
 				FString MacroName = FString(NameLookupCPP.GetNameCPP(Class)) + Suffix;
 				FString Macroized = Macroize(*MacroName, *ClassBoilerplate);
 				GeneratedHeaderText.Log(*Macroized);
-				InClassMacroCalls.Logf( TEXT("%s%s\r\n"), FCString::Spc(4),  *MacroName);
+				InClassMacroCalls.Logf(TEXT("%s%s\r\n"), FCString::Spc(4), *MacroName);
+
+				ExportConstructorsMacros(Class);
 			}
 		}
 	}
@@ -1946,6 +1976,104 @@ void FNativeClassHeaderGenerator::ExportClassHeaderInner(FClass* Class, bool bVa
 	}
 }
 
+/**
+ * Generates standard constructor declaration.
+ *
+ * @param Out Output device to generate to.
+ * @param Class Class to generate constructor for.
+ * @param API API string for this constructor.
+ */
+void ExportStandardConstructorsMacro(FStringOutputDevice& Out, FClass* Class, const FString& API)
+{
+	if (!Class->HasAnyClassFlags(CLASS_CustomConstructor))
+	{
+		Out.Logf(TEXT("%s/** Standard constructor, called after all reflected properties have been initialized */\r\n"), FCString::Spc(4));
+		Out.Logf(TEXT("%s%s_API %s(const class FPostConstructInitializeProperties& PCIP);\r\n"), FCString::Spc(4), *API, NameLookupCPP.GetNameCPP(Class));
+	}
+
+	Out.Logf(TEXT("%sDEFINE_DEFAULT_PCIP_CONSTRUCTOR_CALL(%s)\r\n"), FCString::Spc(4), NameLookupCPP.GetNameCPP(Class));
+}
+
+/**
+ * Generates constructor definition.
+ *
+ * @param Out Output device to generate to.
+ * @param Class Class to generate constructor for.
+ * @param API API string for this constructor.
+ */
+void ExportConstructorDefinition(FStringOutputDevice& Out, FClass* Class, const FString& API)
+{
+	auto* ClassData = GScriptHelper.FindClassData(Class);
+
+	if (!ClassData->bConstructorDeclared)
+	{
+		Out.Logf(TEXT("%s/** Standard constructor, called after all reflected properties have been initialized */"), FCString::Spc(4));
+		Out.Logf(TEXT("%s%s_API %s(const class FPostConstructInitializeProperties& PCIP) : Super(PCIP) { };\r\n"), FCString::Spc(4), *API, NameLookupCPP.GetNameCPP(Class));
+
+		ClassData->bConstructorDeclared = true;
+		ClassData->bPCIPConstructorDeclared = true;
+	}
+}
+
+/**
+ * Generates constructor call definition.
+ *
+ * @param Out Output device to generate to.
+ * @param Class Class to generate constructor call definition for.
+ */
+void ExportDefaultConstructorCallDefinition(FStringOutputDevice& Out, FClass* Class)
+{
+	auto* ClassData = GScriptHelper.FindClassData(Class);
+
+	if (ClassData->bPCIPConstructorDeclared)
+	{
+		Out.Logf(TEXT("%sDEFINE_DEFAULT_PCIP_CONSTRUCTOR_CALL(%s)\r\n"), FCString::Spc(4), NameLookupCPP.GetNameCPP(Class));
+	}
+	else if (ClassData->bDefaultConstructorDeclared)
+	{
+		Out.Logf(TEXT("%sDEFINE_DEFAULT_CONSTRUCTOR_CALL(%s)\r\n"), FCString::Spc(4), NameLookupCPP.GetNameCPP(Class));
+	}
+	else
+	{
+		Out.Logf(TEXT("%sDEFINE_FORBIDDEN_DEFAULT_CONSTRUCTOR_CALL(%s)\r\n"), FCString::Spc(4), NameLookupCPP.GetNameCPP(Class));
+	}
+}
+
+/**
+ * Generates enhanced constructor declaration.
+ *
+ * @param Out Output device to generate to.
+ * @param Class Class to generate constructor for.
+ * @param API API string for this constructor.
+ */
+void ExportEnhancedConstructorsMacro(FStringOutputDevice& Out, FClass* Class, const FString& API)
+{
+	ExportConstructorDefinition(Out, Class, API);
+	ExportDefaultConstructorCallDefinition(Out, Class);
+}
+
+void FNativeClassHeaderGenerator::ExportConstructorsMacros(FClass* Class)
+{
+	FString APIArg(API);
+	if (!Class->HasAnyClassFlags(CLASS_MinimalAPI))
+	{
+		APIArg = TEXT("NO");
+	}
+
+	FStringOutputDevice StdMacro, EnhMacro;
+	FString StdMacroName = FString(NameLookupCPP.GetNameCPP(Class)) + TEXT("_STANDARD_CONSTRUCTORS");
+	FString EnhMacroName = FString(NameLookupCPP.GetNameCPP(Class)) + TEXT("_ENHANCED_CONSTRUCTORS");
+
+	ExportStandardConstructorsMacro(StdMacro, Class, APIArg);
+	ExportEnhancedConstructorsMacro(EnhMacro, Class, APIArg);
+
+	GeneratedHeaderText.Log(*Macroize(*StdMacroName, *StdMacro));
+	GeneratedHeaderText.Log(*Macroize(*EnhMacroName, *EnhMacro));
+
+	StandardUObjectConstructorsMacroCall.Logf(TEXT("%s%s\r\n"), FCString::Spc(4), *StdMacroName);
+	EnhancedUObjectConstructorsMacroCall.Logf(TEXT("%s%s\r\n"), FCString::Spc(4), *EnhMacroName);
+}
+
 void FNativeClassHeaderGenerator::ExportClassHeaderWrapper( FClass* Class, bool bIsExportClass )
 {
 	const bool bTemporaryJunkClass = Class->HasAnyClassFlags(CLASS_Temporary);
@@ -1979,45 +2107,46 @@ void FNativeClassHeaderGenerator::ExportClassHeaderWrapper( FClass* Class, bool
 		}
 
 		{
+			bool bIsIInterface = Class->HasAnyClassFlags(CLASS_Interface);
+
 			GeneratedHeaderText.Logf(TEXT("#undef GENERATED_UCLASS_BODY\r\n"));
+			
+			if (!bIsIInterface)
+			{
+				GeneratedHeaderText.Logf(TEXT("#undef GENERATED_BODY\r\n"));
+			}
+
 			GeneratedHeaderText.Logf(TEXT("#undef GENERATED_IINTERFACE_BODY\r\n"));
-			auto ClangDeprecationPushString =
-				TEXT("_Pragma (\"clang diagnostic push\")") LINE_TERMINATOR
-				TEXT("_Pragma (\"clang diagnostic ignored \\\"-Wdeprecated-declarations\\\"\")") LINE_TERMINATOR;
-			auto MSVCDeprecationPushString =
-				TEXT("__pragma (warning(push))") LINE_TERMINATOR
-				TEXT("__pragma (warning(disable:4995))") LINE_TERMINATOR
-				TEXT("__pragma (warning(disable:4996))") LINE_TERMINATOR;
-			auto UnknownCompilerDeprecationPushString = 
-				TEXT("#error Unsupported compiler") LINE_TERMINATOR;
-
-			auto ClangDeprecationPopString =
-				TEXT("_Pragma(\"clang diagnostic pop\")") LINE_TERMINATOR;
-			auto MSVCDeprecationPopString =
-				TEXT("__pragma(warning(pop))") LINE_TERMINATOR;
-			auto UnknownCompilerDeprecationPopString =
-				TEXT("#error Unsupported compiler") LINE_TERMINATOR;
+
+			FString MacroName = FString::Printf(TEXT("GENERATED_%s_BODY()"), bIsIInterface ? TEXT("IINTERFACE") : TEXT("UCLASS"));
+
+			auto DeprecationWarning = bIsIInterface ? FString(TEXT("")) : GetGeneratedMacroDeprecationWarning(*MacroName);
+
+			auto DeprecationPushString = TEXT("PRAGMA_DISABLE_DEPRECATION_WARNINGS") LINE_TERMINATOR;
+			auto DeprecationPopString = TEXT("PRAGMA_ENABLE_DEPRECATION_WARNINGS") LINE_TERMINATOR;
 
 			auto Public = TEXT("public:" LINE_TERMINATOR);
-			auto WrappedInClassMacroCalls = Public + InClassMacroCalls + Public;
-
-			FString MacroName = FString::Printf(TEXT("GENERATED_%s_BODY()"), Class->HasAnyClassFlags(CLASS_Interface) ? TEXT("IINTERFACE") : TEXT("UCLASS"));
-			FString MacroizedMSVC =
-				FString(TEXT("#if _MSC_VER") LINE_TERMINATOR) +
-				Macroize(*MacroName, *(FString(MSVCDeprecationPushString) + WrappedInClassMacroCalls + MSVCDeprecationPopString));
-			FString MacroizedClang = 
-				FString(TEXT("#elif __clang__") LINE_TERMINATOR) +
-				Macroize(*MacroName, *(FString(ClangDeprecationPushString) + WrappedInClassMacroCalls + ClangDeprecationPopString));
-			FString MacroizedUnknownCompiler =
-				FString(TEXT("#else") LINE_TERMINATOR) +
-				Macroize(*MacroName, *(FString(UnknownCompilerDeprecationPushString) + WrappedInClassMacroCalls + UnknownCompilerDeprecationPopString)) +
-				FString(TEXT("#endif") LINE_TERMINATOR);
-
-			GeneratedHeaderText.Log(*MacroizedMSVC);
-			GeneratedHeaderText.Log(*MacroizedClang);
-			GeneratedHeaderText.Log(*MacroizedUnknownCompiler);
+
+			FString BodyMacros;
+			
+			if (bIsIInterface)
+			{
+				BodyMacros = Macroize(*MacroName, *(DeprecationPushString + FString(Public) + InClassMacroCalls + Public + DeprecationPopString));
+			}
+			else
+			{
+				auto StandardWrappedInClassMacroCalls = Public + InClassMacroCalls + StandardUObjectConstructorsMacroCall + Public;
+				auto EnhancedWrappedInClassMacroCalls = Public + InClassMacroCalls + EnhancedUObjectConstructorsMacroCall + Public;
+
+				BodyMacros = Macroize(*MacroName, *(DeprecationWarning + DeprecationPushString + StandardWrappedInClassMacroCalls + DeprecationPopString)) +
+					Macroize(TEXT("GENERATED_BODY()"), *(FString(DeprecationPushString) + EnhancedWrappedInClassMacroCalls + DeprecationPopString));
+			}
+
+			GeneratedHeaderText.Log(*BodyMacros);
 
 			InClassMacroCalls.Empty();
+			StandardUObjectConstructorsMacroCall.Empty();
+			EnhancedUObjectConstructorsMacroCall.Empty();
 		}
 	}
 
diff --git a/Engine/Source/Programs/UnrealHeaderTool/Private/HeaderParser.cpp b/Engine/Source/Programs/UnrealHeaderTool/Private/HeaderParser.cpp
--- a/Engine/Source/Programs/UnrealHeaderTool/Private/HeaderParser.cpp
+++ b/Engine/Source/Programs/UnrealHeaderTool/Private/HeaderParser.cpp
@@ -3828,6 +3828,24 @@ bool FHeaderParser::CompileDeclaration( FClasses& AllClasses, FToken& Token )
 		bHaveSeenSecondInterfaceClass = true;
 		ParseSecondInterfaceClass(AllClasses);
 	}
+	else if (Token.Matches(TEXT("GENERATED_BODY")) && (TopNest->NestType != NEST_Interface
+		|| (bHaveSeenFirstInterfaceClass && !bHaveSeenSecondInterfaceClass)))
+	{
+		if (TopNest->NestType != NEST_Class && TopNest->NestType != NEST_Interface)
+		{
+			FError::Throwf(TEXT("%s must occur inside the class or interface definition"), Token.Identifier);
+		}
+
+		RequireSymbol(TEXT("("), Token.Identifier);
+		RequireSymbol(TEXT(")"), Token.Identifier);
+
+		CurrentAccessSpecifier = ACCESS_Public;
+
+		if (TopNest->NestType == NEST_Class)
+		{
+			bClassHasGeneratedBody = true;
+		}
+	}
 	else if (Token.Matches(TEXT("GENERATED_UCLASS_BODY")))
 	{
 		if (TopNest->NestType != NEST_Class)
@@ -3960,6 +3978,13 @@ bool FHeaderParser::CompileDeclaration( FClasses& AllClasses, FToken& Token )
 			FError::Throwf(TEXT("Extra ';' before end of file") );
 		}
 	}
+	else if (Token.Matches(NameLookupCPP.GetNameCPP(Class)))
+	{
+		if(!TryToMatchConstructorParameterList(Token))
+		{
+			return SkipDeclaration(Token);
+		}
+	}
 	else
 	{
 		// Ignore C++ declaration / function definition. 
@@ -7472,3 +7497,103 @@ bool FHeaderParser::DefaultValueStringCppFormatToInnerFormat(const UProperty* Pr
 
 	return !OutForm.IsEmpty();
 }
+
+bool FHeaderParser::TryToMatchConstructorParameterList(FToken Token)
+{
+	FToken PotentialParenthesisToken;
+	if (!GetToken(PotentialParenthesisToken))
+	{
+		return false;
+	}
+
+	if (!PotentialParenthesisToken.Matches(TEXT("(")))
+	{
+		UngetToken(PotentialParenthesisToken);
+		return false;
+	}
+
+	auto* ClassData = GScriptHelper.FindClassData(Class);
+
+	ClassData->bConstructorDeclared = true;
+
+	if (!ClassData->bDefaultConstructorDeclared && MatchSymbol(TEXT(")")))
+	{
+		ClassData->bDefaultConstructorDeclared = true;
+	}
+	else if (!ClassData->bPCIPConstructorDeclared)
+	{
+		FToken PCIPParamParsingToken;
+
+		bool bIsConst = false;
+		bool bIsRef = false;
+		int32 ParenthesesNestingLevel = 1;
+
+		while (ParenthesesNestingLevel && GetToken(PCIPParamParsingToken))
+		{
+			// Template instantiation or additional parameter excludes PCIP constructor.
+			if (PCIPParamParsingToken.Matches(TEXT(",")) || PCIPParamParsingToken.Matches(TEXT("<")))
+			{
+				ClassData->bPCIPConstructorDeclared = false;
+				break;
+			}
+
+			if (PCIPParamParsingToken.Matches(TEXT("(")))
+			{
+				ParenthesesNestingLevel++;
+				continue;
+			}
+
+			if (PCIPParamParsingToken.Matches(TEXT(")")))
+			{
+				ParenthesesNestingLevel--;
+				continue;
+			}
+
+			if (PCIPParamParsingToken.Matches(TEXT("const")))
+			{
+				bIsConst = true;
+				continue;
+			}
+
+			if (PCIPParamParsingToken.Matches(TEXT("&")))
+			{
+				bIsRef = true;
+				continue;
+			}
+
+			if (PCIPParamParsingToken.Matches(TEXT("FPostConstructInitializeProperties")))
+			{
+				ClassData->bPCIPConstructorDeclared = true;
+			}
+		}
+
+		// Parse until finish.
+		while (ParenthesesNestingLevel && GetToken(PCIPParamParsingToken))
+		{
+			if (PCIPParamParsingToken.Matches(TEXT("(")))
+			{
+				ParenthesesNestingLevel++;
+				continue;
+			}
+
+			if (PCIPParamParsingToken.Matches(TEXT(")")))
+			{
+				ParenthesesNestingLevel--;
+				continue;
+			}
+		}
+
+		ClassData->bPCIPConstructorDeclared = ClassData->bPCIPConstructorDeclared && bIsRef && bIsConst;
+	}
+
+	// Optionally match semicolon.
+	if (!MatchSymbol(TEXT(";")))
+	{
+		// If not matched a semicolon, this is inline constructor definition. We have to skip it.
+		UngetToken(Token); // Resets input stream to the initial token.
+		GetToken(Token); // Re-gets the initial token to start constructor definition skip.
+		return SkipDeclaration(Token);
+	}
+
+	return true;
+}
diff --git a/Engine/Source/Programs/UnrealHeaderTool/Private/HeaderParser.h b/Engine/Source/Programs/UnrealHeaderTool/Private/HeaderParser.h
--- a/Engine/Source/Programs/UnrealHeaderTool/Private/HeaderParser.h
+++ b/Engine/Source/Programs/UnrealHeaderTool/Private/HeaderParser.h
@@ -221,7 +221,7 @@ protected:
 	// Indicates that UCLASS/USTRUCT/UINTERFACE has already been parsed in this .h file..
 	bool bHaveSeenUClass;
 
-	// Indicates that a GENERATED_UCLASS_BODY has been found in the UClass.
+	// Indicates that a GENERATED_UCLASS_BODY or GENERATED_BODY has been found in the UClass.
 	bool bClassHasGeneratedBody;
 
 	// public, private, etc at the current parse spot
@@ -315,6 +315,7 @@ protected:
 	UEnum* CompileEnum(UClass* Owner);
 	UScriptStruct* CompileStructDeclaration(FClasses& AllClasses, FClass* Owner);
 	bool CompileDeclaration(FClasses& AllClasses, FToken& Token);
+
 	/** Skip C++ (noexport) declaration. */
 	bool SkipDeclaration(FToken& Token);
 	/** Similar to MatchSymbol() but will return to the exact location as on entry if the symbol was not found. */
@@ -543,6 +544,19 @@ protected:
 	}
 
 	static void ValidatePropertyIsDeprecatedIfNecessary(FPropertyBase& VarProperty, FToken* OuterPropertyType);
+
+private:
+	/**
+	 * Tries to match constructor parameter list. Assumes that constructor
+	 * name is already matched.
+	 *
+	 * If fails it reverts all parsing done.
+	 *
+	 * @param Token Token to start parsing from.
+	 *
+	 * @returns True if matched. False otherwise.
+	 */
+	bool TryToMatchConstructorParameterList(FToken Token);
 };
 
 /////////////////////////////////////////////////////
diff --git a/Engine/Source/Programs/UnrealHeaderTool/Private/NativeClassExporter.h b/Engine/Source/Programs/UnrealHeaderTool/Private/NativeClassExporter.h
--- a/Engine/Source/Programs/UnrealHeaderTool/Private/NativeClassExporter.h
+++ b/Engine/Source/Programs/UnrealHeaderTool/Private/NativeClassExporter.h
@@ -61,6 +61,8 @@ private:
 	FStringOutputDevice GeneratedMCPText;
 	FStringOutputDevice	PrologMacroCalls;
 	FStringOutputDevice	InClassMacroCalls;
+	FStringOutputDevice	StandardUObjectConstructorsMacroCall;
+	FStringOutputDevice	EnhancedUObjectConstructorsMacroCall;
 	FStringOutputDevice	AllConstructors;
 
 	/** Generated function implementations that belong in the cpp file, split into multiple files base on line count **/
@@ -463,6 +465,13 @@ private:
 	 */
 	void DeleteUnusedGeneratedHeaders();
 
+	/**
+	 * Exports macros that manages UObject constructors.
+	 * 
+	 * @param Class Class for which to export macros.
+	 */
+	void ExportConstructorsMacros(FClass* Class);
+
 public:
 
 	// Constructor
diff --git a/Engine/Source/Programs/UnrealHeaderTool/Private/ParserHelper.h b/Engine/Source/Programs/UnrealHeaderTool/Private/ParserHelper.h
--- a/Engine/Source/Programs/UnrealHeaderTool/Private/ParserHelper.h
+++ b/Engine/Source/Programs/UnrealHeaderTool/Private/ParserHelper.h
@@ -1307,7 +1307,10 @@ class FClassMetaData
 public:
 	/** Default constructor */
 	FClassMetaData()
-	: bContainsDelegates(false)
+		: bContainsDelegates(false),
+		bConstructorDeclared(false),
+		bDefaultConstructorDeclared(false),
+		bPCIPConstructorDeclared(false)
 	{
 	}
 
@@ -1546,6 +1549,15 @@ public:
 		FunctionData.Shrink();
 		MultipleInheritanceParents.Shrink();
 	}
+
+	// Is constructor declared?
+	bool bConstructorDeclared;
+
+	// Is default constructor declared?
+	bool bDefaultConstructorDeclared;
+
+	// Is PCIP constructor (i.e. a constructor with only one parameter of type FPostConstructInitializeProperties) declared?
+	bool bPCIPConstructorDeclared;
 };
 
 /**
diff --git a/Engine/Source/Runtime/CoreUObject/Classes/Object.h b/Engine/Source/Runtime/CoreUObject/Classes/Object.h
--- a/Engine/Source/Runtime/CoreUObject/Classes/Object.h
+++ b/Engine/Source/Runtime/CoreUObject/Classes/Object.h
@@ -837,7 +837,12 @@ struct FInt32Range
 UCLASS(abstract, noexport)
 class UObject
 {
-	GENERATED_UCLASS_BODY()
+	GENERATED_BODY()
+
+	/**
+	 * Default UObject constructor.
+	 */
+	UObject(const class FPostConstructInitializeProperties& PCIP);
 
 	//=============================================================================
 	// K2 support functions.
diff --git a/Engine/Source/Runtime/CoreUObject/Public/UObject/Class.h b/Engine/Source/Runtime/CoreUObject/Public/UObject/Class.h
--- a/Engine/Source/Runtime/CoreUObject/Public/UObject/Class.h
+++ b/Engine/Source/Runtime/CoreUObject/Public/UObject/Class.h
@@ -2188,7 +2188,7 @@ protected:
 template<class T>
 void InternalConstructor( const class FPostConstructInitializeProperties& X )
 { 
-	new( (EInternal*)X.Obj )T(X); 
+	T::__DefaultConstructor(X);
 }
 
 COREUOBJECT_API void InitializePrivateStaticClass(
diff --git a/Engine/Source/Runtime/CoreUObject/Public/UObject/ObjectBase.h b/Engine/Source/Runtime/CoreUObject/Public/UObject/ObjectBase.h
--- a/Engine/Source/Runtime/CoreUObject/Public/UObject/ObjectBase.h
+++ b/Engine/Source/Runtime/CoreUObject/Public/UObject/ObjectBase.h
@@ -495,6 +495,25 @@ struct COREUOBJECT_API FReferencerInformationList
 // This macro is used to declare a thunk function in autogenerated boilerplate code
 #define DECLARE_FUNCTION(func) void func( FFrame& Stack, RESULT_DECL )
 
+#if _MSC_VER
+	#define PRAGMA_DISABLE_DEPRECATION_WARNINGS \
+		__pragma (warning(push)) \
+		__pragma (warning(disable:4995)) \
+		__pragma (warning(disable:4996))
+
+	#define PRAGMA_ENABLE_DEPRECATION_WARNINGS \
+		__pragma(warning(pop))
+#elif __clang__
+	#define PRAGMA_DISABLE_DEPRECATION_WARNINGS \
+		_Pragma ("clang diagnostic push") \
+		_Pragma ("clang diagnostic ignored \"-Wdeprecated-declarations\"")
+
+	#define PRAGMA_ENABLE_DEPRECATION_WARNINGS \
+		_Pragma("clang diagnostic pop")
+#else
+	#error Unsupported compiler
+#endif
+
 // These are used for syntax highlighting and to allow autocomplete hints
 
 namespace UC
@@ -1155,28 +1174,40 @@ public: \
 		return (void*)InMem; \
 	}
 
+#define DEFINE_FORBIDDEN_DEFAULT_CONSTRUCTOR_CALL(TClass) \
+	static_assert(false, "You have to define " #TClass "::" #TClass "() or " #TClass "::" #TClass "(const FPostConstructInitializeProperties&). This is required by UObject system to work correctly.");
+
+#define DEFINE_DEFAULT_CONSTRUCTOR_CALL(TClass) \
+	static void __DefaultConstructor(const class FPostConstructInitializeProperties& X) { const_cast<FPostConstructInitializeProperties&>(X).FinalizeSubobjectClassInitialization(); new((EInternal*)X.GetObject())TClass(); }
+
+#define DEFINE_DEFAULT_PCIP_CONSTRUCTOR_CALL(TClass) \
+	static void __DefaultConstructor(const class FPostConstructInitializeProperties& X) { new((EInternal*)X.GetObject())TClass(X); }
+
 #define DECLARE_CLASS_INTRINSIC(TClass,TSuperClass,TStaticFlags,TPackage) \
 	DECLARE_CLASS(TClass,TSuperClass,TStaticFlags|CLASS_Intrinsic,CASTCLASS_None,TPackage,NO_API ) \
 	RELAY_CONSTRUCTOR(TClass, TSuperClass) \
 	enum {IsIntrinsic=1}; \
 	static void StaticRegisterNatives##TClass() {} \
-	DECLARE_SERIALIZER(TClass)
+	DECLARE_SERIALIZER(TClass) \
+	DEFINE_DEFAULT_PCIP_CONSTRUCTOR_CALL(TClass)
 
 #define DECLARE_CASTED_CLASS_INTRINSIC_WITH_API( TClass, TSuperClass, TStaticFlags, TPackage, TStaticCastFlags, TRequiredAPI ) \
 	DECLARE_CLASS(TClass,TSuperClass,TStaticFlags|CLASS_Intrinsic,TStaticCastFlags,TPackage,TRequiredAPI ) \
 	RELAY_CONSTRUCTOR(TClass, TSuperClass) \
 	enum {IsIntrinsic=1}; \
 	static void StaticRegisterNatives##TClass() {} \
-	DECLARE_SERIALIZER(TClass)
+	DECLARE_SERIALIZER(TClass) \
+	DEFINE_DEFAULT_PCIP_CONSTRUCTOR_CALL(TClass)
 
 #define DECLARE_CASTED_CLASS_INTRINSIC( TClass, TSuperClass, TStaticFlags, TPackage, TStaticCastFlags ) \
-	DECLARE_CASTED_CLASS_INTRINSIC_WITH_API( TClass, TSuperClass, TStaticFlags, TPackage, TStaticCastFlags, NO_API)
+	DECLARE_CASTED_CLASS_INTRINSIC_WITH_API( TClass, TSuperClass, TStaticFlags, TPackage, TStaticCastFlags, NO_API) \
 
 #define DECLARE_CASTED_CLASS_INTRINSIC_NO_CTOR( TClass, TSuperClass, TStaticFlags, TPackage, TStaticCastFlags, TRequiredAPI ) \
 	DECLARE_CLASS(TClass,TSuperClass,TStaticFlags|CLASS_Intrinsic,TStaticCastFlags,TPackage, TRequiredAPI ) \
 	enum {IsIntrinsic=1}; \
 	static void StaticRegisterNatives##TClass() {} \
-	DECLARE_SERIALIZER(TClass)
+	DECLARE_SERIALIZER(TClass) \
+	DEFINE_DEFAULT_PCIP_CONSTRUCTOR_CALL(TClass)
 
 
 
diff --git a/Engine/Source/Runtime/CoreUObject/Public/UObject/UObject.h b/Engine/Source/Runtime/CoreUObject/Public/UObject/UObject.h
--- a/Engine/Source/Runtime/CoreUObject/Public/UObject/UObject.h
+++ b/Engine/Source/Runtime/CoreUObject/Public/UObject/UObject.h
@@ -37,11 +37,13 @@ class COREUOBJECT_API UObject : public UObjectBaseUtility
 {
 	// Declarations.
 	DECLARE_CLASS(UObject,UObject,CLASS_Abstract|CLASS_NoExport,CASTCLASS_None,CoreUObject,NO_API)
+	DEFINE_DEFAULT_PCIP_CONSTRUCTOR_CALL(UObject)
 
 	typedef UObject WithinClass;
 	static const TCHAR* StaticConfigName() {return TEXT("Engine");}
 
 	// Constructors and destructors.
+	UObject() { };
 	UObject(const class FPostConstructInitializeProperties& PCIP);
 	UObject( EStaticConstructor, EObjectFlags InFlags );
 
diff --git a/Engine/Source/Runtime/Engine/Classes/Components/ActorComponent.h b/Engine/Source/Runtime/Engine/Classes/Components/ActorComponent.h
--- a/Engine/Source/Runtime/Engine/Classes/Components/ActorComponent.h
+++ b/Engine/Source/Runtime/Engine/Classes/Components/ActorComponent.h
@@ -11,7 +11,12 @@ struct FReplicationFlags;
 UCLASS(DefaultToInstanced, abstract, hidecategories=(ComponentReplication))
 class ENGINE_API UActorComponent : public UObject, public IInterface_AssetUserData
 {
-	GENERATED_UCLASS_BODY()
+	GENERATED_BODY()
+
+	/**
+	 * Default UObject constructor.
+	 */
+	UActorComponent(const class FPostConstructInitializeProperties& PCIP);
 
 	/** Main tick function for the Actor */
 	UPROPERTY()
diff --git a/Engine/Source/Runtime/Engine/Classes/Components/PrimitiveComponent.h b/Engine/Source/Runtime/Engine/Classes/Components/PrimitiveComponent.h
--- a/Engine/Source/Runtime/Engine/Classes/Components/PrimitiveComponent.h
+++ b/Engine/Source/Runtime/Engine/Classes/Components/PrimitiveComponent.h
@@ -101,7 +101,12 @@ DECLARE_DYNAMIC_MULTICAST_DELEGATE_TwoParams( FComponentEndTouchOverSignature, E
 UCLASS(abstract, HideCategories=(Mobility), ShowCategories=(PhysicsVolume))
 class ENGINE_API UPrimitiveComponent : public USceneComponent, public INavRelevantInterface
 {
-	GENERATED_UCLASS_BODY()
+	GENERATED_BODY()
+
+	/**
+	 * Default UObject constructor.
+	 */
+	UPrimitiveComponent(const class FPostConstructInitializeProperties& PCIP);
 
 	// Rendering
 	
diff --git a/Engine/Source/Runtime/Engine/Classes/Components/SceneComponent.h b/Engine/Source/Runtime/Engine/Classes/Components/SceneComponent.h
--- a/Engine/Source/Runtime/Engine/Classes/Components/SceneComponent.h
+++ b/Engine/Source/Runtime/Engine/Classes/Components/SceneComponent.h
@@ -79,7 +79,12 @@ DECLARE_DYNAMIC_MULTICAST_DELEGATE_OneParam(FPhysicsVolumeChanged, class APhysic
 UCLASS(ClassGroup=Utility, BlueprintType, HideCategories=(Trigger, PhysicsVolume), meta=(BlueprintSpawnableComponent))
 class ENGINE_API USceneComponent : public UActorComponent
 {
-	GENERATED_UCLASS_BODY()
+	GENERATED_BODY()
+
+	/**
+	 * Default UObject constructor.
+	 */
+	USceneComponent(const class FPostConstructInitializeProperties& PCIP);
 
 	/** Current transform of this component, relative to the world */
 	FTransform ComponentToWorld;
diff --git a/Engine/Source/Runtime/Engine/Classes/GameFramework/Actor.h b/Engine/Source/Runtime/Engine/Classes/GameFramework/Actor.h
--- a/Engine/Source/Runtime/Engine/Classes/GameFramework/Actor.h
+++ b/Engine/Source/Runtime/Engine/Classes/GameFramework/Actor.h
@@ -66,7 +66,12 @@ class ENGINE_API AActor : public UObject
 	* AActor::BeginPlay - Called when the level is started
 	*/
 
-	GENERATED_UCLASS_BODY()
+	GENERATED_BODY()
+
+	/**
+	 * Default UObject constructor.
+	 */
+	AActor(const class FPostConstructInitializeProperties& PCIP);
 
 public:
 
diff --git a/Engine/Source/Runtime/Engine/Classes/GameFramework/Character.h b/Engine/Source/Runtime/Engine/Classes/GameFramework/Character.h
--- a/Engine/Source/Runtime/Engine/Classes/GameFramework/Character.h
+++ b/Engine/Source/Runtime/Engine/Classes/GameFramework/Character.h
@@ -199,7 +199,12 @@ typedef FBasedMovementInfo FRepRelativeMovement;
 UCLASS(abstract, config=Game, BlueprintType, hidecategories=("Pawn|Character|InternalEvents"))
 class ENGINE_API ACharacter : public APawn
 {
-	GENERATED_UCLASS_BODY()
+	GENERATED_BODY()
+
+	/**
+	 * Default UObject constructor.
+	 */
+	ACharacter(const class FPostConstructInitializeProperties& PCIP);
 
 	/** The main skeletal mesh associated with this Character (optional sub-object). */
 	UPROPERTY(Category=Character, VisibleAnywhere, BlueprintReadOnly)
diff --git a/Engine/Source/Runtime/Engine/Classes/GameFramework/CharacterMovementComponent.h b/Engine/Source/Runtime/Engine/Classes/GameFramework/CharacterMovementComponent.h
--- a/Engine/Source/Runtime/Engine/Classes/GameFramework/CharacterMovementComponent.h
+++ b/Engine/Source/Runtime/Engine/Classes/GameFramework/CharacterMovementComponent.h
@@ -124,7 +124,12 @@ struct FCharacterMovementComponentPreClothTickFunction : public FTickFunction
 UCLASS()
 class ENGINE_API UCharacterMovementComponent : public UPawnMovementComponent, public INetworkPredictionInterface
 {
-	GENERATED_UCLASS_BODY()
+	GENERATED_BODY()
+
+	/**
+	 * Default UObject constructor.
+	 */
+	UCharacterMovementComponent(const class FPostConstructInitializeProperties& PCIP);
 
 protected:
 
diff --git a/Engine/Source/Runtime/Engine/Classes/GameFramework/Pawn.h b/Engine/Source/Runtime/Engine/Classes/GameFramework/Pawn.h
--- a/Engine/Source/Runtime/Engine/Classes/GameFramework/Pawn.h
+++ b/Engine/Source/Runtime/Engine/Classes/GameFramework/Pawn.h
@@ -29,7 +29,12 @@ class UPrimitiveComponent;
 UCLASS(abstract, config=Game, BlueprintType, Blueprintable, hidecategories=(Navigation, "AI|Navigation"))
 class ENGINE_API APawn : public AActor, public INavAgentInterface
 {
-	GENERATED_UCLASS_BODY()
+	GENERATED_BODY()
+
+	/**
+	 * Default UObject constructor.
+	 */
+	APawn(const class FPostConstructInitializeProperties& PCIP);
 
 	/** Return our PawnMovementComponent, if we have one. By default, returns the first PawnMovementComponent found. Native classes that create their own movement component should override this method for more efficiency. */
 	UFUNCTION(BlueprintCallable, meta=(Tooltip="Return our PawnMovementComponent, if we have one."), Category="Pawn")
diff --git a/Engine/Source/Runtime/Engine/Classes/GameFramework/PawnMovementComponent.h b/Engine/Source/Runtime/Engine/Classes/GameFramework/PawnMovementComponent.h
--- a/Engine/Source/Runtime/Engine/Classes/GameFramework/PawnMovementComponent.h
+++ b/Engine/Source/Runtime/Engine/Classes/GameFramework/PawnMovementComponent.h
@@ -16,7 +16,7 @@
 UCLASS(abstract)
 class ENGINE_API UPawnMovementComponent : public UNavMovementComponent
 {
-	GENERATED_UCLASS_BODY()
+	GENERATED_BODY()
 
 public:
 
diff --git a/Engine/Source/Runtime/Engine/Private/PawnMovementComponent.cpp b/Engine/Source/Runtime/Engine/Private/PawnMovementComponent.cpp
--- a/Engine/Source/Runtime/Engine/Private/PawnMovementComponent.cpp
+++ b/Engine/Source/Runtime/Engine/Private/PawnMovementComponent.cpp
@@ -9,11 +9,6 @@ DEFINE_LOG_CATEGORY_STATIC(LogPawnMovementComponent, Log, All);
 //----------------------------------------------------------------------//
 // UPawnMovementComponent
 //----------------------------------------------------------------------//
-UPawnMovementComponent::UPawnMovementComponent(const class FPostConstructInitializeProperties& PCIP)
-	: Super(PCIP)
-{
-}
-
 void UPawnMovementComponent::SetUpdatedComponent(UPrimitiveComponent* NewUpdatedComponent)
 {
 	if (NewUpdatedComponent)
