DonÆt default to æpublicÆ protection level in classes.

UECORE-8

diff --git a/Engine/Source/Programs/UnrealHeaderTool/Private/ClassMaps.h b/Engine/Source/Programs/UnrealHeaderTool/Private/ClassMaps.h
--- a/Engine/Source/Programs/UnrealHeaderTool/Private/ClassMaps.h
+++ b/Engine/Source/Programs/UnrealHeaderTool/Private/ClassMaps.h
@@ -39,3 +39,13 @@ extern TMap<UProperty*, FString>                GArrayDimensions;
 extern TMap<UPackage*,  const FManifestModule*> GPackageToManifestModuleMap;
 extern TMap<UField*, uint32>                    GGeneratedCodeCRCs;
 extern TMap<UEnum*,  EPropertyType>             GEnumUnderlyingTypes;
+
+/** Types access specifiers. */
+enum EAccessSpecifier
+{
+	ACCESS_NotAnAccessSpecifier = 0,
+	ACCESS_Public,
+	ACCESS_Private,
+	ACCESS_Protected,
+	ACCESS_Num,
+};
diff --git a/Engine/Source/Programs/UnrealHeaderTool/Private/CodeGenerator.cpp b/Engine/Source/Programs/UnrealHeaderTool/Private/CodeGenerator.cpp
--- a/Engine/Source/Programs/UnrealHeaderTool/Private/CodeGenerator.cpp
+++ b/Engine/Source/Programs/UnrealHeaderTool/Private/CodeGenerator.cpp
@@ -1647,6 +1647,33 @@ FString GetGeneratedMacroDeprecationWarning(const TCHAR* MacroName)
 	return TEXT("");
 }
 
+/**
+ * Returns a string with access specifier that was met before parsing GENERATED_BODY() macro to preserve it.
+ *
+ * @param Class Class for which to return the access specifier.
+ *
+ * @returns Access specifier string.
+ */
+FString GetPreservedAccessSpecifierString(FClass* Class)
+{
+	FString PreservedAccessSpecifier = "static_assert(false, \"Access specifier for GENERATED_BODY() macro is wrong.\");";
+
+	switch (GScriptHelper.FindClassData(Class)->GeneratedBodyMacroAccessSpecifier)
+	{
+	case EAccessSpecifier::ACCESS_Private:
+		PreservedAccessSpecifier = "private:";
+		break;
+	case EAccessSpecifier::ACCESS_Protected:
+		PreservedAccessSpecifier = "protected:";
+		break;
+	case EAccessSpecifier::ACCESS_Public:
+		PreservedAccessSpecifier = "public:";
+		break;
+	}
+
+	return PreservedAccessSpecifier + LINE_TERMINATOR;
+}
+
 /**
  * Exports the C++ class declarations for a native interface class.
  */
@@ -1754,6 +1781,7 @@ void FNativeClassHeaderGenerator::ExportInterfaceClassDeclaration( FClass* Class
 			Offset + DeprecationPushString +
 			Offset + TEXT("GENERATED_UINTERFACE_BODY_COMMON()") LINE_TERMINATOR +
 			EnhancedUObjectConstructorsMacroCall +
+			GetPreservedAccessSpecifierString(Class) +
 			Offset + DeprecationPopString)));
 	}
 
@@ -2136,7 +2164,7 @@ void FNativeClassHeaderGenerator::ExportClassHeaderWrapper( FClass* Class, bool
 			else
 			{
 				auto StandardWrappedInClassMacroCalls = Public + InClassMacroCalls + StandardUObjectConstructorsMacroCall + Public;
-				auto EnhancedWrappedInClassMacroCalls = Public + InClassMacroCalls + EnhancedUObjectConstructorsMacroCall + Public;
+				auto EnhancedWrappedInClassMacroCalls = Public + InClassMacroCalls + EnhancedUObjectConstructorsMacroCall + GetPreservedAccessSpecifierString(Class);
 
 				BodyMacros = Macroize(*MacroName, *(DeprecationWarning + DeprecationPushString + StandardWrappedInClassMacroCalls + DeprecationPopString)) +
 					Macroize(TEXT("GENERATED_BODY()"), *(FString(DeprecationPushString) + EnhancedWrappedInClassMacroCalls + DeprecationPopString));
diff --git a/Engine/Source/Programs/UnrealHeaderTool/Private/HeaderParser.cpp b/Engine/Source/Programs/UnrealHeaderTool/Private/HeaderParser.cpp
--- a/Engine/Source/Programs/UnrealHeaderTool/Private/HeaderParser.cpp
+++ b/Engine/Source/Programs/UnrealHeaderTool/Private/HeaderParser.cpp
@@ -3839,7 +3839,7 @@ bool FHeaderParser::CompileDeclaration( FClasses& AllClasses, FToken& Token )
 		RequireSymbol(TEXT("("), Token.Identifier);
 		RequireSymbol(TEXT(")"), Token.Identifier);
 
-		CurrentAccessSpecifier = ACCESS_Public;
+		GScriptHelper.FindClassData(Class)->GeneratedBodyMacroAccessSpecifier = CurrentAccessSpecifier;
 
 		if (TopNest->NestType == NEST_Class)
 		{
diff --git a/Engine/Source/Programs/UnrealHeaderTool/Private/HeaderParser.h b/Engine/Source/Programs/UnrealHeaderTool/Private/HeaderParser.h
--- a/Engine/Source/Programs/UnrealHeaderTool/Private/HeaderParser.h
+++ b/Engine/Source/Programs/UnrealHeaderTool/Private/HeaderParser.h
@@ -35,16 +35,6 @@ enum ENestAllowFlags
 	ALLOW_TypeDecl			= 16,	// Allow declarations which do not affect memory layout, such as structs, enums, and consts
 };
 
-/** Types access specifiers. */
-enum EAccessSpecifier
-{
-	ACCESS_NotAnAccessSpecifier	= 0,	
-	ACCESS_Public,
-	ACCESS_Private,
-	ACCESS_Protected,
-	ACCESS_Num,
-};
-
 namespace EDelegateSpecifierAction
 {
 	enum Type
diff --git a/Engine/Source/Programs/UnrealHeaderTool/Private/ParserHelper.h b/Engine/Source/Programs/UnrealHeaderTool/Private/ParserHelper.h
--- a/Engine/Source/Programs/UnrealHeaderTool/Private/ParserHelper.h
+++ b/Engine/Source/Programs/UnrealHeaderTool/Private/ParserHelper.h
@@ -1558,6 +1558,9 @@ public:
 
 	// Is PCIP constructor (i.e. a constructor with only one parameter of type FPostConstructInitializeProperties) declared?
 	bool bPCIPConstructorDeclared;
+
+	// GENERATED_BODY access specifier to preserve.
+	EAccessSpecifier GeneratedBodyMacroAccessSpecifier;
 };
 
 /**
diff --git a/Engine/Source/Runtime/CoreUObject/Classes/Object.h b/Engine/Source/Runtime/CoreUObject/Classes/Object.h
--- a/Engine/Source/Runtime/CoreUObject/Classes/Object.h
+++ b/Engine/Source/Runtime/CoreUObject/Classes/Object.h
@@ -838,6 +838,7 @@ UCLASS(abstract, noexport)
 class UObject
 {
 	GENERATED_BODY()
+public:
 
 	/**
 	 * Default UObject constructor.
diff --git a/Engine/Source/Runtime/Engine/Classes/Components/ActorComponent.h b/Engine/Source/Runtime/Engine/Classes/Components/ActorComponent.h
--- a/Engine/Source/Runtime/Engine/Classes/Components/ActorComponent.h
+++ b/Engine/Source/Runtime/Engine/Classes/Components/ActorComponent.h
@@ -12,6 +12,7 @@ UCLASS(DefaultToInstanced, abstract, hidecategories=(ComponentReplication))
 class ENGINE_API UActorComponent : public UObject, public IInterface_AssetUserData
 {
 	GENERATED_BODY()
+public:
 
 	/**
 	 * Default UObject constructor.
diff --git a/Engine/Source/Runtime/Engine/Classes/Components/PrimitiveComponent.h b/Engine/Source/Runtime/Engine/Classes/Components/PrimitiveComponent.h
--- a/Engine/Source/Runtime/Engine/Classes/Components/PrimitiveComponent.h
+++ b/Engine/Source/Runtime/Engine/Classes/Components/PrimitiveComponent.h
@@ -102,6 +102,7 @@ UCLASS(abstract, HideCategories=(Mobility), ShowCategories=(PhysicsVolume))
 class ENGINE_API UPrimitiveComponent : public USceneComponent, public INavRelevantInterface
 {
 	GENERATED_BODY()
+public:
 
 	/**
 	 * Default UObject constructor.
diff --git a/Engine/Source/Runtime/Engine/Classes/Components/SceneComponent.h b/Engine/Source/Runtime/Engine/Classes/Components/SceneComponent.h
--- a/Engine/Source/Runtime/Engine/Classes/Components/SceneComponent.h
+++ b/Engine/Source/Runtime/Engine/Classes/Components/SceneComponent.h
@@ -80,6 +80,7 @@ UCLASS(ClassGroup=Utility, BlueprintType, HideCategories=(Trigger, PhysicsVolume
 class ENGINE_API USceneComponent : public UActorComponent
 {
 	GENERATED_BODY()
+public:
 
 	/**
 	 * Default UObject constructor.
diff --git a/Engine/Source/Runtime/Engine/Classes/GameFramework/Character.h b/Engine/Source/Runtime/Engine/Classes/GameFramework/Character.h
--- a/Engine/Source/Runtime/Engine/Classes/GameFramework/Character.h
+++ b/Engine/Source/Runtime/Engine/Classes/GameFramework/Character.h
@@ -200,6 +200,7 @@ UCLASS(abstract, config=Game, BlueprintType, hidecategories=("Pawn|Character|Int
 class ENGINE_API ACharacter : public APawn
 {
 	GENERATED_BODY()
+public:
 
 	/**
 	 * Default UObject constructor.
diff --git a/Engine/Source/Runtime/Engine/Classes/GameFramework/CharacterMovementComponent.h b/Engine/Source/Runtime/Engine/Classes/GameFramework/CharacterMovementComponent.h
--- a/Engine/Source/Runtime/Engine/Classes/GameFramework/CharacterMovementComponent.h
+++ b/Engine/Source/Runtime/Engine/Classes/GameFramework/CharacterMovementComponent.h
@@ -125,6 +125,7 @@ UCLASS()
 class ENGINE_API UCharacterMovementComponent : public UPawnMovementComponent, public INetworkPredictionInterface
 {
 	GENERATED_BODY()
+public:
 
 	/**
 	 * Default UObject constructor.
