NUXF-11 Don't require a Category on UPROPERTIES (for game classes)

TTP 346443

diff --git a/Engine/Source/Programs/UnrealHeaderTool/Private/HeaderParser.cpp b/Engine/Source/Programs/UnrealHeaderTool/Private/HeaderParser.cpp
--- a/Engine/Source/Programs/UnrealHeaderTool/Private/HeaderParser.cpp
+++ b/Engine/Source/Programs/UnrealHeaderTool/Private/HeaderParser.cpp
@@ -5958,7 +5958,18 @@ void FHeaderParser::CompileVariableDeclaration(FClasses& AllClasses, UStruct* St
 		// First check if the category was specified at all and if the property was exposed to the editor.
 		if (!Category && (OriginalProperty.PropertyFlags & (CPF_Edit|CPF_BlueprintVisible)))
 		{
-			FError::Throwf(TEXT("Property is exposed to the editor or blueprints but has no Category specified."));
+			static const FString AbsoluteEngineDir = FPaths::ConvertRelativePathToFull(FPaths::EngineDir());
+			FString SourceFilename = GClassSourceFileMap[Class];
+			FPaths::NormalizeFilename(SourceFilename);
+			if (Struct->GetOutermost() != nullptr && !SourceFilename.StartsWith(AbsoluteEngineDir))
+			{
+				OriginalProperty.MetaData.Add("Category", Struct->GetFName().ToString());
+				Category = OriginalProperty.MetaData.Find("Category");
+			}
+			else
+			{
+				FError::Throwf(TEXT("Property is exposed to the editor or blueprints but has no Category specified."));
+			}
 		}
 
 		// Validate that pointer properties are not interfaces (which are not GC'd and so will cause runtime errors)
