Rename FPostConstructInitializeProperties to something simpler

diff --git a/Engine/Source/Editor/WorldBrowser/Private/Tiles/WorldTileDetails.cpp b/Engine/Source/Editor/WorldBrowser/Private/Tiles/WorldTileDetails.cpp
--- a/Engine/Source/Editor/WorldBrowser/Private/Tiles/WorldTileDetails.cpp
+++ b/Engine/Source/Editor/WorldBrowser/Private/Tiles/WorldTileDetails.cpp
@@ -20,8 +20,8 @@ FTileLODEntryDetails::FTileLODEntryDetails()
 
 /////////////////////////////////////////////////////
 // UWorldTileDetails
-UWorldTileDetails::UWorldTileDetails(const class FPostConstructInitializeProperties& PCIP)
-	: Super(PCIP)
+UWorldTileDetails::UWorldTileDetails(const FObjectInitializer& ObjectInitializer)
+	: Super(ObjectInitializer)
 {
 	LOD1.LODIndex = 0;
 	LOD2.LODIndex = 1;
diff --git a/Engine/Source/Programs/UnrealHeaderTool/Private/CodeGenerator.cpp b/Engine/Source/Programs/UnrealHeaderTool/Private/CodeGenerator.cpp
--- a/Engine/Source/Programs/UnrealHeaderTool/Private/CodeGenerator.cpp
+++ b/Engine/Source/Programs/UnrealHeaderTool/Private/CodeGenerator.cpp
@@ -836,7 +836,7 @@ void FNativeClassHeaderGenerator::OutputProperty(FString& Meta, FOutputDevice& O
 				*PropNameDep, 
 				*SourceStruct,
 				*BoolProperty->GetCPPType(NULL, 0));
-			PropMacroOuterClass = FString::Printf(TEXT("FPostConstructInitializeProperties(), EC_CppProperty, CPP_BOOL_PROPERTY_OFFSET(%s, %s)"),
+			PropMacroOuterClass = FString::Printf(TEXT("FObjectInitializer(), EC_CppProperty, CPP_BOOL_PROPERTY_OFFSET(%s, %s)"),
 				*PropNameDep, *SourceStruct);
 		}
 		else
@@ -849,7 +849,7 @@ void FNativeClassHeaderGenerator::OutputProperty(FString& Meta, FOutputDevice& O
 	if (ArrayProperty)
 	{
 		FString InnerOuterString = FString::Printf(TEXT("NewProp_%s"), *Prop->GetName());
-		FString PropMacroOuterArray = TEXT("FPostConstructInitializeProperties(), EC_CppProperty, 0");
+		FString PropMacroOuterArray = TEXT("FObjectInitializer(), EC_CppProperty, 0");
 		OutputDevice.Log(*PropertyNew(Meta, ArrayProperty->Inner, InnerOuterString, PropMacroOuterArray, TEXT("_Inner"), Spaces));
 	}
 }
@@ -1017,7 +1017,7 @@ void FNativeClassHeaderGenerator::ExportNativeGeneratedInitCode(FClass* Class)
 		GeneratedEnumRegisterFunctionText.Logf(TEXT("        if (!ReturnEnum)\r\n"));
 		GeneratedEnumRegisterFunctionText.Logf(TEXT("        {\r\n"));
 
-		GeneratedEnumRegisterFunctionText.Logf(TEXT("            ReturnEnum = new(Outer, TEXT(\"%s\"), RF_Public|RF_Transient|RF_Native) UEnum(FPostConstructInitializeProperties());\r\n"), *Enum->GetName());
+		GeneratedEnumRegisterFunctionText.Logf(TEXT("            ReturnEnum = new(Outer, TEXT(\"%s\"), RF_Public|RF_Transient|RF_Native) UEnum(FObjectInitializer());\r\n"), *Enum->GetName());
 		GeneratedEnumRegisterFunctionText.Logf(TEXT("            TArray<FName> EnumNames;\r\n"));
 		for (int32 Index = 0; Index < Enum->NumEnums(); Index++)
 		{
@@ -1110,7 +1110,7 @@ void FNativeClassHeaderGenerator::ExportNativeGeneratedInitCode(FClass* Class)
 			ExplicitSizeString = FString::Printf( TEXT(", sizeof(%s), ALIGNOF(%s)"), NameLookupCPP.GetNameCPP(ScriptStruct), NameLookupCPP.GetNameCPP(ScriptStruct) );
 		}
 
-		GeneratedStructRegisterFunctionText.Logf(TEXT("            ReturnStruct = new(Outer, TEXT(\"%s\"), RF_Public|RF_Transient|RF_Native) UScriptStruct(FPostConstructInitializeProperties(), %s, %s, EStructFlags(0x%08X)%s);\r\n"),
+		GeneratedStructRegisterFunctionText.Logf(TEXT("            ReturnStruct = new(Outer, TEXT(\"%s\"), RF_Public|RF_Transient|RF_Native) UScriptStruct(FObjectInitializer(), %s, %s, EStructFlags(0x%08X)%s);\r\n"),
 			*ScriptStruct->GetName(),
 			*BaseStructString,
 			*CppStructOpsString,
@@ -1228,7 +1228,7 @@ void FNativeClassHeaderGenerator::ExportNativeGeneratedInitCode(FClass* Class)
 			StructureSize = FString::Printf(TEXT(", sizeof(%s)"), *GetEventStructParamsName(*CastChecked<UClass>(TempFunction->GetOuter())->GetName(), *FunctionName));
 		}
 
-		GeneratedFunctionText.Logf(TEXT("            ReturnFunction = new(OuterClass, TEXT(\"%s\"), RF_Public|RF_Transient|RF_Native) UFunction(FPostConstructInitializeProperties(), %s, 0x%08X, %d%s);\r\n"), 
+		GeneratedFunctionText.Logf(TEXT("            ReturnFunction = new(OuterClass, TEXT(\"%s\"), RF_Public|RF_Transient|RF_Native) UFunction(FObjectInitializer(), %s, 0x%08X, %d%s);\r\n"), 
 			*Function->GetName(),
 			*SuperFunctionString,
 			Function->FunctionFlags,
@@ -2016,10 +2016,10 @@ void ExportStandardConstructorsMacro(FStringOutputDevice& Out, FClass* Class, co
 	if (!Class->HasAnyClassFlags(CLASS_CustomConstructor))
 	{
 		Out.Logf(TEXT("%s/** Standard constructor, called after all reflected properties have been initialized */\r\n"), FCString::Spc(4));
-		Out.Logf(TEXT("%s%s_API %s(const class FPostConstructInitializeProperties& PCIP);\r\n"), FCString::Spc(4), *API, NameLookupCPP.GetNameCPP(Class));
+		Out.Logf(TEXT("%s%s_API %s(const class FObjectInitializer& ObjectInitializer);\r\n"), FCString::Spc(4), *API, NameLookupCPP.GetNameCPP(Class));
 	}
 
-	Out.Logf(TEXT("%sDEFINE_DEFAULT_PCIP_CONSTRUCTOR_CALL(%s)\r\n"), FCString::Spc(4), NameLookupCPP.GetNameCPP(Class));
+	Out.Logf(TEXT("%sDEFINE_DEFAULT_OBJECT_INITIALIZER_CONSTRUCTOR_CALL(%s)\r\n"), FCString::Spc(4), NameLookupCPP.GetNameCPP(Class));
 }
 
 /**
@@ -2036,10 +2036,10 @@ void ExportConstructorDefinition(FStringOutputDevice& Out, FClass* Class, const
 	if (!ClassData->bConstructorDeclared)
 	{
 		Out.Logf(TEXT("%s/** Standard constructor, called after all reflected properties have been initialized */"), FCString::Spc(4));
-		Out.Logf(TEXT("%s%s_API %s(const class FPostConstructInitializeProperties& PCIP) : Super(PCIP) { };\r\n"), FCString::Spc(4), *API, NameLookupCPP.GetNameCPP(Class));
+		Out.Logf(TEXT("%s%s_API %s(const class FObjectInitializer& ObjectInitializer) : Super(ObjectInitializer) { };\r\n"), FCString::Spc(4), *API, NameLookupCPP.GetNameCPP(Class));
 
 		ClassData->bConstructorDeclared = true;
-		ClassData->bPCIPConstructorDeclared = true;
+		ClassData->bObjectInitializerConstructorDeclared = true;
 	}
 }
 
@@ -2053,9 +2053,9 @@ void ExportDefaultConstructorCallDefinition(FStringOutputDevice& Out, FClass* Cl
 {
 	auto* ClassData = GScriptHelper.FindClassData(Class);
 
-	if (ClassData->bPCIPConstructorDeclared)
+	if (ClassData->bObjectInitializerConstructorDeclared)
 	{
-		Out.Logf(TEXT("%sDEFINE_DEFAULT_PCIP_CONSTRUCTOR_CALL(%s)\r\n"), FCString::Spc(4), NameLookupCPP.GetNameCPP(Class));
+		Out.Logf(TEXT("%sDEFINE_DEFAULT_OBJECT_INITIALIZER_CONSTRUCTOR_CALL(%s)\r\n"), FCString::Spc(4), NameLookupCPP.GetNameCPP(Class));
 	}
 	else if (ClassData->bDefaultConstructorDeclared)
 	{
@@ -5028,7 +5028,7 @@ UClass* GenerateCodeForHeader
 		}
 
 		// Create new class.
-		ResultClass = new( InParent, *ClassNameStripped, Flags ) UClass( FPostConstructInitializeProperties(),NULL );
+		ResultClass = new( InParent, *ClassNameStripped, Flags ) UClass( FObjectInitializer(),NULL );
 		GClassHeaderNameWithNoPathMap.Add(ResultClass, ClassNameStripped);
 
 		// add CLASS_Interface flag if the class is an interface
@@ -5058,7 +5058,7 @@ UClass* GenerateCodeForHeader
 		if (ResultClass->GetSuperStruct() == NULL)
 		{
 			// don't know its parent class yet
-			ResultClass->SetSuperStruct( new(InParent, *BaseClassNameStripped) UClass(FPostConstructInitializeProperties(),NULL) );
+			ResultClass->SetSuperStruct( new(InParent, *BaseClassNameStripped) UClass(FObjectInitializer(),NULL) );
 		}
 
 		if (ResultClass->GetSuperStruct() != NULL)
diff --git a/Engine/Source/Programs/UnrealHeaderTool/Private/HeaderParser.cpp b/Engine/Source/Programs/UnrealHeaderTool/Private/HeaderParser.cpp
--- a/Engine/Source/Programs/UnrealHeaderTool/Private/HeaderParser.cpp
+++ b/Engine/Source/Programs/UnrealHeaderTool/Private/HeaderParser.cpp
@@ -610,7 +610,7 @@ UEnum* FHeaderParser::CompileEnum(UClass* Scope)
 	ParseFieldMetaData(EnumToken.MetaData, EnumToken.Identifier);
 
 	// Create enum definition.
-	UEnum* Enum = new(Scope->HasAnyClassFlags(CLASS_Temporary) ? Scope->GetOuter() : Scope, EnumToken.Identifier, RF_Public) UEnum(FPostConstructInitializeProperties());
+	UEnum* Enum = new(Scope->HasAnyClassFlags(CLASS_Temporary) ? Scope->GetOuter() : Scope, EnumToken.Identifier, RF_Public) UEnum(FObjectInitializer());
 	Enum->Next = Scope->Children;
 	Scope->Children = Enum;
 
@@ -1329,7 +1329,7 @@ UScriptStruct* FHeaderParser::CompileStructDeclaration(FClasses& AllClasses, FCl
 	// If declared in a header without UClass, the outer is the class package.
 	UObject* StructOuter = Scope->HasAnyClassFlags(CLASS_Temporary) ? Scope->GetOuter() : Scope;
 	// Create.
-	UScriptStruct* Struct = new(StructOuter, *EffectiveStructName, RF_Public) UScriptStruct(FPostConstructInitializeProperties(), BaseStruct);
+	UScriptStruct* Struct = new(StructOuter, *EffectiveStructName, RF_Public) UScriptStruct(FObjectInitializer(), BaseStruct);
 	Struct->Next = Scope->Children;
 	Scope->Children = Struct;
 
@@ -1722,7 +1722,7 @@ void FHeaderParser::PushNest( ENestType NestType, FName ThisName, UStruct* InNod
 			}
 			else if (NestType == NEST_FunctionDeclaration)
 			{
-				UFunction* Function = new(PrevTopNode ? (UObject*)PrevTopNode : (UObject*)Class, ThisName, RF_Public) UFunction( FPostConstructInitializeProperties(), NULL );
+				UFunction* Function = new(PrevTopNode ? (UObject*)PrevTopNode : (UObject*)Class, ThisName, RF_Public) UFunction( FObjectInitializer(), NULL );
 				TopNest->Node = TopNode = Function;
 				Function->RepOffset = MAX_uint16;
 				Function->ReturnValueOffset = MAX_uint16;
@@ -3572,7 +3572,7 @@ UProperty* FHeaderParser::GetVarNameAndDim
 		int32      ArrayDim = 1; // 1 = not a static array, 2 = static array
 		if (VarProperty.ArrayType == EArrayType::Dynamic)
 		{
-			Array       = new(Scope,PropertyName,ObjectFlags)UArrayProperty(FPostConstructInitializeProperties());
+			Array       = new(Scope,PropertyName,ObjectFlags)UArrayProperty(FObjectInitializer());
 			NewScope    = Array;
 			ObjectFlags = RF_Public;
 		}
@@ -3583,40 +3583,40 @@ UProperty* FHeaderParser::GetVarNameAndDim
 
 		if (VarProperty.Type == CPT_Byte)
 		{
-			NewProperty = new(NewScope,PropertyName,ObjectFlags)UByteProperty(FPostConstructInitializeProperties());
+			NewProperty = new(NewScope,PropertyName,ObjectFlags)UByteProperty(FObjectInitializer());
 			((UByteProperty*)NewProperty)->Enum = VarProperty.Enum;
 		}
 		else if (VarProperty.Type == CPT_Int8)
 		{
-			NewProperty = new(NewScope,PropertyName,ObjectFlags)UInt8Property(FPostConstructInitializeProperties());
+			NewProperty = new(NewScope,PropertyName,ObjectFlags)UInt8Property(FObjectInitializer());
 		}
 		else if (VarProperty.Type == CPT_Int16)
 		{
-			NewProperty = new(NewScope,PropertyName,ObjectFlags)UInt16Property(FPostConstructInitializeProperties());
+			NewProperty = new(NewScope,PropertyName,ObjectFlags)UInt16Property(FObjectInitializer());
 		}
 		else if (VarProperty.Type == CPT_Int)
 		{
-			NewProperty = new(NewScope,PropertyName,ObjectFlags)UIntProperty(FPostConstructInitializeProperties());
+			NewProperty = new(NewScope,PropertyName,ObjectFlags)UIntProperty(FObjectInitializer());
 		}
 		else if (VarProperty.Type == CPT_Int64)
 		{
-			NewProperty = new(NewScope,PropertyName,ObjectFlags)UInt64Property(FPostConstructInitializeProperties());
+			NewProperty = new(NewScope,PropertyName,ObjectFlags)UInt64Property(FObjectInitializer());
 		}
 		else if (VarProperty.Type == CPT_UInt16)
 		{
-			NewProperty = new(NewScope,PropertyName,ObjectFlags)UUInt16Property(FPostConstructInitializeProperties());
+			NewProperty = new(NewScope,PropertyName,ObjectFlags)UUInt16Property(FObjectInitializer());
 		}
 		else if (VarProperty.Type == CPT_UInt32)
 		{
-			NewProperty = new(NewScope,PropertyName,ObjectFlags)UUInt32Property(FPostConstructInitializeProperties());
+			NewProperty = new(NewScope,PropertyName,ObjectFlags)UUInt32Property(FObjectInitializer());
 		}
 		else if (VarProperty.Type == CPT_UInt64)
 		{
-			NewProperty = new(NewScope,PropertyName,ObjectFlags)UUInt64Property(FPostConstructInitializeProperties());
+			NewProperty = new(NewScope,PropertyName,ObjectFlags)UUInt64Property(FObjectInitializer());
 		}
 		else if (VarProperty.IsBool())
 		{
-			UBoolProperty* NewBoolProperty = new(NewScope,PropertyName,ObjectFlags)UBoolProperty(FPostConstructInitializeProperties());
+			UBoolProperty* NewBoolProperty = new(NewScope,PropertyName,ObjectFlags)UBoolProperty(FObjectInitializer());
 			NewProperty = NewBoolProperty;
 			if (HardcodedName && FCString::Stricmp(HardcodedName, TEXT("ReturnValue")) == 0)
 			{
@@ -3646,18 +3646,18 @@ UProperty* FHeaderParser::GetVarNameAndDim
 		}
 		else if (VarProperty.Type == CPT_Float)
 		{
-			NewProperty = new(NewScope,PropertyName,ObjectFlags)UFloatProperty(FPostConstructInitializeProperties());
+			NewProperty = new(NewScope,PropertyName,ObjectFlags)UFloatProperty(FObjectInitializer());
 		}
 		else if (VarProperty.Type == CPT_Double)
 		{
-			NewProperty = new(NewScope,PropertyName,ObjectFlags)UDoubleProperty(FPostConstructInitializeProperties());
+			NewProperty = new(NewScope,PropertyName,ObjectFlags)UDoubleProperty(FObjectInitializer());
 		}
 		else if (VarProperty.Type == CPT_ObjectReference)
 		{
 			check(VarProperty.PropertyClass);
 			if (VarProperty.PropertyClass->IsChildOf(UClass::StaticClass()))
 			{
-				NewProperty = new(NewScope,PropertyName,ObjectFlags)UClassProperty(FPostConstructInitializeProperties());
+				NewProperty = new(NewScope,PropertyName,ObjectFlags)UClassProperty(FObjectInitializer());
 				((UClassProperty*)NewProperty)->MetaClass = VarProperty.MetaClass;
 			}
 			else
@@ -3667,7 +3667,7 @@ UProperty* FHeaderParser::GetVarNameAndDim
 					VarProperty.PropertyFlags |= CPF_InstancedReference;
 					AddEditInlineMetaData(VarProperty.MetaData);
 				}
-				NewProperty = new(NewScope,PropertyName,ObjectFlags)UObjectProperty(FPostConstructInitializeProperties());
+				NewProperty = new(NewScope,PropertyName,ObjectFlags)UObjectProperty(FObjectInitializer());
 			}
 
 			((UObjectPropertyBase*)NewProperty)->PropertyClass = VarProperty.PropertyClass;
@@ -3675,13 +3675,13 @@ UProperty* FHeaderParser::GetVarNameAndDim
 		else if (VarProperty.Type == CPT_WeakObjectReference)
 		{
 			check(VarProperty.PropertyClass);
-			NewProperty = new(NewScope,PropertyName,ObjectFlags)UWeakObjectProperty(FPostConstructInitializeProperties());
+			NewProperty = new(NewScope,PropertyName,ObjectFlags)UWeakObjectProperty(FObjectInitializer());
 			((UObjectPropertyBase*)NewProperty)->PropertyClass = VarProperty.PropertyClass;
 		}
 		else if( VarProperty.Type == CPT_LazyObjectReference )
 		{
 			check(VarProperty.PropertyClass);
-			NewProperty = new(NewScope,PropertyName,ObjectFlags)ULazyObjectProperty(FPostConstructInitializeProperties());
+			NewProperty = new(NewScope,PropertyName,ObjectFlags)ULazyObjectProperty(FObjectInitializer());
 			((UObjectPropertyBase*)NewProperty)->PropertyClass = VarProperty.PropertyClass;
 		}
 		else if( VarProperty.Type == CPT_AssetObjectReference )
@@ -3689,12 +3689,12 @@ UProperty* FHeaderParser::GetVarNameAndDim
 			check(VarProperty.PropertyClass);
 			if (VarProperty.PropertyClass->IsChildOf(UClass::StaticClass()))
 			{
-				NewProperty = new(NewScope,PropertyName,ObjectFlags)UAssetClassProperty(FPostConstructInitializeProperties());
+				NewProperty = new(NewScope,PropertyName,ObjectFlags)UAssetClassProperty(FObjectInitializer());
 				((UAssetClassProperty*)NewProperty)->MetaClass = VarProperty.MetaClass;
 			}
 			else
 			{
-				NewProperty = new(NewScope,PropertyName,ObjectFlags)UAssetObjectProperty(FPostConstructInitializeProperties());
+				NewProperty = new(NewScope,PropertyName,ObjectFlags)UAssetObjectProperty(FObjectInitializer());
 			}
 			((UObjectPropertyBase*)NewProperty)->PropertyClass = VarProperty.PropertyClass;
 		}
@@ -3703,26 +3703,26 @@ UProperty* FHeaderParser::GetVarNameAndDim
 			check(VarProperty.PropertyClass);
 			check(VarProperty.PropertyClass->HasAnyClassFlags(CLASS_Interface));
 
-			UInterfaceProperty* InterfaceProperty = new(NewScope,PropertyName,ObjectFlags) UInterfaceProperty(FPostConstructInitializeProperties());
+			UInterfaceProperty* InterfaceProperty = new(NewScope,PropertyName,ObjectFlags) UInterfaceProperty(FObjectInitializer());
 			InterfaceProperty->InterfaceClass = VarProperty.PropertyClass;
 
 			NewProperty = InterfaceProperty;
 		}
 		else if (VarProperty.Type == CPT_Name)
 		{
-			NewProperty = new(NewScope,PropertyName,ObjectFlags)UNameProperty(FPostConstructInitializeProperties());
+			NewProperty = new(NewScope,PropertyName,ObjectFlags)UNameProperty(FObjectInitializer());
 		}
 		else if (VarProperty.Type == CPT_String)
 		{
-			NewProperty = new(NewScope,PropertyName,ObjectFlags)UStrProperty(FPostConstructInitializeProperties());
+			NewProperty = new(NewScope,PropertyName,ObjectFlags)UStrProperty(FObjectInitializer());
 		}
 		else if (VarProperty.Type == CPT_Text)
 		{
-			NewProperty = new(NewScope,PropertyName,ObjectFlags)UTextProperty(FPostConstructInitializeProperties());
+			NewProperty = new(NewScope,PropertyName,ObjectFlags)UTextProperty(FObjectInitializer());
 		}
 		else if (VarProperty.Type == CPT_Struct)
 		{
-			NewProperty = new(NewScope,PropertyName,ObjectFlags)UStructProperty(FPostConstructInitializeProperties());
+			NewProperty = new(NewScope,PropertyName,ObjectFlags)UStructProperty(FObjectInitializer());
 			if (VarProperty.Struct->StructFlags & STRUCT_HasInstancedReference)
 			{
 				VarProperty.PropertyFlags |= CPF_ContainsInstancedReference;
@@ -3731,11 +3731,11 @@ UProperty* FHeaderParser::GetVarNameAndDim
 		}
 		else if (VarProperty.Type == CPT_Delegate)
 		{
-			NewProperty = new(NewScope,PropertyName,ObjectFlags)UDelegateProperty(FPostConstructInitializeProperties());
+			NewProperty = new(NewScope,PropertyName,ObjectFlags)UDelegateProperty(FObjectInitializer());
 		}
 		else if (VarProperty.Type == CPT_MulticastDelegate)
 		{
-			NewProperty = new(NewScope,PropertyName,ObjectFlags)UMulticastDelegateProperty(FPostConstructInitializeProperties());
+			NewProperty = new(NewScope,PropertyName,ObjectFlags)UMulticastDelegateProperty(FObjectInitializer());
 		}
 		else
 		{
@@ -7531,70 +7531,72 @@ bool FHeaderParser::TryToMatchConstructorParameterList(FToken Token)
 	{
 		ClassData->bDefaultConstructorDeclared = true;
 	}
-	else if (!ClassData->bPCIPConstructorDeclared)
+	else if (!ClassData->bObjectInitializerConstructorDeclared)
 	{
-		FToken PCIPParamParsingToken;
+		FToken ObjectInitializerParamParsingToken;
 
 		bool bIsConst = false;
 		bool bIsRef = false;
 		int32 ParenthesesNestingLevel = 1;
 
-		while (ParenthesesNestingLevel && GetToken(PCIPParamParsingToken))
+		while (ParenthesesNestingLevel && GetToken(ObjectInitializerParamParsingToken))
 		{
-			// Template instantiation or additional parameter excludes PCIP constructor.
-			if (PCIPParamParsingToken.Matches(TEXT(",")) || PCIPParamParsingToken.Matches(TEXT("<")))
+			// Template instantiation or additional parameter excludes ObjectInitializer constructor.
+			if (ObjectInitializerParamParsingToken.Matches(TEXT(",")) || ObjectInitializerParamParsingToken.Matches(TEXT("<")))
 			{
-				ClassData->bPCIPConstructorDeclared = false;
+				ClassData->bObjectInitializerConstructorDeclared = false;
 				break;
 			}
 
-			if (PCIPParamParsingToken.Matches(TEXT("(")))
+			if (ObjectInitializerParamParsingToken.Matches(TEXT("(")))
 			{
 				ParenthesesNestingLevel++;
 				continue;
 			}
 
-			if (PCIPParamParsingToken.Matches(TEXT(")")))
+			if (ObjectInitializerParamParsingToken.Matches(TEXT(")")))
 			{
 				ParenthesesNestingLevel--;
 				continue;
 			}
 
-			if (PCIPParamParsingToken.Matches(TEXT("const")))
+			if (ObjectInitializerParamParsingToken.Matches(TEXT("const")))
 			{
 				bIsConst = true;
 				continue;
 			}
 
-			if (PCIPParamParsingToken.Matches(TEXT("&")))
+			if (ObjectInitializerParamParsingToken.Matches(TEXT("&")))
 			{
 				bIsRef = true;
 				continue;
 			}
 
-			if (PCIPParamParsingToken.Matches(TEXT("FPostConstructInitializeProperties")))
+			if (ObjectInitializerParamParsingToken.Matches(TEXT("FObjectInitializer"))
+				|| ObjectInitializerParamParsingToken.Matches(TEXT("FPostConstructInitializeProperties")) // Deprecated, but left here, so it won't break legacy code.
+				)
 			{
-				ClassData->bPCIPConstructorDeclared = true;
+				ClassData->bObjectInitializerConstructorDeclared = true;
 			}
 		}
 
 		// Parse until finish.
-		while (ParenthesesNestingLevel && GetToken(PCIPParamParsingToken))
+		while (ParenthesesNestingLevel && GetToken(ObjectInitializerParamParsingToken))
 		{
-			if (PCIPParamParsingToken.Matches(TEXT("(")))
+			if (ObjectInitializerParamParsingToken.Matches(TEXT("(")))
 			{
 				ParenthesesNestingLevel++;
 				continue;
 			}
 
-			if (PCIPParamParsingToken.Matches(TEXT(")")))
+			if (ObjectInitializerParamParsingToken.Matches(TEXT(")")))
 			{
 				ParenthesesNestingLevel--;
 				continue;
 			}
 		}
 
-		ClassData->bPCIPConstructorDeclared = ClassData->bPCIPConstructorDeclared && bIsRef && bIsConst;
+		ClassData->bObjectInitializerConstructorDeclared = ClassData->bObjectInitializerConstructorDeclared && bIsRef && bIsConst;
 	}
 
 	// Optionally match semicolon.
diff --git a/Engine/Source/Programs/UnrealHeaderTool/Private/ParserHelper.h b/Engine/Source/Programs/UnrealHeaderTool/Private/ParserHelper.h
--- a/Engine/Source/Programs/UnrealHeaderTool/Private/ParserHelper.h
+++ b/Engine/Source/Programs/UnrealHeaderTool/Private/ParserHelper.h
@@ -1310,7 +1310,7 @@ public:
 		: bContainsDelegates(false),
 		bConstructorDeclared(false),
 		bDefaultConstructorDeclared(false),
-		bPCIPConstructorDeclared(false)
+		bObjectInitializerConstructorDeclared(false)
 	{
 	}
 
@@ -1556,8 +1556,8 @@ public:
 	// Is default constructor declared?
 	bool bDefaultConstructorDeclared;
 
-	// Is PCIP constructor (i.e. a constructor with only one parameter of type FPostConstructInitializeProperties) declared?
-	bool bPCIPConstructorDeclared;
+	// Is ObjectInitializer constructor (i.e. a constructor with only one parameter of type FObjectInitializer) declared?
+	bool bObjectInitializerConstructorDeclared;
 
 	// GENERATED_BODY access specifier to preserve.
 	EAccessSpecifier GeneratedBodyMacroAccessSpecifier;
diff --git a/Engine/Source/Runtime/AIModule/Classes/EnvironmentQuery/EnvQueryDebugHelpers.h b/Engine/Source/Runtime/AIModule/Classes/EnvironmentQuery/EnvQueryDebugHelpers.h
--- a/Engine/Source/Runtime/AIModule/Classes/EnvironmentQuery/EnvQueryDebugHelpers.h
+++ b/Engine/Source/Runtime/AIModule/Classes/EnvironmentQuery/EnvQueryDebugHelpers.h
@@ -155,7 +155,7 @@ class AIMODULE_API UEnvQueryDebugHelpers : public UObject
 {
 	GENERATED_UCLASS_BODY()
 
-	UEnvQueryDebugHelpers(const class FPostConstructInitializeProperties& PCIP) : Super(PCIP) {}
+	UEnvQueryDebugHelpers(const FObjectInitializer& ObjectInitializer) : Super(ObjectInitializer) {}
 #if USE_EQS_DEBUGGER
 	static void QueryToDebugData(struct FEnvQueryInstance* Query, EQSDebug::FQueryData& EQSLocalData);
 	static void QueryToBlobArray(struct FEnvQueryInstance* Query, TArray<uint8>& BlobArray, bool bUseCompression = false);
diff --git a/Engine/Source/Runtime/AIModule/Classes/EnvironmentQuery/EnvQueryTypes.h b/Engine/Source/Runtime/AIModule/Classes/EnvironmentQuery/EnvQueryTypes.h
--- a/Engine/Source/Runtime/AIModule/Classes/EnvironmentQuery/EnvQueryTypes.h
+++ b/Engine/Source/Runtime/AIModule/Classes/EnvironmentQuery/EnvQueryTypes.h
@@ -963,7 +963,7 @@ class AIMODULE_API UEnvQueryTypes : public UObject
 	/** special test value assigned to items skipped by condition check */
 	static float SkippedItemValue;
 
-	UEnvQueryTypes(const class FPostConstructInitializeProperties& PCIP) : Super(PCIP) {}
+	UEnvQueryTypes(const FObjectInitializer& ObjectInitializer) : Super(ObjectInitializer) {}
 	static FText GetShortTypeName(const UObject* Ob);
 
 	static FText DescribeContext(TSubclassOf<class UEnvQueryContext> ContextClass);
diff --git a/Engine/Source/Runtime/AIModule/Classes/Navigation/CrowdManager.h b/Engine/Source/Runtime/AIModule/Classes/Navigation/CrowdManager.h
--- a/Engine/Source/Runtime/AIModule/Classes/Navigation/CrowdManager.h
+++ b/Engine/Source/Runtime/AIModule/Classes/Navigation/CrowdManager.h
@@ -19,8 +19,8 @@ class UCrowdFollowingComponent;
  *  (usually: replace class of PathFollowingComponent in AIController by adding 
  *   those lines in controller's constructor
  *
- *   ACrowdAIController::ACrowdAIController(const class FPostConstructInitializeProperties& PCIP)
- *       : Super(PCIP.SetDefaultSubobjectClass<UCrowdFollowingComponent>(TEXT("PathFollowingComponent")))
+ *   ACrowdAIController::ACrowdAIController(const FObjectInitializer& ObjectInitializer)
+ *       : Super(ObjectInitializer.SetDefaultSubobjectClass<UCrowdFollowingComponent>(TEXT("PathFollowingComponent")))
  *
  *   or simply add both components and switch move requests between them)
  *
diff --git a/Engine/Source/Runtime/AIModule/Classes/VisualLoggerExtension.h b/Engine/Source/Runtime/AIModule/Classes/VisualLoggerExtension.h
--- a/Engine/Source/Runtime/AIModule/Classes/VisualLoggerExtension.h
+++ b/Engine/Source/Runtime/AIModule/Classes/VisualLoggerExtension.h
@@ -35,5 +35,5 @@ class AIMODULE_API UVisualLoggerExtension : public UObject
 {
 	GENERATED_UCLASS_BODY()
 
-	UVisualLoggerExtension(const class FPostConstructInitializeProperties& PCIP) : Super(PCIP) {}
+	UVisualLoggerExtension(const FObjectInitializer& ObjectInitializer) : Super(ObjectInitializer) {}
 };
diff --git a/Engine/Source/Runtime/AIModule/Private/AIController.cpp b/Engine/Source/Runtime/AIModule/Private/AIController.cpp
--- a/Engine/Source/Runtime/AIModule/Private/AIController.cpp
+++ b/Engine/Source/Runtime/AIModule/Private/AIController.cpp
@@ -33,16 +33,16 @@ DECLARE_CYCLE_STAT(TEXT("MoveToActor"),STAT_MoveToActor,STATGROUP_AI);
 
 DEFINE_LOG_CATEGORY(LogAINavigation);
 
-AAIController::AAIController(const class FPostConstructInitializeProperties& PCIP)
-	: Super(PCIP)
+AAIController::AAIController(const FObjectInitializer& ObjectInitializer)
+	: Super(ObjectInitializer)
 {
 	// set up navigation component
-	NavComponent = PCIP.CreateDefaultSubobject<UNavigationComponent>(this, TEXT("NavComponent"));
+	NavComponent = ObjectInitializer.CreateDefaultSubobject<UNavigationComponent>(this, TEXT("NavComponent"));
 
-	PathFollowingComponent = PCIP.CreateDefaultSubobject<UPathFollowingComponent>(this, TEXT("PathFollowingComponent"));
+	PathFollowingComponent = ObjectInitializer.CreateDefaultSubobject<UPathFollowingComponent>(this, TEXT("PathFollowingComponent"));
 	PathFollowingComponent->OnMoveFinished.AddUObject(this, &AAIController::OnMoveCompleted);
 
-	ActionsComp = PCIP.CreateDefaultSubobject<UPawnActionsComponent>(this, "ActionsComp");
+	ActionsComp = ObjectInitializer.CreateDefaultSubobject<UPawnActionsComponent>(this, "ActionsComp");
 
 	bSkipExtraLOSChecks = true;
 	bWantsPlayerState = false;
diff --git a/Engine/Source/Runtime/AIModule/Private/AIInterfaces.cpp b/Engine/Source/Runtime/AIModule/Private/AIInterfaces.cpp
--- a/Engine/Source/Runtime/AIModule/Private/AIInterfaces.cpp
+++ b/Engine/Source/Runtime/AIModule/Private/AIInterfaces.cpp
@@ -4,8 +4,8 @@
 #include "AIResourceInterface.h"
 #include "GenericTeamAgentInterface.h"
 
-UAIResourceInterface::UAIResourceInterface(const class FPostConstructInitializeProperties& PCIP)
-	: Super(PCIP)
+UAIResourceInterface::UAIResourceInterface(const FObjectInitializer& ObjectInitializer)
+	: Super(ObjectInitializer)
 {
 }
 
@@ -29,7 +29,7 @@ ETeamAttitude::Type FGenericTeamId::GetAttitude(const AActor* A, const AActor* B
 	return TeamAgentA == NULL || B == NULL ? ETeamAttitude::Neutral : TeamAgentA->GetTeamAttitudeTowards(*B);
 }
 
-UGenericTeamAgentInterface::UGenericTeamAgentInterface(const class FPostConstructInitializeProperties& PCIP)
-	: Super(PCIP)
+UGenericTeamAgentInterface::UGenericTeamAgentInterface(const FObjectInitializer& ObjectInitializer)
+	: Super(ObjectInitializer)
 {
 }
diff --git a/Engine/Source/Runtime/AIModule/Private/AISystem.cpp b/Engine/Source/Runtime/AIModule/Private/AISystem.cpp
--- a/Engine/Source/Runtime/AIModule/Private/AISystem.cpp
+++ b/Engine/Source/Runtime/AIModule/Private/AISystem.cpp
@@ -10,7 +10,7 @@
 #include "AISystem.h"
 
 
-UAISystem::UAISystem(const FPostConstructInitializeProperties& PCIP) : Super(PCIP)
+UAISystem::UAISystem(const FObjectInitializer& ObjectInitializer) : Super(ObjectInitializer)
 {
 	if (HasAnyFlags(RF_ClassDefaultObject) == false)
 	{
diff --git a/Engine/Source/Runtime/AIModule/Private/Actions/PawnAction.cpp b/Engine/Source/Runtime/AIModule/Private/Actions/PawnAction.cpp
--- a/Engine/Source/Runtime/AIModule/Private/Actions/PawnAction.cpp
+++ b/Engine/Source/Runtime/AIModule/Private/Actions/PawnAction.cpp
@@ -16,8 +16,8 @@ namespace
 	}
 }
 
-UPawnAction::UPawnAction(const class FPostConstructInitializeProperties& PCIP)
-	: Super(PCIP)
+UPawnAction::UPawnAction(const FObjectInitializer& ObjectInitializer)
+	: Super(ObjectInitializer)
 	, AbortState(EPawnActionAbortState::NeverStarted)
 	, FinishResult(EPawnActionResult::NotStarted)
 {
diff --git a/Engine/Source/Runtime/AIModule/Private/Actions/PawnAction_BlueprintBase.cpp b/Engine/Source/Runtime/AIModule/Private/Actions/PawnAction_BlueprintBase.cpp
--- a/Engine/Source/Runtime/AIModule/Private/Actions/PawnAction_BlueprintBase.cpp
+++ b/Engine/Source/Runtime/AIModule/Private/Actions/PawnAction_BlueprintBase.cpp
@@ -11,8 +11,8 @@ namespace
 	}
 }
 
-UPawnAction_BlueprintBase::UPawnAction_BlueprintBase(const class FPostConstructInitializeProperties& PCIP)
-	: Super(PCIP)
+UPawnAction_BlueprintBase::UPawnAction_BlueprintBase(const FObjectInitializer& ObjectInitializer)
+	: Super(ObjectInitializer)
 {
 	UClass* StopAtClass = UPawnAction_BlueprintBase::StaticClass();
 	bWantsTick = HasBlueprintFunction(TEXT("ActionTick"), this, StopAtClass);
diff --git a/Engine/Source/Runtime/AIModule/Private/Actions/PawnAction_Move.cpp b/Engine/Source/Runtime/AIModule/Private/Actions/PawnAction_Move.cpp
--- a/Engine/Source/Runtime/AIModule/Private/Actions/PawnAction_Move.cpp
+++ b/Engine/Source/Runtime/AIModule/Private/Actions/PawnAction_Move.cpp
@@ -3,8 +3,8 @@
 #include "AIModulePrivate.h"
 #include "Actions/PawnAction_Move.h"
 
-UPawnAction_Move::UPawnAction_Move(const class FPostConstructInitializeProperties& PCIP)
-	: Super(PCIP)
+UPawnAction_Move::UPawnAction_Move(const FObjectInitializer& ObjectInitializer)
+	: Super(ObjectInitializer)
 	, GoalLocation(FAISystem::InvalidLocation)
 	, AcceptableRadius(30.f)
 	, bFinishOnOverlap(true)
diff --git a/Engine/Source/Runtime/AIModule/Private/Actions/PawnAction_Repeat.cpp b/Engine/Source/Runtime/AIModule/Private/Actions/PawnAction_Repeat.cpp
--- a/Engine/Source/Runtime/AIModule/Private/Actions/PawnAction_Repeat.cpp
+++ b/Engine/Source/Runtime/AIModule/Private/Actions/PawnAction_Repeat.cpp
@@ -3,8 +3,8 @@
 #include "AIModulePrivate.h"
 #include "Actions/PawnAction_Repeat.h"
 
-UPawnAction_Repeat::UPawnAction_Repeat(const class FPostConstructInitializeProperties& PCIP)
-	: Super(PCIP)
+UPawnAction_Repeat::UPawnAction_Repeat(const FObjectInitializer& ObjectInitializer)
+	: Super(ObjectInitializer)
 	, SubActionTriggeringPolicy(EPawnSubActionTriggeringPolicy::CopyBeforeTriggering)
 {
 	ChildFailureHandlingMode = EPawnActionFailHandling::IgnoreFailure;
diff --git a/Engine/Source/Runtime/AIModule/Private/Actions/PawnAction_Sequence.cpp b/Engine/Source/Runtime/AIModule/Private/Actions/PawnAction_Sequence.cpp
--- a/Engine/Source/Runtime/AIModule/Private/Actions/PawnAction_Sequence.cpp
+++ b/Engine/Source/Runtime/AIModule/Private/Actions/PawnAction_Sequence.cpp
@@ -3,8 +3,8 @@
 #include "AIModulePrivate.h"
 #include "Actions/PawnAction_Sequence.h"
 
-UPawnAction_Sequence::UPawnAction_Sequence(const class FPostConstructInitializeProperties& PCIP)
-	: Super(PCIP)
+UPawnAction_Sequence::UPawnAction_Sequence(const FObjectInitializer& ObjectInitializer)
+	: Super(ObjectInitializer)
 	, SubActionTriggeringPolicy(EPawnSubActionTriggeringPolicy::CopyBeforeTriggering)
 {
 }
diff --git a/Engine/Source/Runtime/AIModule/Private/Actions/PawnAction_Wait.cpp b/Engine/Source/Runtime/AIModule/Private/Actions/PawnAction_Wait.cpp
--- a/Engine/Source/Runtime/AIModule/Private/Actions/PawnAction_Wait.cpp
+++ b/Engine/Source/Runtime/AIModule/Private/Actions/PawnAction_Wait.cpp
@@ -3,8 +3,8 @@
 #include "AIModulePrivate.h"
 #include "Actions/PawnAction_Wait.h"
 
-UPawnAction_Wait::UPawnAction_Wait(const FPostConstructInitializeProperties& PCIP)
-	: Super(PCIP)
+UPawnAction_Wait::UPawnAction_Wait(const FObjectInitializer& ObjectInitializer)
+	: Super(ObjectInitializer)
 	, TimeToWait(0.f)
 	, FinishTimeStamp(0.f)
 {
diff --git a/Engine/Source/Runtime/AIModule/Private/Actions/PawnActionsComponent.cpp b/Engine/Source/Runtime/AIModule/Private/Actions/PawnActionsComponent.cpp
--- a/Engine/Source/Runtime/AIModule/Private/Actions/PawnActionsComponent.cpp
+++ b/Engine/Source/Runtime/AIModule/Private/Actions/PawnActionsComponent.cpp
@@ -147,8 +147,8 @@ int32 FPawnActionStack::GetStackSize() const
 // UPawnActionsComponent
 //----------------------------------------------------------------------//
 
-UPawnActionsComponent::UPawnActionsComponent(const class FPostConstructInitializeProperties& PCIP)
-	: Super(PCIP)
+UPawnActionsComponent::UPawnActionsComponent(const FObjectInitializer& ObjectInitializer)
+	: Super(ObjectInitializer)
 {
 	PrimaryComponentTick.bCanEverTick = true;
 	PrimaryComponentTick.bStartWithTickEnabled = true;
diff --git a/Engine/Source/Runtime/AIModule/Private/BehaviorTree/BTAuxiliaryNode.cpp b/Engine/Source/Runtime/AIModule/Private/BehaviorTree/BTAuxiliaryNode.cpp
--- a/Engine/Source/Runtime/AIModule/Private/BehaviorTree/BTAuxiliaryNode.cpp
+++ b/Engine/Source/Runtime/AIModule/Private/BehaviorTree/BTAuxiliaryNode.cpp
@@ -3,7 +3,7 @@
 #include "AIModulePrivate.h"
 #include "BehaviorTree/BTAuxiliaryNode.h"
 
-UBTAuxiliaryNode::UBTAuxiliaryNode(const class FPostConstructInitializeProperties& PCIP) : Super(PCIP)
+UBTAuxiliaryNode::UBTAuxiliaryNode(const FObjectInitializer& ObjectInitializer) : Super(ObjectInitializer)
 {
 	bNotifyBecomeRelevant = false;
 	bNotifyCeaseRelevant = false;
diff --git a/Engine/Source/Runtime/AIModule/Private/BehaviorTree/BTCompositeNode.cpp b/Engine/Source/Runtime/AIModule/Private/BehaviorTree/BTCompositeNode.cpp
--- a/Engine/Source/Runtime/AIModule/Private/BehaviorTree/BTCompositeNode.cpp
+++ b/Engine/Source/Runtime/AIModule/Private/BehaviorTree/BTCompositeNode.cpp
@@ -3,7 +3,7 @@
 #include "AIModulePrivate.h"
 #include "BehaviorTree/BTCompositeNode.h"
 
-UBTCompositeNode::UBTCompositeNode(const class FPostConstructInitializeProperties& PCIP) : Super(PCIP)
+UBTCompositeNode::UBTCompositeNode(const FObjectInitializer& ObjectInitializer) : Super(ObjectInitializer)
 {
 	NodeName = "UnknownComposite";
 	bUseChildExecutionNotify = false;
diff --git a/Engine/Source/Runtime/AIModule/Private/BehaviorTree/BTDecorator.cpp b/Engine/Source/Runtime/AIModule/Private/BehaviorTree/BTDecorator.cpp
--- a/Engine/Source/Runtime/AIModule/Private/BehaviorTree/BTDecorator.cpp
+++ b/Engine/Source/Runtime/AIModule/Private/BehaviorTree/BTDecorator.cpp
@@ -4,7 +4,7 @@
 #include "BehaviorTree/BTCompositeNode.h"
 #include "BehaviorTree/BTDecorator.h"
 
-UBTDecorator::UBTDecorator(const class FPostConstructInitializeProperties& PCIP) : Super(PCIP)
+UBTDecorator::UBTDecorator(const FObjectInitializer& ObjectInitializer) : Super(ObjectInitializer)
 {
 	NodeName = "UnknownDecorator";
 
diff --git a/Engine/Source/Runtime/AIModule/Private/BehaviorTree/BTFunctionLibrary.cpp b/Engine/Source/Runtime/AIModule/Private/BehaviorTree/BTFunctionLibrary.cpp
--- a/Engine/Source/Runtime/AIModule/Private/BehaviorTree/BTFunctionLibrary.cpp
+++ b/Engine/Source/Runtime/AIModule/Private/BehaviorTree/BTFunctionLibrary.cpp
@@ -6,7 +6,7 @@
 #include "BehaviorTree/Blackboard/BlackboardKeyAllTypes.h"
 #include "BehaviorTree/BTFunctionLibrary.h"
 
-UBTFunctionLibrary::UBTFunctionLibrary(const class FPostConstructInitializeProperties& PCIP) : Super(PCIP)
+UBTFunctionLibrary::UBTFunctionLibrary(const FObjectInitializer& ObjectInitializer) : Super(ObjectInitializer)
 {
 }
 
diff --git a/Engine/Source/Runtime/AIModule/Private/BehaviorTree/BTNode.cpp b/Engine/Source/Runtime/AIModule/Private/BehaviorTree/BTNode.cpp
--- a/Engine/Source/Runtime/AIModule/Private/BehaviorTree/BTNode.cpp
+++ b/Engine/Source/Runtime/AIModule/Private/BehaviorTree/BTNode.cpp
@@ -9,7 +9,7 @@
 // UBTNode
 //----------------------------------------------------------------------//
 
-UBTNode::UBTNode(const class FPostConstructInitializeProperties& PCIP) : Super(PCIP)
+UBTNode::UBTNode(const FObjectInitializer& ObjectInitializer) : Super(ObjectInitializer)
 {
 	NodeName = "";
 	ParentNode = NULL;
diff --git a/Engine/Source/Runtime/AIModule/Private/BehaviorTree/BTService.cpp b/Engine/Source/Runtime/AIModule/Private/BehaviorTree/BTService.cpp
--- a/Engine/Source/Runtime/AIModule/Private/BehaviorTree/BTService.cpp
+++ b/Engine/Source/Runtime/AIModule/Private/BehaviorTree/BTService.cpp
@@ -3,7 +3,7 @@
 #include "AIModulePrivate.h"
 #include "BehaviorTree/BTService.h"
 
-UBTService::UBTService(const class FPostConstructInitializeProperties& PCIP) : Super(PCIP)
+UBTService::UBTService(const FObjectInitializer& ObjectInitializer) : Super(ObjectInitializer)
 {
 	bNotifyTick = true;
 	bNotifyOnSearch = true;
diff --git a/Engine/Source/Runtime/AIModule/Private/BehaviorTree/BTTaskNode.cpp b/Engine/Source/Runtime/AIModule/Private/BehaviorTree/BTTaskNode.cpp
--- a/Engine/Source/Runtime/AIModule/Private/BehaviorTree/BTTaskNode.cpp
+++ b/Engine/Source/Runtime/AIModule/Private/BehaviorTree/BTTaskNode.cpp
@@ -3,7 +3,7 @@
 #include "AIModulePrivate.h"
 #include "BehaviorTree/BTTaskNode.h"
 
-UBTTaskNode::UBTTaskNode(const class FPostConstructInitializeProperties& PCIP) : Super(PCIP)
+UBTTaskNode::UBTTaskNode(const FObjectInitializer& ObjectInitializer) : Super(ObjectInitializer)
 {
 	NodeName = "UnknownTask";
 	bNotifyTick = false;
diff --git a/Engine/Source/Runtime/AIModule/Private/BehaviorTree/BehaviorTree.cpp b/Engine/Source/Runtime/AIModule/Private/BehaviorTree/BehaviorTree.cpp
--- a/Engine/Source/Runtime/AIModule/Private/BehaviorTree/BehaviorTree.cpp
+++ b/Engine/Source/Runtime/AIModule/Private/BehaviorTree/BehaviorTree.cpp
@@ -5,6 +5,6 @@
 
 DEFINE_LOG_CATEGORY(LogBehaviorTree);
 
-UBehaviorTree::UBehaviorTree(const class FPostConstructInitializeProperties& PCIP) : Super(PCIP)
+UBehaviorTree::UBehaviorTree(const FObjectInitializer& ObjectInitializer) : Super(ObjectInitializer)
 {
 }
diff --git a/Engine/Source/Runtime/AIModule/Private/BehaviorTree/BehaviorTreeComponent.cpp b/Engine/Source/Runtime/AIModule/Private/BehaviorTree/BehaviorTreeComponent.cpp
--- a/Engine/Source/Runtime/AIModule/Private/BehaviorTree/BehaviorTreeComponent.cpp
+++ b/Engine/Source/Runtime/AIModule/Private/BehaviorTree/BehaviorTreeComponent.cpp
@@ -13,7 +13,7 @@
 // UBehaviorTreeComponent
 //----------------------------------------------------------------------//
 
-UBehaviorTreeComponent::UBehaviorTreeComponent(const FPostConstructInitializeProperties& PCIP) : Super(PCIP)
+UBehaviorTreeComponent::UBehaviorTreeComponent(const FObjectInitializer& ObjectInitializer) : Super(ObjectInitializer)
 {
 	ActiveInstanceIdx = 0;
 	bLoopExecution = false;
diff --git a/Engine/Source/Runtime/AIModule/Private/BehaviorTree/BehaviorTreeManager.cpp b/Engine/Source/Runtime/AIModule/Private/BehaviorTree/BehaviorTreeManager.cpp
--- a/Engine/Source/Runtime/AIModule/Private/BehaviorTree/BehaviorTreeManager.cpp
+++ b/Engine/Source/Runtime/AIModule/Private/BehaviorTree/BehaviorTreeManager.cpp
@@ -17,7 +17,7 @@ DEFINE_STAT(STAT_AI_BehaviorTree_NumTemplates);
 DEFINE_STAT(STAT_AI_BehaviorTree_NumInstances);
 DEFINE_STAT(STAT_AI_BehaviorTree_InstanceMemory);
 
-UBehaviorTreeManager::UBehaviorTreeManager(const class FPostConstructInitializeProperties& PCIP) : Super(PCIP)
+UBehaviorTreeManager::UBehaviorTreeManager(const FObjectInitializer& ObjectInitializer) : Super(ObjectInitializer)
 {
 	MaxDebuggerSteps = 100;
 }
diff --git a/Engine/Source/Runtime/AIModule/Private/BehaviorTree/BehaviorTreeTypes.cpp b/Engine/Source/Runtime/AIModule/Private/BehaviorTree/BehaviorTreeTypes.cpp
--- a/Engine/Source/Runtime/AIModule/Private/BehaviorTree/BehaviorTreeTypes.cpp
+++ b/Engine/Source/Runtime/AIModule/Private/BehaviorTree/BehaviorTreeTypes.cpp
@@ -304,7 +304,7 @@ void FBlackboardKeySelector::AddNameFilter(UObject* Owner)
 //----------------------------------------------------------------------//
 // UBehaviorTreeTypes
 //----------------------------------------------------------------------//
-UBehaviorTreeTypes::UBehaviorTreeTypes(const class FPostConstructInitializeProperties& PCIP) : Super(PCIP)
+UBehaviorTreeTypes::UBehaviorTreeTypes(const FObjectInitializer& ObjectInitializer) : Super(ObjectInitializer)
 {
 }
 
diff --git a/Engine/Source/Runtime/AIModule/Private/BehaviorTree/Blackboard/BlackboardKeyType.cpp b/Engine/Source/Runtime/AIModule/Private/BehaviorTree/Blackboard/BlackboardKeyType.cpp
--- a/Engine/Source/Runtime/AIModule/Private/BehaviorTree/Blackboard/BlackboardKeyType.cpp
+++ b/Engine/Source/Runtime/AIModule/Private/BehaviorTree/Blackboard/BlackboardKeyType.cpp
@@ -3,7 +3,7 @@
 #include "AIModulePrivate.h"
 #include "BehaviorTree/Blackboard/BlackboardKeyType.h"
 
-UBlackboardKeyType::UBlackboardKeyType(const class FPostConstructInitializeProperties& PCIP) : Super(PCIP)
+UBlackboardKeyType::UBlackboardKeyType(const FObjectInitializer& ObjectInitializer) : Super(ObjectInitializer)
 {
 	ValueSize = 0;
 	SupportedOp = EBlackboardKeyOperation::Basic;
diff --git a/Engine/Source/Runtime/AIModule/Private/BehaviorTree/Blackboard/BlackboardKeyType_Bool.cpp b/Engine/Source/Runtime/AIModule/Private/BehaviorTree/Blackboard/BlackboardKeyType_Bool.cpp
--- a/Engine/Source/Runtime/AIModule/Private/BehaviorTree/Blackboard/BlackboardKeyType_Bool.cpp
+++ b/Engine/Source/Runtime/AIModule/Private/BehaviorTree/Blackboard/BlackboardKeyType_Bool.cpp
@@ -3,7 +3,7 @@
 #include "AIModulePrivate.h"
 #include "BehaviorTree/Blackboard/BlackboardKeyType_Bool.h"
 
-UBlackboardKeyType_Bool::UBlackboardKeyType_Bool(const class FPostConstructInitializeProperties& PCIP) : Super(PCIP)
+UBlackboardKeyType_Bool::UBlackboardKeyType_Bool(const FObjectInitializer& ObjectInitializer) : Super(ObjectInitializer)
 {
 	ValueSize = sizeof(uint8);
 	SupportedOp = EBlackboardKeyOperation::Basic;
diff --git a/Engine/Source/Runtime/AIModule/Private/BehaviorTree/Blackboard/BlackboardKeyType_Class.cpp b/Engine/Source/Runtime/AIModule/Private/BehaviorTree/Blackboard/BlackboardKeyType_Class.cpp
--- a/Engine/Source/Runtime/AIModule/Private/BehaviorTree/Blackboard/BlackboardKeyType_Class.cpp
+++ b/Engine/Source/Runtime/AIModule/Private/BehaviorTree/Blackboard/BlackboardKeyType_Class.cpp
@@ -3,7 +3,7 @@
 #include "AIModulePrivate.h"
 #include "BehaviorTree/Blackboard/BlackboardKeyType_Class.h"
 
-UBlackboardKeyType_Class::UBlackboardKeyType_Class(const class FPostConstructInitializeProperties& PCIP) : Super(PCIP)
+UBlackboardKeyType_Class::UBlackboardKeyType_Class(const FObjectInitializer& ObjectInitializer) : Super(ObjectInitializer)
 {
 	ValueSize = sizeof(TWeakObjectPtr<UClass>);
 	BaseClass = UObject::StaticClass();
diff --git a/Engine/Source/Runtime/AIModule/Private/BehaviorTree/Blackboard/BlackboardKeyType_Enum.cpp b/Engine/Source/Runtime/AIModule/Private/BehaviorTree/Blackboard/BlackboardKeyType_Enum.cpp
--- a/Engine/Source/Runtime/AIModule/Private/BehaviorTree/Blackboard/BlackboardKeyType_Enum.cpp
+++ b/Engine/Source/Runtime/AIModule/Private/BehaviorTree/Blackboard/BlackboardKeyType_Enum.cpp
@@ -3,7 +3,7 @@
 #include "AIModulePrivate.h"
 #include "BehaviorTree/Blackboard/BlackboardKeyType_Enum.h"
 
-UBlackboardKeyType_Enum::UBlackboardKeyType_Enum(const class FPostConstructInitializeProperties& PCIP) : Super(PCIP)
+UBlackboardKeyType_Enum::UBlackboardKeyType_Enum(const FObjectInitializer& ObjectInitializer) : Super(ObjectInitializer)
 {
 	ValueSize = sizeof(uint8);
 	SupportedOp = EBlackboardKeyOperation::Arithmetic;
diff --git a/Engine/Source/Runtime/AIModule/Private/BehaviorTree/Blackboard/BlackboardKeyType_Float.cpp b/Engine/Source/Runtime/AIModule/Private/BehaviorTree/Blackboard/BlackboardKeyType_Float.cpp
--- a/Engine/Source/Runtime/AIModule/Private/BehaviorTree/Blackboard/BlackboardKeyType_Float.cpp
+++ b/Engine/Source/Runtime/AIModule/Private/BehaviorTree/Blackboard/BlackboardKeyType_Float.cpp
@@ -3,7 +3,7 @@
 #include "AIModulePrivate.h"
 #include "BehaviorTree/Blackboard/BlackboardKeyType_Float.h"
 
-UBlackboardKeyType_Float::UBlackboardKeyType_Float(const class FPostConstructInitializeProperties& PCIP) : Super(PCIP)
+UBlackboardKeyType_Float::UBlackboardKeyType_Float(const FObjectInitializer& ObjectInitializer) : Super(ObjectInitializer)
 {
 	ValueSize = sizeof(float);
 	SupportedOp = EBlackboardKeyOperation::Arithmetic;
diff --git a/Engine/Source/Runtime/AIModule/Private/BehaviorTree/Blackboard/BlackboardKeyType_Int.cpp b/Engine/Source/Runtime/AIModule/Private/BehaviorTree/Blackboard/BlackboardKeyType_Int.cpp
--- a/Engine/Source/Runtime/AIModule/Private/BehaviorTree/Blackboard/BlackboardKeyType_Int.cpp
+++ b/Engine/Source/Runtime/AIModule/Private/BehaviorTree/Blackboard/BlackboardKeyType_Int.cpp
@@ -3,7 +3,7 @@
 #include "AIModulePrivate.h"
 #include "BehaviorTree/Blackboard/BlackboardKeyType_Int.h"
 
-UBlackboardKeyType_Int::UBlackboardKeyType_Int(const class FPostConstructInitializeProperties& PCIP) : Super(PCIP)
+UBlackboardKeyType_Int::UBlackboardKeyType_Int(const FObjectInitializer& ObjectInitializer) : Super(ObjectInitializer)
 {
 	ValueSize = sizeof(int32);
 	SupportedOp = EBlackboardKeyOperation::Arithmetic;
diff --git a/Engine/Source/Runtime/AIModule/Private/BehaviorTree/Blackboard/BlackboardKeyType_Name.cpp b/Engine/Source/Runtime/AIModule/Private/BehaviorTree/Blackboard/BlackboardKeyType_Name.cpp
--- a/Engine/Source/Runtime/AIModule/Private/BehaviorTree/Blackboard/BlackboardKeyType_Name.cpp
+++ b/Engine/Source/Runtime/AIModule/Private/BehaviorTree/Blackboard/BlackboardKeyType_Name.cpp
@@ -3,7 +3,7 @@
 #include "AIModulePrivate.h"
 #include "BehaviorTree/Blackboard/BlackboardKeyType_Name.h"
 
-UBlackboardKeyType_Name::UBlackboardKeyType_Name(const class FPostConstructInitializeProperties& PCIP) : Super(PCIP)
+UBlackboardKeyType_Name::UBlackboardKeyType_Name(const FObjectInitializer& ObjectInitializer) : Super(ObjectInitializer)
 {
 	ValueSize = sizeof(FName);
 	SupportedOp = EBlackboardKeyOperation::Text;
diff --git a/Engine/Source/Runtime/AIModule/Private/BehaviorTree/Blackboard/BlackboardKeyType_NativeEnum.cpp b/Engine/Source/Runtime/AIModule/Private/BehaviorTree/Blackboard/BlackboardKeyType_NativeEnum.cpp
--- a/Engine/Source/Runtime/AIModule/Private/BehaviorTree/Blackboard/BlackboardKeyType_NativeEnum.cpp
+++ b/Engine/Source/Runtime/AIModule/Private/BehaviorTree/Blackboard/BlackboardKeyType_NativeEnum.cpp
@@ -3,7 +3,7 @@
 #include "AIModulePrivate.h"
 #include "BehaviorTree/Blackboard/BlackboardKeyType_NativeEnum.h"
 
-UBlackboardKeyType_NativeEnum::UBlackboardKeyType_NativeEnum(const class FPostConstructInitializeProperties& PCIP) : Super(PCIP)
+UBlackboardKeyType_NativeEnum::UBlackboardKeyType_NativeEnum(const FObjectInitializer& ObjectInitializer) : Super(ObjectInitializer)
 {
 	ValueSize = sizeof(uint8);
 	SupportedOp = EBlackboardKeyOperation::Arithmetic;
diff --git a/Engine/Source/Runtime/AIModule/Private/BehaviorTree/Blackboard/BlackboardKeyType_Object.cpp b/Engine/Source/Runtime/AIModule/Private/BehaviorTree/Blackboard/BlackboardKeyType_Object.cpp
--- a/Engine/Source/Runtime/AIModule/Private/BehaviorTree/Blackboard/BlackboardKeyType_Object.cpp
+++ b/Engine/Source/Runtime/AIModule/Private/BehaviorTree/Blackboard/BlackboardKeyType_Object.cpp
@@ -3,7 +3,7 @@
 #include "AIModulePrivate.h"
 #include "BehaviorTree/Blackboard/BlackboardKeyType_Object.h"
 
-UBlackboardKeyType_Object::UBlackboardKeyType_Object(const class FPostConstructInitializeProperties& PCIP) : Super(PCIP)
+UBlackboardKeyType_Object::UBlackboardKeyType_Object(const FObjectInitializer& ObjectInitializer) : Super(ObjectInitializer)
 {
 	ValueSize = sizeof(FWeakObjectPtr);
 	BaseClass = UObject::StaticClass();
diff --git a/Engine/Source/Runtime/AIModule/Private/BehaviorTree/Blackboard/BlackboardKeyType_Rotator.cpp b/Engine/Source/Runtime/AIModule/Private/BehaviorTree/Blackboard/BlackboardKeyType_Rotator.cpp
--- a/Engine/Source/Runtime/AIModule/Private/BehaviorTree/Blackboard/BlackboardKeyType_Rotator.cpp
+++ b/Engine/Source/Runtime/AIModule/Private/BehaviorTree/Blackboard/BlackboardKeyType_Rotator.cpp
@@ -3,7 +3,7 @@
 #include "AIModulePrivate.h"
 #include "BehaviorTree/Blackboard/BlackboardKeyType_Rotator.h"
 
-UBlackboardKeyType_Rotator::UBlackboardKeyType_Rotator(const class FPostConstructInitializeProperties& PCIP) : Super(PCIP)
+UBlackboardKeyType_Rotator::UBlackboardKeyType_Rotator(const FObjectInitializer& ObjectInitializer) : Super(ObjectInitializer)
 {
 	ValueSize = sizeof(FRotator);
 	SupportedOp = EBlackboardKeyOperation::Basic;
diff --git a/Engine/Source/Runtime/AIModule/Private/BehaviorTree/Blackboard/BlackboardKeyType_String.cpp b/Engine/Source/Runtime/AIModule/Private/BehaviorTree/Blackboard/BlackboardKeyType_String.cpp
--- a/Engine/Source/Runtime/AIModule/Private/BehaviorTree/Blackboard/BlackboardKeyType_String.cpp
+++ b/Engine/Source/Runtime/AIModule/Private/BehaviorTree/Blackboard/BlackboardKeyType_String.cpp
@@ -3,7 +3,7 @@
 #include "AIModulePrivate.h"
 #include "BehaviorTree/Blackboard/BlackboardKeyType_String.h"
 
-UBlackboardKeyType_String::UBlackboardKeyType_String(const class FPostConstructInitializeProperties& PCIP) : Super(PCIP)
+UBlackboardKeyType_String::UBlackboardKeyType_String(const FObjectInitializer& ObjectInitializer) : Super(ObjectInitializer)
 {
 	ValueSize = sizeof(FString);
 	SupportedOp = EBlackboardKeyOperation::Text;
diff --git a/Engine/Source/Runtime/AIModule/Private/BehaviorTree/Blackboard/BlackboardKeyType_Vector.cpp b/Engine/Source/Runtime/AIModule/Private/BehaviorTree/Blackboard/BlackboardKeyType_Vector.cpp
--- a/Engine/Source/Runtime/AIModule/Private/BehaviorTree/Blackboard/BlackboardKeyType_Vector.cpp
+++ b/Engine/Source/Runtime/AIModule/Private/BehaviorTree/Blackboard/BlackboardKeyType_Vector.cpp
@@ -3,7 +3,7 @@
 #include "AIModulePrivate.h"
 #include "BehaviorTree/Blackboard/BlackboardKeyType_Vector.h"
 
-UBlackboardKeyType_Vector::UBlackboardKeyType_Vector(const class FPostConstructInitializeProperties& PCIP) : Super(PCIP)
+UBlackboardKeyType_Vector::UBlackboardKeyType_Vector(const FObjectInitializer& ObjectInitializer) : Super(ObjectInitializer)
 {
 	ValueSize = sizeof(FVector);
 	SupportedOp = EBlackboardKeyOperation::Basic;
diff --git a/Engine/Source/Runtime/AIModule/Private/BehaviorTree/BlackboardComponent.cpp b/Engine/Source/Runtime/AIModule/Private/BehaviorTree/BlackboardComponent.cpp
--- a/Engine/Source/Runtime/AIModule/Private/BehaviorTree/BlackboardComponent.cpp
+++ b/Engine/Source/Runtime/AIModule/Private/BehaviorTree/BlackboardComponent.cpp
@@ -5,7 +5,7 @@
 #include "BehaviorTree/Blackboard/BlackboardKeyAllTypes.h"
 #include "BehaviorTree/BlackboardComponent.h"
 
-UBlackboardComponent::UBlackboardComponent(const class FPostConstructInitializeProperties& PCIP) : Super(PCIP)
+UBlackboardComponent::UBlackboardComponent(const FObjectInitializer& ObjectInitializer) : Super(ObjectInitializer)
 {
 	PrimaryComponentTick.bCanEverTick = false;
 	bWantsInitializeComponent = true;
diff --git a/Engine/Source/Runtime/AIModule/Private/BehaviorTree/BlackboardData.cpp b/Engine/Source/Runtime/AIModule/Private/BehaviorTree/BlackboardData.cpp
--- a/Engine/Source/Runtime/AIModule/Private/BehaviorTree/BlackboardData.cpp
+++ b/Engine/Source/Runtime/AIModule/Private/BehaviorTree/BlackboardData.cpp
@@ -15,7 +15,7 @@ static void UpdatePersistentKeys(class UBlackboardData* Asset)
 	}
 }
 
-UBlackboardData::UBlackboardData(const class FPostConstructInitializeProperties& PCIP) : Super(PCIP)
+UBlackboardData::UBlackboardData(const FObjectInitializer& ObjectInitializer) : Super(ObjectInitializer)
 {
 	if (HasAnyFlags(RF_ClassDefaultObject))
 	{
diff --git a/Engine/Source/Runtime/AIModule/Private/BehaviorTree/Composites/BTComposite_Selector.cpp b/Engine/Source/Runtime/AIModule/Private/BehaviorTree/Composites/BTComposite_Selector.cpp
--- a/Engine/Source/Runtime/AIModule/Private/BehaviorTree/Composites/BTComposite_Selector.cpp
+++ b/Engine/Source/Runtime/AIModule/Private/BehaviorTree/Composites/BTComposite_Selector.cpp
@@ -3,7 +3,7 @@
 #include "AIModulePrivate.h"
 #include "BehaviorTree/Composites/BTComposite_Selector.h"
 
-UBTComposite_Selector::UBTComposite_Selector(const class FPostConstructInitializeProperties& PCIP) : Super(PCIP)
+UBTComposite_Selector::UBTComposite_Selector(const FObjectInitializer& ObjectInitializer) : Super(ObjectInitializer)
 {
 	NodeName = "Selector";
 
diff --git a/Engine/Source/Runtime/AIModule/Private/BehaviorTree/Composites/BTComposite_Sequence.cpp b/Engine/Source/Runtime/AIModule/Private/BehaviorTree/Composites/BTComposite_Sequence.cpp
--- a/Engine/Source/Runtime/AIModule/Private/BehaviorTree/Composites/BTComposite_Sequence.cpp
+++ b/Engine/Source/Runtime/AIModule/Private/BehaviorTree/Composites/BTComposite_Sequence.cpp
@@ -3,7 +3,7 @@
 #include "AIModulePrivate.h"
 #include "BehaviorTree/Composites/BTComposite_Sequence.h"
 
-UBTComposite_Sequence::UBTComposite_Sequence(const class FPostConstructInitializeProperties& PCIP) : Super(PCIP)
+UBTComposite_Sequence::UBTComposite_Sequence(const FObjectInitializer& ObjectInitializer) : Super(ObjectInitializer)
 {
 	NodeName = "Sequence";
 
diff --git a/Engine/Source/Runtime/AIModule/Private/BehaviorTree/Composites/BTComposite_SimpleParallel.cpp b/Engine/Source/Runtime/AIModule/Private/BehaviorTree/Composites/BTComposite_SimpleParallel.cpp
--- a/Engine/Source/Runtime/AIModule/Private/BehaviorTree/Composites/BTComposite_SimpleParallel.cpp
+++ b/Engine/Source/Runtime/AIModule/Private/BehaviorTree/Composites/BTComposite_SimpleParallel.cpp
@@ -3,7 +3,7 @@
 #include "AIModulePrivate.h"
 #include "BehaviorTree/Composites/BTComposite_SimpleParallel.h"
 
-UBTComposite_SimpleParallel::UBTComposite_SimpleParallel(const class FPostConstructInitializeProperties& PCIP) : Super(PCIP)
+UBTComposite_SimpleParallel::UBTComposite_SimpleParallel(const FObjectInitializer& ObjectInitializer) : Super(ObjectInitializer)
 {
 	NodeName = "Simple Parallel";
 	bUseChildExecutionNotify = true;
diff --git a/Engine/Source/Runtime/AIModule/Private/BehaviorTree/Decorators/BTDecorator_Blackboard.cpp b/Engine/Source/Runtime/AIModule/Private/BehaviorTree/Decorators/BTDecorator_Blackboard.cpp
--- a/Engine/Source/Runtime/AIModule/Private/BehaviorTree/Decorators/BTDecorator_Blackboard.cpp
+++ b/Engine/Source/Runtime/AIModule/Private/BehaviorTree/Decorators/BTDecorator_Blackboard.cpp
@@ -5,7 +5,7 @@
 #include "BehaviorTree/BlackboardComponent.h"
 #include "BehaviorTree/Blackboard/BlackboardKeyAllTypes.h"
 
-UBTDecorator_Blackboard::UBTDecorator_Blackboard(const class FPostConstructInitializeProperties& PCIP) : Super(PCIP)
+UBTDecorator_Blackboard::UBTDecorator_Blackboard(const FObjectInitializer& ObjectInitializer) : Super(ObjectInitializer)
 {
 	NodeName = "Blackboard Based Condition";
 	NotifyObserver = EBTBlackboardRestart::ResultChange;
diff --git a/Engine/Source/Runtime/AIModule/Private/BehaviorTree/Decorators/BTDecorator_BlackboardBase.cpp b/Engine/Source/Runtime/AIModule/Private/BehaviorTree/Decorators/BTDecorator_BlackboardBase.cpp
--- a/Engine/Source/Runtime/AIModule/Private/BehaviorTree/Decorators/BTDecorator_BlackboardBase.cpp
+++ b/Engine/Source/Runtime/AIModule/Private/BehaviorTree/Decorators/BTDecorator_BlackboardBase.cpp
@@ -4,7 +4,7 @@
 #include "BehaviorTree/BlackboardComponent.h"
 #include "BehaviorTree/Decorators/BTDecorator_BlackboardBase.h"
 
-UBTDecorator_BlackboardBase::UBTDecorator_BlackboardBase(const class FPostConstructInitializeProperties& PCIP) : Super(PCIP)
+UBTDecorator_BlackboardBase::UBTDecorator_BlackboardBase(const FObjectInitializer& ObjectInitializer) : Super(ObjectInitializer)
 {
 	NodeName = "BlackboardBase";
 
diff --git a/Engine/Source/Runtime/AIModule/Private/BehaviorTree/Decorators/BTDecorator_BlueprintBase.cpp b/Engine/Source/Runtime/AIModule/Private/BehaviorTree/Decorators/BTDecorator_BlueprintBase.cpp
--- a/Engine/Source/Runtime/AIModule/Private/BehaviorTree/Decorators/BTDecorator_BlueprintBase.cpp
+++ b/Engine/Source/Runtime/AIModule/Private/BehaviorTree/Decorators/BTDecorator_BlueprintBase.cpp
@@ -6,7 +6,7 @@
 #include "BehaviorTree/BlackboardComponent.h"
 #include "BehaviorTree/BlackboardData.h"
 
-UBTDecorator_BlueprintBase::UBTDecorator_BlueprintBase(const class FPostConstructInitializeProperties& PCIP) : Super(PCIP)
+UBTDecorator_BlueprintBase::UBTDecorator_BlueprintBase(const FObjectInitializer& ObjectInitializer) : Super(ObjectInitializer)
 {
 	UClass* StopAtClass = UBTDecorator_BlueprintBase::StaticClass();
 	bImplementsReceiveTick = BlueprintNodeHelpers::HasBlueprintFunction(TEXT("ReceiveTick"), this, StopAtClass);
diff --git a/Engine/Source/Runtime/AIModule/Private/BehaviorTree/Decorators/BTDecorator_CompareBBEntries.cpp b/Engine/Source/Runtime/AIModule/Private/BehaviorTree/Decorators/BTDecorator_CompareBBEntries.cpp
--- a/Engine/Source/Runtime/AIModule/Private/BehaviorTree/Decorators/BTDecorator_CompareBBEntries.cpp
+++ b/Engine/Source/Runtime/AIModule/Private/BehaviorTree/Decorators/BTDecorator_CompareBBEntries.cpp
@@ -4,7 +4,7 @@
 #include "BehaviorTree/BlackboardComponent.h"
 #include "BehaviorTree/Decorators/BTDecorator_CompareBBEntries.h"
 
-UBTDecorator_CompareBBEntries::UBTDecorator_CompareBBEntries(const class FPostConstructInitializeProperties& PCIP) : Super(PCIP)
+UBTDecorator_CompareBBEntries::UBTDecorator_CompareBBEntries(const FObjectInitializer& ObjectInitializer) : Super(ObjectInitializer)
 {
 	NodeName = "Compare Blackboard entries";
 
diff --git a/Engine/Source/Runtime/AIModule/Private/BehaviorTree/Decorators/BTDecorator_ConditionalLoop.cpp b/Engine/Source/Runtime/AIModule/Private/BehaviorTree/Decorators/BTDecorator_ConditionalLoop.cpp
--- a/Engine/Source/Runtime/AIModule/Private/BehaviorTree/Decorators/BTDecorator_ConditionalLoop.cpp
+++ b/Engine/Source/Runtime/AIModule/Private/BehaviorTree/Decorators/BTDecorator_ConditionalLoop.cpp
@@ -4,7 +4,7 @@
 #include "BehaviorTree/BTCompositeNode.h"
 #include "BehaviorTree/Decorators/BTDecorator_ConditionalLoop.h"
 
-UBTDecorator_ConditionalLoop::UBTDecorator_ConditionalLoop(const class FPostConstructInitializeProperties& PCIP) : Super(PCIP)
+UBTDecorator_ConditionalLoop::UBTDecorator_ConditionalLoop(const FObjectInitializer& ObjectInitializer) : Super(ObjectInitializer)
 {
 	NodeName = "Conditional Loop";
 	bNotifyDeactivation = true;
diff --git a/Engine/Source/Runtime/AIModule/Private/BehaviorTree/Decorators/BTDecorator_ConeCheck.cpp b/Engine/Source/Runtime/AIModule/Private/BehaviorTree/Decorators/BTDecorator_ConeCheck.cpp
--- a/Engine/Source/Runtime/AIModule/Private/BehaviorTree/Decorators/BTDecorator_ConeCheck.cpp
+++ b/Engine/Source/Runtime/AIModule/Private/BehaviorTree/Decorators/BTDecorator_ConeCheck.cpp
@@ -4,7 +4,7 @@
 #include "BehaviorTree/BlackboardComponent.h"
 #include "BehaviorTree/Decorators/BTDecorator_ConeCheck.h"
 
-UBTDecorator_ConeCheck::UBTDecorator_ConeCheck(const class FPostConstructInitializeProperties& PCIP) : Super(PCIP)
+UBTDecorator_ConeCheck::UBTDecorator_ConeCheck(const FObjectInitializer& ObjectInitializer) : Super(ObjectInitializer)
 {
 	NodeName = "Cone Check";
 
diff --git a/Engine/Source/Runtime/AIModule/Private/BehaviorTree/Decorators/BTDecorator_Cooldown.cpp b/Engine/Source/Runtime/AIModule/Private/BehaviorTree/Decorators/BTDecorator_Cooldown.cpp
--- a/Engine/Source/Runtime/AIModule/Private/BehaviorTree/Decorators/BTDecorator_Cooldown.cpp
+++ b/Engine/Source/Runtime/AIModule/Private/BehaviorTree/Decorators/BTDecorator_Cooldown.cpp
@@ -3,7 +3,7 @@
 #include "AIModulePrivate.h"
 #include "BehaviorTree/Decorators/BTDecorator_Cooldown.h"
 
-UBTDecorator_Cooldown::UBTDecorator_Cooldown(const class FPostConstructInitializeProperties& PCIP) : Super(PCIP)
+UBTDecorator_Cooldown::UBTDecorator_Cooldown(const FObjectInitializer& ObjectInitializer) : Super(ObjectInitializer)
 {
 	NodeName = "Cooldown";
 	CoolDownTime = 5.0f;
diff --git a/Engine/Source/Runtime/AIModule/Private/BehaviorTree/Decorators/BTDecorator_DoesPathExist.cpp b/Engine/Source/Runtime/AIModule/Private/BehaviorTree/Decorators/BTDecorator_DoesPathExist.cpp
--- a/Engine/Source/Runtime/AIModule/Private/BehaviorTree/Decorators/BTDecorator_DoesPathExist.cpp
+++ b/Engine/Source/Runtime/AIModule/Private/BehaviorTree/Decorators/BTDecorator_DoesPathExist.cpp
@@ -8,7 +8,7 @@
 #include "BehaviorTree/BlackboardComponent.h"
 #include "BehaviorTree/Decorators/BTDecorator_DoesPathExist.h"
 
-UBTDecorator_DoesPathExist::UBTDecorator_DoesPathExist(const FPostConstructInitializeProperties& PCIP) : Super(PCIP)
+UBTDecorator_DoesPathExist::UBTDecorator_DoesPathExist(const FObjectInitializer& ObjectInitializer) : Super(ObjectInitializer)
 {
 	NodeName = "Does path exist";
 
diff --git a/Engine/Source/Runtime/AIModule/Private/BehaviorTree/Decorators/BTDecorator_ForceSuccess.cpp b/Engine/Source/Runtime/AIModule/Private/BehaviorTree/Decorators/BTDecorator_ForceSuccess.cpp
--- a/Engine/Source/Runtime/AIModule/Private/BehaviorTree/Decorators/BTDecorator_ForceSuccess.cpp
+++ b/Engine/Source/Runtime/AIModule/Private/BehaviorTree/Decorators/BTDecorator_ForceSuccess.cpp
@@ -3,7 +3,7 @@
 #include "AIModulePrivate.h"
 #include "BehaviorTree/Decorators/BTDecorator_ForceSuccess.h"
 
-UBTDecorator_ForceSuccess::UBTDecorator_ForceSuccess(const class FPostConstructInitializeProperties& PCIP) : Super(PCIP)
+UBTDecorator_ForceSuccess::UBTDecorator_ForceSuccess(const FObjectInitializer& ObjectInitializer) : Super(ObjectInitializer)
 {
 	NodeName = TEXT("Force Success");
 	bNotifyProcessed = true;
diff --git a/Engine/Source/Runtime/AIModule/Private/BehaviorTree/Decorators/BTDecorator_KeepInCone.cpp b/Engine/Source/Runtime/AIModule/Private/BehaviorTree/Decorators/BTDecorator_KeepInCone.cpp
--- a/Engine/Source/Runtime/AIModule/Private/BehaviorTree/Decorators/BTDecorator_KeepInCone.cpp
+++ b/Engine/Source/Runtime/AIModule/Private/BehaviorTree/Decorators/BTDecorator_KeepInCone.cpp
@@ -4,7 +4,7 @@
 #include "BehaviorTree/BlackboardComponent.h"
 #include "BehaviorTree/Decorators/BTDecorator_KeepInCone.h"
 
-UBTDecorator_KeepInCone::UBTDecorator_KeepInCone(const class FPostConstructInitializeProperties& PCIP) : Super(PCIP)
+UBTDecorator_KeepInCone::UBTDecorator_KeepInCone(const FObjectInitializer& ObjectInitializer) : Super(ObjectInitializer)
 {
 	NodeName = "Keep in Cone";
 
diff --git a/Engine/Source/Runtime/AIModule/Private/BehaviorTree/Decorators/BTDecorator_Loop.cpp b/Engine/Source/Runtime/AIModule/Private/BehaviorTree/Decorators/BTDecorator_Loop.cpp
--- a/Engine/Source/Runtime/AIModule/Private/BehaviorTree/Decorators/BTDecorator_Loop.cpp
+++ b/Engine/Source/Runtime/AIModule/Private/BehaviorTree/Decorators/BTDecorator_Loop.cpp
@@ -4,7 +4,7 @@
 #include "BehaviorTree/BTCompositeNode.h"
 #include "BehaviorTree/Decorators/BTDecorator_Loop.h"
 
-UBTDecorator_Loop::UBTDecorator_Loop(const class FPostConstructInitializeProperties& PCIP) : Super(PCIP)
+UBTDecorator_Loop::UBTDecorator_Loop(const FObjectInitializer& ObjectInitializer) : Super(ObjectInitializer)
 {
 	NodeName = "Loop";
 	NumLoops = 3;
diff --git a/Engine/Source/Runtime/AIModule/Private/BehaviorTree/Decorators/BTDecorator_ReachedMoveGoal.cpp b/Engine/Source/Runtime/AIModule/Private/BehaviorTree/Decorators/BTDecorator_ReachedMoveGoal.cpp
--- a/Engine/Source/Runtime/AIModule/Private/BehaviorTree/Decorators/BTDecorator_ReachedMoveGoal.cpp
+++ b/Engine/Source/Runtime/AIModule/Private/BehaviorTree/Decorators/BTDecorator_ReachedMoveGoal.cpp
@@ -3,7 +3,7 @@
 #include "AIModulePrivate.h"
 #include "BehaviorTree/Decorators/BTDecorator_ReachedMoveGoal.h"
 
-UBTDecorator_ReachedMoveGoal::UBTDecorator_ReachedMoveGoal(const FPostConstructInitializeProperties& PCIP) : Super(PCIP)
+UBTDecorator_ReachedMoveGoal::UBTDecorator_ReachedMoveGoal(const FObjectInitializer& ObjectInitializer) : Super(ObjectInitializer)
 {
 	NodeName = "Reached move goal";
 
diff --git a/Engine/Source/Runtime/AIModule/Private/BehaviorTree/Decorators/BTDecorator_TimeLimit.cpp b/Engine/Source/Runtime/AIModule/Private/BehaviorTree/Decorators/BTDecorator_TimeLimit.cpp
--- a/Engine/Source/Runtime/AIModule/Private/BehaviorTree/Decorators/BTDecorator_TimeLimit.cpp
+++ b/Engine/Source/Runtime/AIModule/Private/BehaviorTree/Decorators/BTDecorator_TimeLimit.cpp
@@ -3,7 +3,7 @@
 #include "AIModulePrivate.h"
 #include "BehaviorTree/Decorators/BTDecorator_TimeLimit.h"
 
-UBTDecorator_TimeLimit::UBTDecorator_TimeLimit(const class FPostConstructInitializeProperties& PCIP) : Super(PCIP)
+UBTDecorator_TimeLimit::UBTDecorator_TimeLimit(const FObjectInitializer& ObjectInitializer) : Super(ObjectInitializer)
 {
 	NodeName = "TimeLimit";
 	TimeLimit = 5.0f;
diff --git a/Engine/Source/Runtime/AIModule/Private/BehaviorTree/Services/BTService_BlackboardBase.cpp b/Engine/Source/Runtime/AIModule/Private/BehaviorTree/Services/BTService_BlackboardBase.cpp
--- a/Engine/Source/Runtime/AIModule/Private/BehaviorTree/Services/BTService_BlackboardBase.cpp
+++ b/Engine/Source/Runtime/AIModule/Private/BehaviorTree/Services/BTService_BlackboardBase.cpp
@@ -3,7 +3,7 @@
 #include "AIModulePrivate.h"
 #include "BehaviorTree/Services/BTService_BlackboardBase.h"
 
-UBTService_BlackboardBase::UBTService_BlackboardBase(const class FPostConstructInitializeProperties& PCIP) : Super(PCIP)
+UBTService_BlackboardBase::UBTService_BlackboardBase(const FObjectInitializer& ObjectInitializer) : Super(ObjectInitializer)
 {
 	NodeName = "BlackboardBase";
 
diff --git a/Engine/Source/Runtime/AIModule/Private/BehaviorTree/Services/BTService_BlueprintBase.cpp b/Engine/Source/Runtime/AIModule/Private/BehaviorTree/Services/BTService_BlueprintBase.cpp
--- a/Engine/Source/Runtime/AIModule/Private/BehaviorTree/Services/BTService_BlueprintBase.cpp
+++ b/Engine/Source/Runtime/AIModule/Private/BehaviorTree/Services/BTService_BlueprintBase.cpp
@@ -4,7 +4,7 @@
 #include "BlueprintNodeHelpers.h"
 #include "BehaviorTree/Services/BTService_BlueprintBase.h"
 
-UBTService_BlueprintBase::UBTService_BlueprintBase(const FPostConstructInitializeProperties& PCIP) : Super(PCIP)
+UBTService_BlueprintBase::UBTService_BlueprintBase(const FObjectInitializer& ObjectInitializer) : Super(ObjectInitializer)
 {
 	UClass* StopAtClass = UBTService_BlueprintBase::StaticClass();
 	bImplementsReceiveTick = BlueprintNodeHelpers::HasBlueprintFunction(TEXT("ReceiveTick"), this, StopAtClass);
diff --git a/Engine/Source/Runtime/AIModule/Private/BehaviorTree/Services/BTService_DefaultFocus.cpp b/Engine/Source/Runtime/AIModule/Private/BehaviorTree/Services/BTService_DefaultFocus.cpp
--- a/Engine/Source/Runtime/AIModule/Private/BehaviorTree/Services/BTService_DefaultFocus.cpp
+++ b/Engine/Source/Runtime/AIModule/Private/BehaviorTree/Services/BTService_DefaultFocus.cpp
@@ -4,8 +4,8 @@
 #include "BehaviorTree/Services/BTService_DefaultFocus.h"
 #include "BehaviorTree/Blackboard/BlackboardKeyAllTypes.h"
 
-UBTService_DefaultFocus::UBTService_DefaultFocus(const FPostConstructInitializeProperties& PCIP) 
-	: Super(PCIP)
+UBTService_DefaultFocus::UBTService_DefaultFocus(const FObjectInitializer& ObjectInitializer) 
+	: Super(ObjectInitializer)
 {
 	NodeName = "Set default focus";
 
diff --git a/Engine/Source/Runtime/AIModule/Private/BehaviorTree/Tasks/BTTask_BlackboardBase.cpp b/Engine/Source/Runtime/AIModule/Private/BehaviorTree/Tasks/BTTask_BlackboardBase.cpp
--- a/Engine/Source/Runtime/AIModule/Private/BehaviorTree/Tasks/BTTask_BlackboardBase.cpp
+++ b/Engine/Source/Runtime/AIModule/Private/BehaviorTree/Tasks/BTTask_BlackboardBase.cpp
@@ -3,7 +3,7 @@
 #include "AIModulePrivate.h"
 #include "BehaviorTree/Tasks/BTTask_BlackboardBase.h"
 
-UBTTask_BlackboardBase::UBTTask_BlackboardBase(const class FPostConstructInitializeProperties& PCIP) : Super(PCIP)
+UBTTask_BlackboardBase::UBTTask_BlackboardBase(const FObjectInitializer& ObjectInitializer) : Super(ObjectInitializer)
 {
 	NodeName = "BlackboardBase";
 
diff --git a/Engine/Source/Runtime/AIModule/Private/BehaviorTree/Tasks/BTTask_BlueprintBase.cpp b/Engine/Source/Runtime/AIModule/Private/BehaviorTree/Tasks/BTTask_BlueprintBase.cpp
--- a/Engine/Source/Runtime/AIModule/Private/BehaviorTree/Tasks/BTTask_BlueprintBase.cpp
+++ b/Engine/Source/Runtime/AIModule/Private/BehaviorTree/Tasks/BTTask_BlueprintBase.cpp
@@ -4,7 +4,7 @@
 #include "BlueprintNodeHelpers.h"
 #include "BehaviorTree/Tasks/BTTask_BlueprintBase.h"
 
-UBTTask_BlueprintBase::UBTTask_BlueprintBase(const FPostConstructInitializeProperties& PCIP) : Super(PCIP)
+UBTTask_BlueprintBase::UBTTask_BlueprintBase(const FObjectInitializer& ObjectInitializer) : Super(ObjectInitializer)
 {
 	UClass* StopAtClass = UBTTask_BlueprintBase::StaticClass();
 	bImplementsReceiveTick = BlueprintNodeHelpers::HasBlueprintFunction(TEXT("ReceiveTick"), this, StopAtClass);
diff --git a/Engine/Source/Runtime/AIModule/Private/BehaviorTree/Tasks/BTTask_MakeNoise.cpp b/Engine/Source/Runtime/AIModule/Private/BehaviorTree/Tasks/BTTask_MakeNoise.cpp
--- a/Engine/Source/Runtime/AIModule/Private/BehaviorTree/Tasks/BTTask_MakeNoise.cpp
+++ b/Engine/Source/Runtime/AIModule/Private/BehaviorTree/Tasks/BTTask_MakeNoise.cpp
@@ -3,8 +3,8 @@
 #include "AIModulePrivate.h"
 #include "BehaviorTree/Tasks/BTTask_MakeNoise.h"
 
-UBTTask_MakeNoise::UBTTask_MakeNoise(const FPostConstructInitializeProperties& PCIP) 
-	: Super(PCIP)
+UBTTask_MakeNoise::UBTTask_MakeNoise(const FObjectInitializer& ObjectInitializer) 
+	: Super(ObjectInitializer)
 	, Loudnes(1.0f)
 {
 	NodeName = "MakeNoise";
diff --git a/Engine/Source/Runtime/AIModule/Private/BehaviorTree/Tasks/BTTask_MoveDirectlyToward.cpp b/Engine/Source/Runtime/AIModule/Private/BehaviorTree/Tasks/BTTask_MoveDirectlyToward.cpp
--- a/Engine/Source/Runtime/AIModule/Private/BehaviorTree/Tasks/BTTask_MoveDirectlyToward.cpp
+++ b/Engine/Source/Runtime/AIModule/Private/BehaviorTree/Tasks/BTTask_MoveDirectlyToward.cpp
@@ -6,8 +6,8 @@
 #include "BehaviorTree/Blackboard/BlackboardKeyType_Vector.h"
 #include "BehaviorTree/Tasks/BTTask_MoveDirectlyToward.h"
 
-UBTTask_MoveDirectlyToward::UBTTask_MoveDirectlyToward(const FPostConstructInitializeProperties& PCIP) 
-	: Super(PCIP)
+UBTTask_MoveDirectlyToward::UBTTask_MoveDirectlyToward(const FObjectInitializer& ObjectInitializer) 
+	: Super(ObjectInitializer)
 	, AcceptableRadius(50.f)
 	, bProjectVectorGoalToNavigation(true)
 	, bAllowStrafe(true)
diff --git a/Engine/Source/Runtime/AIModule/Private/BehaviorTree/Tasks/BTTask_MoveTo.cpp b/Engine/Source/Runtime/AIModule/Private/BehaviorTree/Tasks/BTTask_MoveTo.cpp
--- a/Engine/Source/Runtime/AIModule/Private/BehaviorTree/Tasks/BTTask_MoveTo.cpp
+++ b/Engine/Source/Runtime/AIModule/Private/BehaviorTree/Tasks/BTTask_MoveTo.cpp
@@ -7,8 +7,8 @@
 #include "BehaviorTree/Blackboard/BlackboardKeyType_Object.h"
 #include "BehaviorTree/Blackboard/BlackboardKeyType_Vector.h"
 
-UBTTask_MoveTo::UBTTask_MoveTo(const class FPostConstructInitializeProperties& PCIP) 
-	: Super(PCIP)
+UBTTask_MoveTo::UBTTask_MoveTo(const FObjectInitializer& ObjectInitializer) 
+	: Super(ObjectInitializer)
 	, AcceptableRadius(50.f)
 	, bAllowStrafe(false)
 {
diff --git a/Engine/Source/Runtime/AIModule/Private/BehaviorTree/Tasks/BTTask_PlaySound.cpp b/Engine/Source/Runtime/AIModule/Private/BehaviorTree/Tasks/BTTask_PlaySound.cpp
--- a/Engine/Source/Runtime/AIModule/Private/BehaviorTree/Tasks/BTTask_PlaySound.cpp
+++ b/Engine/Source/Runtime/AIModule/Private/BehaviorTree/Tasks/BTTask_PlaySound.cpp
@@ -6,7 +6,7 @@
 #include "Sound/SoundCue.h"
 #include "BehaviorTree/Tasks/BTTask_PlaySound.h"
 
-UBTTask_PlaySound::UBTTask_PlaySound(const class FPostConstructInitializeProperties& PCIP) : Super(PCIP)
+UBTTask_PlaySound::UBTTask_PlaySound(const FObjectInitializer& ObjectInitializer) : Super(ObjectInitializer)
 {
 	NodeName = "PlaySound";
 }
diff --git a/Engine/Source/Runtime/AIModule/Private/BehaviorTree/Tasks/BTTask_RotateToFaceBBEntry.cpp b/Engine/Source/Runtime/AIModule/Private/BehaviorTree/Tasks/BTTask_RotateToFaceBBEntry.cpp
--- a/Engine/Source/Runtime/AIModule/Private/BehaviorTree/Tasks/BTTask_RotateToFaceBBEntry.cpp
+++ b/Engine/Source/Runtime/AIModule/Private/BehaviorTree/Tasks/BTTask_RotateToFaceBBEntry.cpp
@@ -4,8 +4,8 @@
 #include "BehaviorTree/BlackboardComponent.h"
 #include "BehaviorTree/Tasks/BTTask_RotateToFaceBBEntry.h"
 
-UBTTask_RotateToFaceBBEntry::UBTTask_RotateToFaceBBEntry(const FPostConstructInitializeProperties& PCIP)
-	: Super(PCIP)
+UBTTask_RotateToFaceBBEntry::UBTTask_RotateToFaceBBEntry(const FObjectInitializer& ObjectInitializer)
+	: Super(ObjectInitializer)
 	, Precision(10.f)
 {
 	NodeName = "Rotate to face BB entry";
diff --git a/Engine/Source/Runtime/AIModule/Private/BehaviorTree/Tasks/BTTask_RunBehavior.cpp b/Engine/Source/Runtime/AIModule/Private/BehaviorTree/Tasks/BTTask_RunBehavior.cpp
--- a/Engine/Source/Runtime/AIModule/Private/BehaviorTree/Tasks/BTTask_RunBehavior.cpp
+++ b/Engine/Source/Runtime/AIModule/Private/BehaviorTree/Tasks/BTTask_RunBehavior.cpp
@@ -4,7 +4,7 @@
 #include "BehaviorTree/Tasks/BTTask_RunBehavior.h"
 #include "BehaviorTree/BehaviorTreeManager.h"
 
-UBTTask_RunBehavior::UBTTask_RunBehavior(const class FPostConstructInitializeProperties& PCIP) : Super(PCIP)
+UBTTask_RunBehavior::UBTTask_RunBehavior(const FObjectInitializer& ObjectInitializer) : Super(ObjectInitializer)
 {
 	NodeName = "Run Behavior";
 }
diff --git a/Engine/Source/Runtime/AIModule/Private/BehaviorTree/Tasks/BTTask_RunEQSQuery.cpp b/Engine/Source/Runtime/AIModule/Private/BehaviorTree/Tasks/BTTask_RunEQSQuery.cpp
--- a/Engine/Source/Runtime/AIModule/Private/BehaviorTree/Tasks/BTTask_RunEQSQuery.cpp
+++ b/Engine/Source/Runtime/AIModule/Private/BehaviorTree/Tasks/BTTask_RunEQSQuery.cpp
@@ -6,7 +6,7 @@
 #include "EnvironmentQuery/EnvQueryManager.h"
 #include "EnvironmentQuery/EnvQuery.h"
 
-UBTTask_RunEQSQuery::UBTTask_RunEQSQuery(const class FPostConstructInitializeProperties& PCIP) : Super(PCIP)
+UBTTask_RunEQSQuery::UBTTask_RunEQSQuery(const FObjectInitializer& ObjectInitializer) : Super(ObjectInitializer)
 {
 	NodeName = "Run EQS Query";
 
diff --git a/Engine/Source/Runtime/AIModule/Private/BehaviorTree/Tasks/BTTask_Wait.cpp b/Engine/Source/Runtime/AIModule/Private/BehaviorTree/Tasks/BTTask_Wait.cpp
--- a/Engine/Source/Runtime/AIModule/Private/BehaviorTree/Tasks/BTTask_Wait.cpp
+++ b/Engine/Source/Runtime/AIModule/Private/BehaviorTree/Tasks/BTTask_Wait.cpp
@@ -3,7 +3,7 @@
 #include "AIModulePrivate.h"
 #include "BehaviorTree/Tasks/BTTask_Wait.h"
 
-UBTTask_Wait::UBTTask_Wait(const class FPostConstructInitializeProperties& PCIP) : Super(PCIP)
+UBTTask_Wait::UBTTask_Wait(const FObjectInitializer& ObjectInitializer) : Super(ObjectInitializer)
 {
 	NodeName = "Wait";
 	WaitTime = 5.0f;
diff --git a/Engine/Source/Runtime/AIModule/Private/Blueprint/AIBlueprintHelperLibrary.cpp b/Engine/Source/Runtime/AIModule/Private/Blueprint/AIBlueprintHelperLibrary.cpp
--- a/Engine/Source/Runtime/AIModule/Private/Blueprint/AIBlueprintHelperLibrary.cpp
+++ b/Engine/Source/Runtime/AIModule/Private/Blueprint/AIBlueprintHelperLibrary.cpp
@@ -12,8 +12,8 @@
 // UAIAsyncTaskBlueprintProxy
 //----------------------------------------------------------------------//
 
-UAIAsyncTaskBlueprintProxy::UAIAsyncTaskBlueprintProxy(const class FPostConstructInitializeProperties& PCIP)
-	: Super(PCIP)
+UAIAsyncTaskBlueprintProxy::UAIAsyncTaskBlueprintProxy(const FObjectInitializer& ObjectInitializer)
+	: Super(ObjectInitializer)
 {
 	MyWorld = Cast<UWorld>(GetOuter());
 	if (HasAnyFlags(RF_ClassDefaultObject) == false)
@@ -73,8 +73,8 @@ void UAIAsyncTaskBlueprintProxy::BeginDestroy()
 // UAIAsyncTaskBlueprintProxy
 //----------------------------------------------------------------------//
 
-UAIBlueprintHelperLibrary::UAIBlueprintHelperLibrary(const class FPostConstructInitializeProperties& PCIP)
-	: Super(PCIP)
+UAIBlueprintHelperLibrary::UAIBlueprintHelperLibrary(const FObjectInitializer& ObjectInitializer)
+	: Super(ObjectInitializer)
 {
 }
 
diff --git a/Engine/Source/Runtime/AIModule/Private/BrainComponent.cpp b/Engine/Source/Runtime/AIModule/Private/BrainComponent.cpp
--- a/Engine/Source/Runtime/AIModule/Private/BrainComponent.cpp
+++ b/Engine/Source/Runtime/AIModule/Private/BrainComponent.cpp
@@ -177,7 +177,7 @@ FString FAIMessageObserver::DescribeObservedMessage() const
 //////////////////////////////////////////////////////////////////////////
 // Brain component
 
-UBrainComponent::UBrainComponent(const FPostConstructInitializeProperties& PCIP) : Super(PCIP)
+UBrainComponent::UBrainComponent(const FObjectInitializer& ObjectInitializer) : Super(ObjectInitializer)
 {
 	PrimaryComponentTick.bCanEverTick = true;
 	bDoLogicRestartOnUnlock = false;
diff --git a/Engine/Source/Runtime/AIModule/Private/EnvironmentQuery/Contexts/EnvQueryContext_BlueprintBase.cpp b/Engine/Source/Runtime/AIModule/Private/EnvironmentQuery/Contexts/EnvQueryContext_BlueprintBase.cpp
--- a/Engine/Source/Runtime/AIModule/Private/EnvironmentQuery/Contexts/EnvQueryContext_BlueprintBase.cpp
+++ b/Engine/Source/Runtime/AIModule/Private/EnvironmentQuery/Contexts/EnvQueryContext_BlueprintBase.cpp
@@ -15,8 +15,8 @@ namespace
 	}
 }
 
-UEnvQueryContext_BlueprintBase::UEnvQueryContext_BlueprintBase(const class FPostConstructInitializeProperties& PCIP) 
-	: Super(PCIP), CallMode(UEnvQueryContext_BlueprintBase::InvalidCallMode)
+UEnvQueryContext_BlueprintBase::UEnvQueryContext_BlueprintBase(const FObjectInitializer& ObjectInitializer) 
+	: Super(ObjectInitializer), CallMode(UEnvQueryContext_BlueprintBase::InvalidCallMode)
 {
 	UClass* StopAtClass = UEnvQueryContext_BlueprintBase::StaticClass();	
 	bool bImplementsProvideSingleActor = DoesImplementBPFunction(GET_FUNCTION_NAME_CHECKED(UEnvQueryContext_BlueprintBase, ProvideSingleActor), this, StopAtClass);
diff --git a/Engine/Source/Runtime/AIModule/Private/EnvironmentQuery/Contexts/EnvQueryContext_Item.cpp b/Engine/Source/Runtime/AIModule/Private/EnvironmentQuery/Contexts/EnvQueryContext_Item.cpp
--- a/Engine/Source/Runtime/AIModule/Private/EnvironmentQuery/Contexts/EnvQueryContext_Item.cpp
+++ b/Engine/Source/Runtime/AIModule/Private/EnvironmentQuery/Contexts/EnvQueryContext_Item.cpp
@@ -3,7 +3,7 @@
 #include "AIModulePrivate.h"
 #include "EnvironmentQuery/Contexts/EnvQueryContext_Item.h"
 
-UEnvQueryContext_Item::UEnvQueryContext_Item(const class FPostConstructInitializeProperties& PCIP) : Super(PCIP)
+UEnvQueryContext_Item::UEnvQueryContext_Item(const FObjectInitializer& ObjectInitializer) : Super(ObjectInitializer)
 {
 	// special case, don't set delegate
 }
diff --git a/Engine/Source/Runtime/AIModule/Private/EnvironmentQuery/Contexts/EnvQueryContext_Querier.cpp b/Engine/Source/Runtime/AIModule/Private/EnvironmentQuery/Contexts/EnvQueryContext_Querier.cpp
--- a/Engine/Source/Runtime/AIModule/Private/EnvironmentQuery/Contexts/EnvQueryContext_Querier.cpp
+++ b/Engine/Source/Runtime/AIModule/Private/EnvironmentQuery/Contexts/EnvQueryContext_Querier.cpp
@@ -4,7 +4,7 @@
 #include "EnvironmentQuery/Contexts/EnvQueryContext_Querier.h"
 #include "EnvironmentQuery/Items/EnvQueryItemType_Actor.h"
 
-UEnvQueryContext_Querier::UEnvQueryContext_Querier(const class FPostConstructInitializeProperties& PCIP) : Super(PCIP)
+UEnvQueryContext_Querier::UEnvQueryContext_Querier(const FObjectInitializer& ObjectInitializer) : Super(ObjectInitializer)
 {
 }
 
diff --git a/Engine/Source/Runtime/AIModule/Private/EnvironmentQuery/EQSRenderingComponent.cpp b/Engine/Source/Runtime/AIModule/Private/EnvironmentQuery/EQSRenderingComponent.cpp
--- a/Engine/Source/Runtime/AIModule/Private/EnvironmentQuery/EQSRenderingComponent.cpp
+++ b/Engine/Source/Runtime/AIModule/Private/EnvironmentQuery/EQSRenderingComponent.cpp
@@ -277,8 +277,8 @@ FPrimitiveViewRelevance FEQSSceneProxy::GetViewRelevance(const FSceneView* View)
 //----------------------------------------------------------------------//
 // UEQSRenderingComponent
 //----------------------------------------------------------------------//
-UEQSRenderingComponent::UEQSRenderingComponent(const class FPostConstructInitializeProperties& PCIP) 
-	: Super(PCIP)
+UEQSRenderingComponent::UEQSRenderingComponent(const FObjectInitializer& ObjectInitializer) 
+	: Super(ObjectInitializer)
 	, DrawFlagName("GameplayDebug")
 	, bDrawOnlyWhenSelected(true)
 {
@@ -328,7 +328,7 @@ void UEQSRenderingComponent::DestroyRenderState_Concurrent()
 //----------------------------------------------------------------------//
 // UEQSQueryResultSourceInterface
 //----------------------------------------------------------------------//
-UEQSQueryResultSourceInterface::UEQSQueryResultSourceInterface(const class FPostConstructInitializeProperties& PCIP)
-	: Super(PCIP)
+UEQSQueryResultSourceInterface::UEQSQueryResultSourceInterface(const FObjectInitializer& ObjectInitializer)
+	: Super(ObjectInitializer)
 {
 }
diff --git a/Engine/Source/Runtime/AIModule/Private/EnvironmentQuery/EQSTestingPawn.cpp b/Engine/Source/Runtime/AIModule/Private/EnvironmentQuery/EQSTestingPawn.cpp
--- a/Engine/Source/Runtime/AIModule/Private/EnvironmentQuery/EQSTestingPawn.cpp
+++ b/Engine/Source/Runtime/AIModule/Private/EnvironmentQuery/EQSTestingPawn.cpp
@@ -24,8 +24,8 @@ extern UNREALED_API class UEditorEngine* GEditor;
 //----------------------------------------------------------------------//
 // AEQSTestingPawn
 //----------------------------------------------------------------------//
-AEQSTestingPawn::AEQSTestingPawn(const class FPostConstructInitializeProperties& PCIP)
-	: Super(PCIP)
+AEQSTestingPawn::AEQSTestingPawn(const FObjectInitializer& ObjectInitializer)
+	: Super(ObjectInitializer)
 	, TimeLimitPerStep(-1)
 	, StepToDebugDraw(0)
 	, bDrawLabels(true)
@@ -37,7 +37,7 @@ AEQSTestingPawn::AEQSTestingPawn(const class FPostConstructInitializeProperties&
 	CapsuleComponent->SetCollisionProfileName(CollisionProfileName);
 
 #if WITH_EDITORONLY_DATA
-	EdRenderComp = PCIP.CreateEditorOnlyDefaultSubobject<UEQSRenderingComponent>(this, TEXT("EQSRender"));
+	EdRenderComp = ObjectInitializer.CreateEditorOnlyDefaultSubobject<UEQSRenderingComponent>(this, TEXT("EQSRender"));
 	if (HasAnyFlags(RF_ClassDefaultObject) == false)
 	{
 		UArrowComponent* ArrowComponent = FindComponentByClass<UArrowComponent>();
@@ -47,7 +47,7 @@ AEQSTestingPawn::AEQSTestingPawn(const class FPostConstructInitializeProperties&
 			ArrowComponent->bIsScreenSizeScaled = true;
 		}
 
-		TSubobjectPtr<UBillboardComponent> SpriteComponent = PCIP.CreateEditorOnlyDefaultSubobject<UBillboardComponent>(this, TEXT("Sprite"));
+		TSubobjectPtr<UBillboardComponent> SpriteComponent = ObjectInitializer.CreateEditorOnlyDefaultSubobject<UBillboardComponent>(this, TEXT("Sprite"));
 		if (!IsRunningCommandlet() && (SpriteComponent != nullptr))
 		{
 			struct FConstructorStatics
diff --git a/Engine/Source/Runtime/AIModule/Private/EnvironmentQuery/EnvQuery.cpp b/Engine/Source/Runtime/AIModule/Private/EnvironmentQuery/EnvQuery.cpp
--- a/Engine/Source/Runtime/AIModule/Private/EnvironmentQuery/EnvQuery.cpp
+++ b/Engine/Source/Runtime/AIModule/Private/EnvironmentQuery/EnvQuery.cpp
@@ -6,7 +6,7 @@
 #include "EnvironmentQuery/EnvQueryOption.h"
 #include "EnvironmentQuery/EnvQueryTest.h"
 
-UEnvQuery::UEnvQuery(const class FPostConstructInitializeProperties& PCIP) : Super(PCIP)
+UEnvQuery::UEnvQuery(const FObjectInitializer& ObjectInitializer) : Super(ObjectInitializer)
 {
 }
 
diff --git a/Engine/Source/Runtime/AIModule/Private/EnvironmentQuery/EnvQueryContext.cpp b/Engine/Source/Runtime/AIModule/Private/EnvironmentQuery/EnvQueryContext.cpp
--- a/Engine/Source/Runtime/AIModule/Private/EnvironmentQuery/EnvQueryContext.cpp
+++ b/Engine/Source/Runtime/AIModule/Private/EnvironmentQuery/EnvQueryContext.cpp
@@ -3,7 +3,7 @@
 #include "AIModulePrivate.h"
 #include "EnvironmentQuery/EnvQueryContext.h"
 
-UEnvQueryContext::UEnvQueryContext(const class FPostConstructInitializeProperties& PCIP) : Super(PCIP)
+UEnvQueryContext::UEnvQueryContext(const FObjectInitializer& ObjectInitializer) : Super(ObjectInitializer)
 {
 }
 
diff --git a/Engine/Source/Runtime/AIModule/Private/EnvironmentQuery/EnvQueryGenerator.cpp b/Engine/Source/Runtime/AIModule/Private/EnvironmentQuery/EnvQueryGenerator.cpp
--- a/Engine/Source/Runtime/AIModule/Private/EnvironmentQuery/EnvQueryGenerator.cpp
+++ b/Engine/Source/Runtime/AIModule/Private/EnvironmentQuery/EnvQueryGenerator.cpp
@@ -4,7 +4,7 @@
 #include "EnvironmentQuery/EnvQueryManager.h"
 #include "EnvironmentQuery/EnvQueryGenerator.h"
 
-UEnvQueryGenerator::UEnvQueryGenerator(const class FPostConstructInitializeProperties& PCIP) : Super(PCIP)
+UEnvQueryGenerator::UEnvQueryGenerator(const FObjectInitializer& ObjectInitializer) : Super(ObjectInitializer)
 {
 }
 
diff --git a/Engine/Source/Runtime/AIModule/Private/EnvironmentQuery/EnvQueryManager.cpp b/Engine/Source/Runtime/AIModule/Private/EnvironmentQuery/EnvQueryManager.cpp
--- a/Engine/Source/Runtime/AIModule/Private/EnvironmentQuery/EnvQueryManager.cpp
+++ b/Engine/Source/Runtime/AIModule/Private/EnvironmentQuery/EnvQueryManager.cpp
@@ -78,7 +78,7 @@ int32 FEnvQueryRequest::Execute(EEnvQueryRunMode::Type RunMode, FQueryFinishedSi
 
 TArray<TSubclassOf<UEnvQueryItemType> > UEnvQueryManager::RegisteredItemTypes;
 
-UEnvQueryManager::UEnvQueryManager(const FPostConstructInitializeProperties& PCIP) : Super(PCIP)
+UEnvQueryManager::UEnvQueryManager(const FObjectInitializer& ObjectInitializer) : Super(ObjectInitializer)
 {
 	if (!HasAnyFlags(RF_ClassDefaultObject))
 	{
diff --git a/Engine/Source/Runtime/AIModule/Private/EnvironmentQuery/EnvQueryOption.cpp b/Engine/Source/Runtime/AIModule/Private/EnvironmentQuery/EnvQueryOption.cpp
--- a/Engine/Source/Runtime/AIModule/Private/EnvironmentQuery/EnvQueryOption.cpp
+++ b/Engine/Source/Runtime/AIModule/Private/EnvironmentQuery/EnvQueryOption.cpp
@@ -4,7 +4,7 @@
 #include "EnvironmentQuery/EnvQueryGenerator.h"
 #include "EnvironmentQuery/EnvQueryOption.h"
 
-UEnvQueryOption::UEnvQueryOption(const class FPostConstructInitializeProperties& PCIP) : Super(PCIP)
+UEnvQueryOption::UEnvQueryOption(const FObjectInitializer& ObjectInitializer) : Super(ObjectInitializer)
 {
 }
 
diff --git a/Engine/Source/Runtime/AIModule/Private/EnvironmentQuery/EnvQueryTest.cpp b/Engine/Source/Runtime/AIModule/Private/EnvironmentQuery/EnvQueryTest.cpp
--- a/Engine/Source/Runtime/AIModule/Private/EnvironmentQuery/EnvQueryTest.cpp
+++ b/Engine/Source/Runtime/AIModule/Private/EnvironmentQuery/EnvQueryTest.cpp
@@ -11,7 +11,7 @@
 
 #define LOCTEXT_NAMESPACE "EnvQueryGenerator"
 
-UEnvQueryTest::UEnvQueryTest(const FPostConstructInitializeProperties& PCIP) : Super(PCIP)
+UEnvQueryTest::UEnvQueryTest(const FObjectInitializer& ObjectInitializer) : Super(ObjectInitializer)
 {
 	TestPurpose = EEnvTestPurpose::FilterAndScore;
 	FilterType = EEnvTestFilterType::Range;
diff --git a/Engine/Source/Runtime/AIModule/Private/EnvironmentQuery/Generators/EnvQueryGenerator_ActorsOfClass.cpp b/Engine/Source/Runtime/AIModule/Private/EnvironmentQuery/Generators/EnvQueryGenerator_ActorsOfClass.cpp
--- a/Engine/Source/Runtime/AIModule/Private/EnvironmentQuery/Generators/EnvQueryGenerator_ActorsOfClass.cpp
+++ b/Engine/Source/Runtime/AIModule/Private/EnvironmentQuery/Generators/EnvQueryGenerator_ActorsOfClass.cpp
@@ -9,7 +9,7 @@
 
 #define LOCTEXT_NAMESPACE "EnvQueryGenerator"
 
-UEnvQueryGenerator_ActorsOfClass::UEnvQueryGenerator_ActorsOfClass(const class FPostConstructInitializeProperties& PCIP) : Super(PCIP)
+UEnvQueryGenerator_ActorsOfClass::UEnvQueryGenerator_ActorsOfClass(const FObjectInitializer& ObjectInitializer) : Super(ObjectInitializer)
 {
 	SearchCenter = UEnvQueryContext_Querier::StaticClass();
 	ItemType = UEnvQueryItemType_Actor::StaticClass();
diff --git a/Engine/Source/Runtime/AIModule/Private/EnvironmentQuery/Generators/EnvQueryGenerator_BlueprintBase.cpp b/Engine/Source/Runtime/AIModule/Private/EnvironmentQuery/Generators/EnvQueryGenerator_BlueprintBase.cpp
--- a/Engine/Source/Runtime/AIModule/Private/EnvironmentQuery/Generators/EnvQueryGenerator_BlueprintBase.cpp
+++ b/Engine/Source/Runtime/AIModule/Private/EnvironmentQuery/Generators/EnvQueryGenerator_BlueprintBase.cpp
@@ -11,8 +11,8 @@
 
 #define LOCTEXT_NAMESPACE "EnvQueryGenerator"
 
-UEnvQueryGenerator_BlueprintBase::UEnvQueryGenerator_BlueprintBase(const class FPostConstructInitializeProperties& PCIP)
-	: Super(PCIP)
+UEnvQueryGenerator_BlueprintBase::UEnvQueryGenerator_BlueprintBase(const FObjectInitializer& ObjectInitializer)
+	: Super(ObjectInitializer)
 {
 	Context = UEnvQueryContext_Querier::StaticClass();
 	ItemType = UEnvQueryItemType_Actor::StaticClass();
diff --git a/Engine/Source/Runtime/AIModule/Private/EnvironmentQuery/Generators/EnvQueryGenerator_Composite.cpp b/Engine/Source/Runtime/AIModule/Private/EnvironmentQuery/Generators/EnvQueryGenerator_Composite.cpp
--- a/Engine/Source/Runtime/AIModule/Private/EnvironmentQuery/Generators/EnvQueryGenerator_Composite.cpp
+++ b/Engine/Source/Runtime/AIModule/Private/EnvironmentQuery/Generators/EnvQueryGenerator_Composite.cpp
@@ -5,7 +5,7 @@
 
 #define LOCTEXT_NAMESPACE "EnvQueryGenerator"
 
-UEnvQueryGenerator_Composite::UEnvQueryGenerator_Composite(const FPostConstructInitializeProperties& PCIP) : Super(PCIP)
+UEnvQueryGenerator_Composite::UEnvQueryGenerator_Composite(const FObjectInitializer& ObjectInitializer) : Super(ObjectInitializer)
 {
 
 }
diff --git a/Engine/Source/Runtime/AIModule/Private/EnvironmentQuery/Generators/EnvQueryGenerator_OnCircle.cpp b/Engine/Source/Runtime/AIModule/Private/EnvironmentQuery/Generators/EnvQueryGenerator_OnCircle.cpp
--- a/Engine/Source/Runtime/AIModule/Private/EnvironmentQuery/Generators/EnvQueryGenerator_OnCircle.cpp
+++ b/Engine/Source/Runtime/AIModule/Private/EnvironmentQuery/Generators/EnvQueryGenerator_OnCircle.cpp
@@ -95,8 +95,8 @@ void FBatchTracingHelper::DoSingleSourceMultiDestinations<EEnvTraceShape::Capsul
 //----------------------------------------------------------------------//
 // UEnvQueryGenerator_OnCircle
 //----------------------------------------------------------------------//
-UEnvQueryGenerator_OnCircle::UEnvQueryGenerator_OnCircle(const class FPostConstructInitializeProperties& PCIP) 
-	: Super(PCIP)
+UEnvQueryGenerator_OnCircle::UEnvQueryGenerator_OnCircle(const FObjectInitializer& ObjectInitializer) 
+	: Super(ObjectInitializer)
 {
 	CircleCenter = UEnvQueryContext_Querier::StaticClass();
 	ItemType = UEnvQueryItemType_Point::StaticClass();
diff --git a/Engine/Source/Runtime/AIModule/Private/EnvironmentQuery/Generators/EnvQueryGenerator_PathingGrid.cpp b/Engine/Source/Runtime/AIModule/Private/EnvironmentQuery/Generators/EnvQueryGenerator_PathingGrid.cpp
--- a/Engine/Source/Runtime/AIModule/Private/EnvironmentQuery/Generators/EnvQueryGenerator_PathingGrid.cpp
+++ b/Engine/Source/Runtime/AIModule/Private/EnvironmentQuery/Generators/EnvQueryGenerator_PathingGrid.cpp
@@ -9,7 +9,7 @@
 
 #define LOCTEXT_NAMESPACE "EnvQueryGenerator"
 
-UEnvQueryGenerator_PathingGrid::UEnvQueryGenerator_PathingGrid(const class FPostConstructInitializeProperties& PCIP) : Super(PCIP)
+UEnvQueryGenerator_PathingGrid::UEnvQueryGenerator_PathingGrid(const FObjectInitializer& ObjectInitializer) : Super(ObjectInitializer)
 {
 	GenerateAround = UEnvQueryContext_Querier::StaticClass();
 	ItemType = UEnvQueryItemType_Point::StaticClass();
diff --git a/Engine/Source/Runtime/AIModule/Private/EnvironmentQuery/Generators/EnvQueryGenerator_ProjectedPoints.cpp b/Engine/Source/Runtime/AIModule/Private/EnvironmentQuery/Generators/EnvQueryGenerator_ProjectedPoints.cpp
--- a/Engine/Source/Runtime/AIModule/Private/EnvironmentQuery/Generators/EnvQueryGenerator_ProjectedPoints.cpp
+++ b/Engine/Source/Runtime/AIModule/Private/EnvironmentQuery/Generators/EnvQueryGenerator_ProjectedPoints.cpp
@@ -3,7 +3,7 @@
 #include "AIModulePrivate.h"
 #include "EnvironmentQuery/Generators/EnvQueryGenerator_ProjectedPoints.h"
 
-UEnvQueryGenerator_ProjectedPoints::UEnvQueryGenerator_ProjectedPoints(const class FPostConstructInitializeProperties& PCIP) : Super(PCIP)
+UEnvQueryGenerator_ProjectedPoints::UEnvQueryGenerator_ProjectedPoints(const FObjectInitializer& ObjectInitializer) : Super(ObjectInitializer)
 {
 	ProjectionData.SetNavmeshOnly();
 	ProjectionData.bCanProjectDown = true;
diff --git a/Engine/Source/Runtime/AIModule/Private/EnvironmentQuery/Generators/EnvQueryGenerator_SimpleGrid.cpp b/Engine/Source/Runtime/AIModule/Private/EnvironmentQuery/Generators/EnvQueryGenerator_SimpleGrid.cpp
--- a/Engine/Source/Runtime/AIModule/Private/EnvironmentQuery/Generators/EnvQueryGenerator_SimpleGrid.cpp
+++ b/Engine/Source/Runtime/AIModule/Private/EnvironmentQuery/Generators/EnvQueryGenerator_SimpleGrid.cpp
@@ -8,7 +8,7 @@
 
 #define LOCTEXT_NAMESPACE "EnvQueryGenerator"
 
-UEnvQueryGenerator_SimpleGrid::UEnvQueryGenerator_SimpleGrid(const class FPostConstructInitializeProperties& PCIP) : Super(PCIP)
+UEnvQueryGenerator_SimpleGrid::UEnvQueryGenerator_SimpleGrid(const FObjectInitializer& ObjectInitializer) : Super(ObjectInitializer)
 {
 	GenerateAround = UEnvQueryContext_Querier::StaticClass();
 	ItemType = UEnvQueryItemType_Point::StaticClass();
diff --git a/Engine/Source/Runtime/AIModule/Private/EnvironmentQuery/Items/EnvQueryItemType.cpp b/Engine/Source/Runtime/AIModule/Private/EnvironmentQuery/Items/EnvQueryItemType.cpp
--- a/Engine/Source/Runtime/AIModule/Private/EnvironmentQuery/Items/EnvQueryItemType.cpp
+++ b/Engine/Source/Runtime/AIModule/Private/EnvironmentQuery/Items/EnvQueryItemType.cpp
@@ -4,7 +4,7 @@
 #include "EnvironmentQuery/EnvQueryManager.h"
 #include "EnvironmentQuery/Items/EnvQueryItemType.h"
 
-UEnvQueryItemType::UEnvQueryItemType(const class FPostConstructInitializeProperties& PCIP) : Super(PCIP)
+UEnvQueryItemType::UEnvQueryItemType(const FObjectInitializer& ObjectInitializer) : Super(ObjectInitializer)
 {
 	// register in known types 
 	if (HasAnyFlags(RF_ClassDefaultObject) && !GetClass()->HasAnyClassFlags(CLASS_Abstract))
diff --git a/Engine/Source/Runtime/AIModule/Private/EnvironmentQuery/Items/EnvQueryItemType_Actor.cpp b/Engine/Source/Runtime/AIModule/Private/EnvironmentQuery/Items/EnvQueryItemType_Actor.cpp
--- a/Engine/Source/Runtime/AIModule/Private/EnvironmentQuery/Items/EnvQueryItemType_Actor.cpp
+++ b/Engine/Source/Runtime/AIModule/Private/EnvironmentQuery/Items/EnvQueryItemType_Actor.cpp
@@ -3,7 +3,7 @@
 #include "AIModulePrivate.h"
 #include "EnvironmentQuery/Items/EnvQueryItemType_Actor.h"
 
-UEnvQueryItemType_Actor::UEnvQueryItemType_Actor(const class FPostConstructInitializeProperties& PCIP) : Super(PCIP)
+UEnvQueryItemType_Actor::UEnvQueryItemType_Actor(const FObjectInitializer& ObjectInitializer) : Super(ObjectInitializer)
 {
 	ValueSize = sizeof(FWeakObjectPtr);
 }
diff --git a/Engine/Source/Runtime/AIModule/Private/EnvironmentQuery/Items/EnvQueryItemType_ActorBase.cpp b/Engine/Source/Runtime/AIModule/Private/EnvironmentQuery/Items/EnvQueryItemType_ActorBase.cpp
--- a/Engine/Source/Runtime/AIModule/Private/EnvironmentQuery/Items/EnvQueryItemType_ActorBase.cpp
+++ b/Engine/Source/Runtime/AIModule/Private/EnvironmentQuery/Items/EnvQueryItemType_ActorBase.cpp
@@ -5,7 +5,7 @@
 #include "BehaviorTree/Blackboard/BlackboardKeyType_Object.h"
 #include "EnvironmentQuery/Items/EnvQueryItemType_ActorBase.h"
 
-UEnvQueryItemType_ActorBase::UEnvQueryItemType_ActorBase(const class FPostConstructInitializeProperties& PCIP) : Super(PCIP)
+UEnvQueryItemType_ActorBase::UEnvQueryItemType_ActorBase(const FObjectInitializer& ObjectInitializer) : Super(ObjectInitializer)
 {
 }
 
diff --git a/Engine/Source/Runtime/AIModule/Private/EnvironmentQuery/Items/EnvQueryItemType_Direction.cpp b/Engine/Source/Runtime/AIModule/Private/EnvironmentQuery/Items/EnvQueryItemType_Direction.cpp
--- a/Engine/Source/Runtime/AIModule/Private/EnvironmentQuery/Items/EnvQueryItemType_Direction.cpp
+++ b/Engine/Source/Runtime/AIModule/Private/EnvironmentQuery/Items/EnvQueryItemType_Direction.cpp
@@ -3,7 +3,7 @@
 #include "AIModulePrivate.h"
 #include "EnvironmentQuery/Items/EnvQueryItemType_Direction.h"
 
-UEnvQueryItemType_Direction::UEnvQueryItemType_Direction(const class FPostConstructInitializeProperties& PCIP) : Super(PCIP)
+UEnvQueryItemType_Direction::UEnvQueryItemType_Direction(const FObjectInitializer& ObjectInitializer) : Super(ObjectInitializer)
 {
 	ValueSize = sizeof(FVector);
 }
diff --git a/Engine/Source/Runtime/AIModule/Private/EnvironmentQuery/Items/EnvQueryItemType_Point.cpp b/Engine/Source/Runtime/AIModule/Private/EnvironmentQuery/Items/EnvQueryItemType_Point.cpp
--- a/Engine/Source/Runtime/AIModule/Private/EnvironmentQuery/Items/EnvQueryItemType_Point.cpp
+++ b/Engine/Source/Runtime/AIModule/Private/EnvironmentQuery/Items/EnvQueryItemType_Point.cpp
@@ -3,7 +3,7 @@
 #include "AIModulePrivate.h"
 #include "EnvironmentQuery/Items/EnvQueryItemType_Point.h"
 
-UEnvQueryItemType_Point::UEnvQueryItemType_Point(const class FPostConstructInitializeProperties& PCIP) : Super(PCIP)
+UEnvQueryItemType_Point::UEnvQueryItemType_Point(const FObjectInitializer& ObjectInitializer) : Super(ObjectInitializer)
 {
 	ValueSize = sizeof(FVector);
 }
diff --git a/Engine/Source/Runtime/AIModule/Private/EnvironmentQuery/Items/EnvQueryItemType_VectorBase.cpp b/Engine/Source/Runtime/AIModule/Private/EnvironmentQuery/Items/EnvQueryItemType_VectorBase.cpp
--- a/Engine/Source/Runtime/AIModule/Private/EnvironmentQuery/Items/EnvQueryItemType_VectorBase.cpp
+++ b/Engine/Source/Runtime/AIModule/Private/EnvironmentQuery/Items/EnvQueryItemType_VectorBase.cpp
@@ -5,7 +5,7 @@
 #include "BehaviorTree/Blackboard/BlackboardKeyType_Vector.h"
 #include "EnvironmentQuery/Items/EnvQueryItemType_VectorBase.h"
 
-UEnvQueryItemType_VectorBase::UEnvQueryItemType_VectorBase(const class FPostConstructInitializeProperties& PCIP) : Super(PCIP)
+UEnvQueryItemType_VectorBase::UEnvQueryItemType_VectorBase(const FObjectInitializer& ObjectInitializer) : Super(ObjectInitializer)
 {
 }
 
diff --git a/Engine/Source/Runtime/AIModule/Private/EnvironmentQuery/Tests/EnvQueryTest_Distance.cpp b/Engine/Source/Runtime/AIModule/Private/EnvironmentQuery/Tests/EnvQueryTest_Distance.cpp
--- a/Engine/Source/Runtime/AIModule/Private/EnvironmentQuery/Tests/EnvQueryTest_Distance.cpp
+++ b/Engine/Source/Runtime/AIModule/Private/EnvironmentQuery/Tests/EnvQueryTest_Distance.cpp
@@ -23,7 +23,7 @@ namespace
 	}
 }
 
-UEnvQueryTest_Distance::UEnvQueryTest_Distance(const class FPostConstructInitializeProperties& PCIP) : Super(PCIP)
+UEnvQueryTest_Distance::UEnvQueryTest_Distance(const FObjectInitializer& ObjectInitializer) : Super(ObjectInitializer)
 {
 	DistanceTo = UEnvQueryContext_Querier::StaticClass();
 	Cost = EEnvTestCost::Low;
diff --git a/Engine/Source/Runtime/AIModule/Private/EnvironmentQuery/Tests/EnvQueryTest_Dot.cpp b/Engine/Source/Runtime/AIModule/Private/EnvironmentQuery/Tests/EnvQueryTest_Dot.cpp
--- a/Engine/Source/Runtime/AIModule/Private/EnvironmentQuery/Tests/EnvQueryTest_Dot.cpp
+++ b/Engine/Source/Runtime/AIModule/Private/EnvironmentQuery/Tests/EnvQueryTest_Dot.cpp
@@ -6,7 +6,7 @@
 #include "EnvironmentQuery/Contexts/EnvQueryContext_Item.h"
 #include "EnvironmentQuery/Tests/EnvQueryTest_Dot.h"
 
-UEnvQueryTest_Dot::UEnvQueryTest_Dot(const class FPostConstructInitializeProperties& PCIP) : Super(PCIP)
+UEnvQueryTest_Dot::UEnvQueryTest_Dot(const FObjectInitializer& ObjectInitializer) : Super(ObjectInitializer)
 {
 	Cost = EEnvTestCost::Low;
 	ValidItemType = UEnvQueryItemType_VectorBase::StaticClass();
