Add parsing DEPRECATED macro in struct and class declaration in UHT.

diff --git a/Engine/Source/Programs/UnrealHeaderTool/Private/HeaderParser.cpp b/Engine/Source/Programs/UnrealHeaderTool/Private/HeaderParser.cpp
--- a/Engine/Source/Programs/UnrealHeaderTool/Private/HeaderParser.cpp
+++ b/Engine/Source/Programs/UnrealHeaderTool/Private/HeaderParser.cpp
@@ -1139,6 +1139,8 @@ UScriptStruct* FHeaderParser::CompileStructDeclaration(FClasses& AllClasses, FCl
 	// The required API module for this struct, if any
 	FString RequiredAPIMacroIfPresent;
 
+	SkipDeprecatedMacroIfNecessary();
+
 	// Read the struct name
 	ParseNameWithPotentialAPIMacroPrefix(/*out*/ StructNameInScript, /*out*/ RequiredAPIMacroIfPresent, TEXT("struct"));
 
@@ -4270,6 +4272,8 @@ void FHeaderParser::CompileClassDeclaration(FClasses& AllClasses)
 	// New style files have the class name / extends afterwards
 	RequireIdentifier(TEXT("class"), TEXT("Class declaration"));
 
+	SkipDeprecatedMacroIfNecessary();
+
 	FString DeclaredClassName;
 	FString RequiredAPIMacroIfPresent;
 	ParseClassNameDeclaration(AllClasses, /*out*/ DeclaredClassName, /*out*/ RequiredAPIMacroIfPresent);
@@ -7613,3 +7617,27 @@ bool FHeaderParser::TryToMatchConstructorParameterList(FToken Token)
 
 	return true;
 }
+
+void FHeaderParser::SkipDeprecatedMacroIfNecessary()
+{
+	if (!MatchIdentifier(TEXT("DEPRECATED")))
+	{
+		return;
+	}
+
+	FToken Token;
+	// DEPRECATED(Version, "Message")
+	RequireSymbol(TEXT("("), TEXT("DEPRECATED macro"));
+	if (GetToken(Token) && (Token.Type != CPT_Float || Token.TokenType != TOKEN_Const))
+	{
+		FError::Throwf(TEXT("Expected engine version in DEPRECATED macro"));
+	}
+
+	RequireSymbol(TEXT(","), TEXT("DEPRECATED macro"));
+	if (GetToken(Token) && (Token.Type != CPT_String || Token.TokenType != TOKEN_Const))
+	{
+		FError::Throwf(TEXT("Expected deprecation message in DEPRECATED macro"));
+	}
+
+	RequireSymbol(TEXT(")"), TEXT("DEPRECATED macro"));
+}
diff --git a/Engine/Source/Programs/UnrealHeaderTool/Private/HeaderParser.h b/Engine/Source/Programs/UnrealHeaderTool/Private/HeaderParser.h
--- a/Engine/Source/Programs/UnrealHeaderTool/Private/HeaderParser.h
+++ b/Engine/Source/Programs/UnrealHeaderTool/Private/HeaderParser.h
@@ -547,6 +547,7 @@ private:
 	 * @returns True if matched. False otherwise.
 	 */
 	bool TryToMatchConstructorParameterList(FToken Token);
+	void SkipDeprecatedMacroIfNecessary();
 };
 
 /////////////////////////////////////////////////////
diff --git a/Engine/Source/Runtime/Core/Public/GenericPlatform/GenericPlatformCompilerPreSetup.h b/Engine/Source/Runtime/Core/Public/GenericPlatform/GenericPlatformCompilerPreSetup.h
--- a/Engine/Source/Runtime/Core/Public/GenericPlatform/GenericPlatformCompilerPreSetup.h
+++ b/Engine/Source/Runtime/Core/Public/GenericPlatform/GenericPlatformCompilerPreSetup.h
@@ -7,6 +7,13 @@
 #pragma once
 
 #ifndef DEPRECATED
+	// Sample usage (note slightly unintuitive syntax for classes and structs) :
+	// 
+	// DEPRECATED(4.6, "Message")
+	// void Function();
+	// 
+	// struct DEPRECATED(4.6, "Message") MyStruct;
+	// class DEPRECATED(4.6, "Message") MyClass;
 	#define DEPRECATED(VERSION, MESSAGE)
 #endif // DEPRECATED
 
diff --git a/Engine/Source/Runtime/Core/Public/Mac/MacPlatformCompilerPreSetup.h b/Engine/Source/Runtime/Core/Public/Mac/MacPlatformCompilerPreSetup.h
--- a/Engine/Source/Runtime/Core/Public/Mac/MacPlatformCompilerPreSetup.h
+++ b/Engine/Source/Runtime/Core/Public/Mac/MacPlatformCompilerPreSetup.h
@@ -8,5 +8,13 @@
 
 #ifndef DISABLE_DEPRECATION
 	#pragma clang diagnostic warning "-Wdeprecated-declarations"
+
+	// Sample usage (note slightly unintuitive syntax for classes and structs) :
+	// 
+	// DEPRECATED(4.6, "Message")
+	// void Function();
+	// 
+	// struct DEPRECATED(4.6, "Message") MyStruct;
+	// class DEPRECATED(4.6, "Message") MyClass;
 	#define DEPRECATED(VERSION, MESSAGE) __attribute__((deprecated(MESSAGE " Please update your code to the new API before upgrading to the next release, otherwise your project will no longer compile.")))
 #endif // DISABLE_DEPRECATION
diff --git a/Engine/Source/Runtime/Core/Public/Windows/WindowsPlatformCompilerPreSetup.h b/Engine/Source/Runtime/Core/Public/Windows/WindowsPlatformCompilerPreSetup.h
--- a/Engine/Source/Runtime/Core/Public/Windows/WindowsPlatformCompilerPreSetup.h
+++ b/Engine/Source/Runtime/Core/Public/Windows/WindowsPlatformCompilerPreSetup.h
@@ -7,6 +7,13 @@
 #pragma once
 
 #ifndef DISABLE_DEPRECATION
+	// Sample usage (note slightly unintuitive syntax for classes and structs) :
+	// 
+	// DEPRECATED(4.6, "Message")
+	// void Function();
+	// 
+	// struct DEPRECATED(4.6, "Message") MyStruct;
+	// class DEPRECATED(4.6, "Message") MyClass;
 	#define DEPRECATED(VERSION, MESSAGE) __declspec(deprecated(MESSAGE " Please update your code to the new API before upgrading to the next release, otherwise your project will no longer compile."))
 
 	#ifndef EMIT_DEPRECATED_WARNING_MESSAGE_STR
