Hot-Reload: fixed hot-reload not detecting UFUNCTION changes in source code

diff --git a/Engine/Source/Programs/UnrealHeaderTool/Private/CodeGenerator.cpp b/Engine/Source/Programs/UnrealHeaderTool/Private/CodeGenerator.cpp
--- a/Engine/Source/Programs/UnrealHeaderTool/Private/CodeGenerator.cpp
+++ b/Engine/Source/Programs/UnrealHeaderTool/Private/CodeGenerator.cpp
@@ -223,6 +223,18 @@ FString Macroize(const TCHAR *MacroName, const TCHAR *StringToMacroize)
 	return FString::Printf(TEXT("#define %s%s\r\n"), MacroName, Result.Len() ? TEXT(" \\") : TEXT("")) + Result;
 }
 
+/** Generates a CRC tag string for the specified field */
+static FString GetGeneratedCodeCRCTag(UField* Field)
+{
+	FString Tag;
+	auto FieldCrc = GGeneratedCodeCRCs.Find(Field);	
+	if (FieldCrc)
+	{
+		Tag = FString::Printf(TEXT(" // %u"), *FieldCrc);
+	}
+	return Tag;
+}
+
 struct FParmsAndReturnProperties
 {
 	FParmsAndReturnProperties()
@@ -697,22 +709,12 @@ FString FNativeClassHeaderGenerator::PropertyNew(FString& Meta, UProperty* Prop,
 		UScriptStruct* Struct = StructProperty->Struct;
 		check(Struct);
 		ExtraArgs = FString::Printf(TEXT(", %s"), *GetSingletonName(Struct));
-		auto StructCrc = GGeneratedCodeCRCs.Find(Struct);
-		if (StructCrc)
-		{
-			GeneratedCrc = FString::Printf(TEXT(" // %u"), *StructCrc);
-		}
 	}
 	else if (UByteProperty* ByteProperty = Cast<UByteProperty>(Prop))
 	{
 		if (ByteProperty->Enum)
 		{
 			ExtraArgs = FString::Printf(TEXT(", %s"), *GetSingletonName(ByteProperty->Enum));
-			auto EnumCrc = GGeneratedCodeCRCs.Find(ByteProperty->Enum);
-			if (EnumCrc)
-			{
-				GeneratedCrc = FString::Printf(TEXT(" // %u"), *EnumCrc);
-			}
 		}
 	}
 	else if (UBoolProperty* BoolProperty = Cast<UBoolProperty>(Prop))
@@ -753,7 +755,7 @@ FString FNativeClassHeaderGenerator::PropertyNew(FString& Meta, UProperty* Prop,
 		Spaces, 
 		*Symbol, 
 		*Constructor,
-		*GeneratedCrc);
+		*GetGeneratedCodeCRCTag(Prop));
 
 	if (Prop->ArrayDim != 1)
 	{
@@ -1180,8 +1182,10 @@ void FNativeClassHeaderGenerator::ExportNativeGeneratedInitCode(FClass* Class)
 		SingletonNameToExternDecl.Add(SingletonName, Extern);
 		GeneratedFunctionDeclarations.Log(*Extern);
 
-		GeneratedFunctionText.Logf(TEXT("    UFunction* %s\r\n"), *SingletonName);
-		GeneratedFunctionText.Logf(TEXT("    {\r\n"));
+		FStringOutputDevice CurrentFunctionText;
+
+		CurrentFunctionText.Logf(TEXT("    UFunction* %s\r\n"), *SingletonName);
+		CurrentFunctionText.Logf(TEXT("    {\r\n"));
 
 		if (bIsNoExport || !(Function->FunctionFlags&FUNC_Event))  // non-events do not export a params struct, so lets do that locally for offset determination
 		{
@@ -1190,17 +1194,17 @@ void FNativeClassHeaderGenerator::ExportNativeGeneratedInitCode(FClass* Class)
 			FindNoExportStructs(Structs, Function);
 			if (Structs.Num())
 			{
-				ExportMirrorsForNoexportStructs(Structs, 8, GeneratedFunctionText);
+				ExportMirrorsForNoexportStructs(Structs, 8, CurrentFunctionText);
 			}
 			TArray<UFunction*> CallbackFunctions;
 			CallbackFunctions.Add(Function);
-			ExportEventParms(CallbackFunctions, NULL, 8, false, &GeneratedFunctionText);
+			ExportEventParms(CallbackFunctions, NULL, 8, false, &CurrentFunctionText);
 		}
 
-		GeneratedFunctionText.Logf(TEXT("        UClass* OuterClass=%s;\r\n"), *GetSingletonName(Class));
-		GeneratedFunctionText.Logf(TEXT("        static UFunction* ReturnFunction = NULL;\r\n"));
-		GeneratedFunctionText.Logf(TEXT("        if (!ReturnFunction)\r\n"));
-		GeneratedFunctionText.Logf(TEXT("        {\r\n"));
+		CurrentFunctionText.Logf(TEXT("        UClass* OuterClass=%s;\r\n"), *GetSingletonName(Class));
+		CurrentFunctionText.Logf(TEXT("        static UFunction* ReturnFunction = NULL;\r\n"));
+		CurrentFunctionText.Logf(TEXT("        if (!ReturnFunction)\r\n"));
+		CurrentFunctionText.Logf(TEXT("        {\r\n"));
 		FString SuperFunctionString(TEXT("NULL"));
 		if (SuperFunction)
 		{
@@ -1228,7 +1232,7 @@ void FNativeClassHeaderGenerator::ExportNativeGeneratedInitCode(FClass* Class)
 			StructureSize = FString::Printf(TEXT(", sizeof(%s)"), *GetEventStructParamsName(*CastChecked<UClass>(TempFunction->GetOuter())->GetName(), *FunctionName));
 		}
 
-		GeneratedFunctionText.Logf(TEXT("            ReturnFunction = new(OuterClass, TEXT(\"%s\"), RF_Public|RF_Transient|RF_Native) UFunction(FObjectInitializer(), %s, 0x%08X, %d%s);\r\n"), 
+		CurrentFunctionText.Logf(TEXT("            ReturnFunction = new(OuterClass, TEXT(\"%s\"), RF_Public|RF_Transient|RF_Native) UFunction(FObjectInitializer(), %s, 0x%08X, %d%s);\r\n"), 
 			*Function->GetName(),
 			*SuperFunctionString,
 			Function->FunctionFlags,
@@ -1243,29 +1247,34 @@ void FNativeClassHeaderGenerator::ExportNativeGeneratedInitCode(FClass* Class)
 
 		for (int32 Index = Props.Num() - 1; Index >= 0; Index--)
 		{
-			OutputProperty(Meta, GeneratedFunctionText, OuterString, Props[Index], TEXT("            "));
+			OutputProperty(Meta, CurrentFunctionText, OuterString, Props[Index], TEXT("            "));
 		}
 
 		if (FunctionData.FunctionFlags & (FUNC_NetRequest | FUNC_NetResponse))
 		{
-			GeneratedFunctionText.Logf(TEXT("            ReturnFunction->RPCId=%d;\r\n"), FunctionData.RPCId);
-			GeneratedFunctionText.Logf(TEXT("            ReturnFunction->RPCResponseId=%d;\r\n"), FunctionData.RPCResponseId);
+			CurrentFunctionText.Logf(TEXT("            ReturnFunction->RPCId=%d;\r\n"), FunctionData.RPCId);
+			CurrentFunctionText.Logf(TEXT("            ReturnFunction->RPCResponseId=%d;\r\n"), FunctionData.RPCResponseId);
 		}
 
-		GeneratedFunctionText.Logf(TEXT("            ReturnFunction->Bind();\r\n"));
-		GeneratedFunctionText.Logf(TEXT("            ReturnFunction->StaticLink();\r\n"));
+		CurrentFunctionText.Logf(TEXT("            ReturnFunction->Bind();\r\n"));
+		CurrentFunctionText.Logf(TEXT("            ReturnFunction->StaticLink();\r\n"));
 
 		if (Meta.Len())
 		{
-			GeneratedFunctionText.Logf(TEXT("#if WITH_METADATA\r\n"));
-			GeneratedFunctionText.Logf(TEXT("            UMetaData* MetaData = ReturnFunction->GetOutermost()->GetMetaData();\r\n"));
-			GeneratedFunctionText.Log(*Meta);
-			GeneratedFunctionText.Logf(TEXT("#endif\r\n"));
+			CurrentFunctionText.Logf(TEXT("#if WITH_METADATA\r\n"));
+			CurrentFunctionText.Logf(TEXT("            UMetaData* MetaData = ReturnFunction->GetOutermost()->GetMetaData();\r\n"));
+			CurrentFunctionText.Log(*Meta);
+			CurrentFunctionText.Logf(TEXT("#endif\r\n"));
 		}
 
-		GeneratedFunctionText.Logf(TEXT("        }\r\n"));
-		GeneratedFunctionText.Logf(TEXT("        return ReturnFunction;\r\n"));
-		GeneratedFunctionText.Logf(TEXT("    }\r\n"));
+		CurrentFunctionText.Logf(TEXT("        }\r\n"));
+		CurrentFunctionText.Logf(TEXT("        return ReturnFunction;\r\n"));
+		CurrentFunctionText.Logf(TEXT("    }\r\n"));
+
+		GeneratedFunctionText += CurrentFunctionText;
+
+		uint32 FunctionCrc = FCrc::MemCrc32(*CurrentFunctionText, CurrentFunctionText.Len() * sizeof(TCHAR));
+		GGeneratedCodeCRCs.Add(Function, FunctionCrc);
 	}
 
 	FStringOutputDevice GeneratedClassRegisterFunctionText;
@@ -1343,7 +1352,7 @@ void FNativeClassHeaderGenerator::ExportNativeGeneratedInitCode(FClass* Class)
 			// Emit code to construct each UFunction and rebuild the function map at runtime
 			for (UFunction* Function : FunctionsInMap)
 			{
-				GeneratedClassRegisterFunctionText.Logf(TEXT("            OuterClass->AddFunctionToFunctionMap(%s);\r\n"), *GetSingletonName(Function));
+				GeneratedClassRegisterFunctionText.Logf(TEXT("            OuterClass->AddFunctionToFunctionMap(%s);%s\r\n"), *GetSingletonName(Function), *GetGeneratedCodeCRCTag(Function));
 			}
 		}
 
