#UE4 Fix header parsing of component properties, private_subobject was breaking the parser and causing the next property to be ignored. This was causing Character's Mesh property to not exist as an example

diff --git a/Engine/Source/Programs/UnrealHeaderTool/Private/HeaderParser.cpp b/Engine/Source/Programs/UnrealHeaderTool/Private/HeaderParser.cpp
--- a/Engine/Source/Programs/UnrealHeaderTool/Private/HeaderParser.cpp
+++ b/Engine/Source/Programs/UnrealHeaderTool/Private/HeaderParser.cpp
@@ -1087,7 +1087,7 @@ EAccessSpecifier FHeaderParser::ParseAccessProtectionSpecifier(FToken& Token)
 
 	for (EAccessSpecifier Test = EAccessSpecifier(ACCESS_NotAnAccessSpecifier + 1); Test != ACCESS_Num; Test = EAccessSpecifier(Test + 1))
 	{
-		if (Token.Matches(GetAccessSpecifierName(Test)))
+		if (Token.Matches(GetAccessSpecifierName(Test)) || (Token.Matches(TEXT("private_subobject")) && Test == ACCESS_Public))
 		{
 			// Consume the colon after the specifier
 			RequireSymbol(TEXT(":"), *FString::Printf(TEXT("after %s"), Token.Identifier));
