Making sure the compiler does not attempt to generate copy-constructors for UObjects as this produces false warnings when the class contains deprecated members and generally UObjects should not be copy-constructed anyway.

diff --git a/Engine/Source/Programs/UnrealHeaderTool/Private/CodeGenerator.cpp b/Engine/Source/Programs/UnrealHeaderTool/Private/CodeGenerator.cpp
--- a/Engine/Source/Programs/UnrealHeaderTool/Private/CodeGenerator.cpp
+++ b/Engine/Source/Programs/UnrealHeaderTool/Private/CodeGenerator.cpp
@@ -2013,6 +2013,22 @@ void FNativeClassHeaderGenerator::ExportClassHeaderInner(FClass* Class, bool bVa
 	}
 }
 
+/**
+* Generates private copy-constructor declaration.
+*
+* @param Out Output device to generate to.
+* @param Class Class to generate constructor for.
+* @param API API string for this constructor.
+*/
+void ExportCopyConstructorDefinition(FStringOutputDevice& Out, FClass* Class, const FString& API)
+{
+	auto ClassNameCPP = NameLookupCPP.GetNameCPP(Class);
+	Out.Logf(TEXT("private:\r\n"));
+	Out.Logf(TEXT("%s/** Private copy-constructor, should never be used */\r\n"), FCString::Spc(4));
+	Out.Logf(TEXT("%s%s_API %s(const %s& InCopy);\r\n"), FCString::Spc(4), *API, ClassNameCPP, ClassNameCPP);
+	Out.Logf(TEXT("public:\r\n"));
+}
+
 /**
  * Generates standard constructor declaration.
  *
@@ -2027,8 +2043,8 @@ void ExportStandardConstructorsMacro(FStringOutputDevice& Out, FClass* Class, co
 		Out.Logf(TEXT("%s/** Standard constructor, called after all reflected properties have been initialized */\r\n"), FCString::Spc(4));
 		Out.Logf(TEXT("%s%s_API %s(const class FObjectInitializer& ObjectInitializer);\r\n"), FCString::Spc(4), *API, NameLookupCPP.GetNameCPP(Class));
 	}
-
 	Out.Logf(TEXT("%sDEFINE_DEFAULT_OBJECT_INITIALIZER_CONSTRUCTOR_CALL(%s)\r\n"), FCString::Spc(4), NameLookupCPP.GetNameCPP(Class));
+	ExportCopyConstructorDefinition(Out, Class, API);
 }
 
 /**
@@ -2041,15 +2057,15 @@ void ExportStandardConstructorsMacro(FStringOutputDevice& Out, FClass* Class, co
 void ExportConstructorDefinition(FStringOutputDevice& Out, FClass* Class, const FString& API)
 {
 	auto* ClassData = GScriptHelper.FindClassData(Class);
-
 	if (!ClassData->bConstructorDeclared)
 	{
-		Out.Logf(TEXT("%s/** Standard constructor, called after all reflected properties have been initialized */"), FCString::Spc(4));
+		Out.Logf(TEXT("%s/** Standard constructor, called after all reflected properties have been initialized */\r\n"), FCString::Spc(4));
 		Out.Logf(TEXT("%s%s_API %s(const class FObjectInitializer& ObjectInitializer) : Super(ObjectInitializer) { };\r\n"), FCString::Spc(4), *API, NameLookupCPP.GetNameCPP(Class));
 
 		ClassData->bConstructorDeclared = true;
 		ClassData->bObjectInitializerConstructorDeclared = true;
 	}
+	ExportCopyConstructorDefinition(Out, Class, API);
 }
 
 /**
diff --git a/Engine/Source/Runtime/Engine/Classes/Camera/CameraComponent.h b/Engine/Source/Runtime/Engine/Classes/Camera/CameraComponent.h
--- a/Engine/Source/Runtime/Engine/Classes/Camera/CameraComponent.h
+++ b/Engine/Source/Runtime/Engine/Classes/Camera/CameraComponent.h
@@ -110,8 +110,4 @@ public:
 	 */
 	DEPRECATED(4.5, "This variable is deprecated; please see the upgrade notes. It only exists to allow compilation of old projects, and code should stop using it in favor of the new bUsePawnControlRotation.")
 	bool bUseControllerViewRotation;
-
-private:
-	// Workaround for compiler-generated function using deprecated variable. Nobody should copy-construct this class anyway.
-	UCameraComponent(const UCameraComponent&);
 };
diff --git a/Engine/Source/Runtime/Engine/Classes/GameFramework/SpringArmComponent.h b/Engine/Source/Runtime/Engine/Classes/GameFramework/SpringArmComponent.h
--- a/Engine/Source/Runtime/Engine/Classes/GameFramework/SpringArmComponent.h
+++ b/Engine/Source/Runtime/Engine/Classes/GameFramework/SpringArmComponent.h
@@ -120,7 +120,4 @@ public:
 	DEPRECATED(4.5, "This variable is deprecated; existing code using this value may not behave correctly. It only exists to allow compilation of old projects, and code should stop using it in favor of the new bUsePawnControlRotation.")
 	bool bUseControllerViewRotation;
 
-private:
-	// Workaround for compiler-generated function using deprecated variable. Nobody should copy-construct this class anyway.
-	USpringArmComponent(const USpringArmComponent&);
 };
