Hot-Reload: fixing infinite loop when changing USTRUCTs caused by unnecessary re-initialization of the owner UClass.

UE-4916

#change Replaced CLASS_Temporary (UHT-only flag) with CLASS_Constructed to check if a class needs to be constructed or not in generated code
#change Modified code generator to make sure UClasses don't get reconstructed if not required

diff --git a/Engine/Plugins/ScriptGeneratorPlugin/Source/ScriptGeneratorPlugin/Private/ScriptCodeGeneratorBase.cpp b/Engine/Plugins/ScriptGeneratorPlugin/Source/ScriptGeneratorPlugin/Private/ScriptCodeGeneratorBase.cpp
--- a/Engine/Plugins/ScriptGeneratorPlugin/Source/ScriptGeneratorPlugin/Private/ScriptCodeGeneratorBase.cpp
+++ b/Engine/Plugins/ScriptGeneratorPlugin/Source/ScriptGeneratorPlugin/Private/ScriptCodeGeneratorBase.cpp
@@ -148,8 +148,7 @@ FString FScriptCodeGeneratorBase::GetScriptHeaderForClass(UClass* Class)
 
 bool FScriptCodeGeneratorBase::CanExportClass(UClass* Class)
 {
-	bool bCanExport = !(Class->GetClassFlags() & CLASS_Temporary) && // Don't export temporary classes
-		(Class->ClassFlags & (CLASS_RequiredAPI | CLASS_MinimalAPI)) && // Don't export classes that don't export DLL symbols
+	bool bCanExport = (Class->ClassFlags & (CLASS_RequiredAPI | CLASS_MinimalAPI)) && // Don't export classes that don't export DLL symbols
 		!ExportedClasses.Contains(Class->GetFName()); // Don't export classes that have already been exported
 
 	return bCanExport;
diff --git a/Engine/Source/Programs/UnrealHeaderTool/Private/ClassMaps.cpp b/Engine/Source/Programs/UnrealHeaderTool/Private/ClassMaps.cpp
--- a/Engine/Source/Programs/UnrealHeaderTool/Private/ClassMaps.cpp
+++ b/Engine/Source/Programs/UnrealHeaderTool/Private/ClassMaps.cpp
@@ -18,3 +18,4 @@ TMap<UProperty*, FString>                GArrayDimensions;
 TMap<UPackage*,  const FManifestModule*> GPackageToManifestModuleMap;
 TMap<UField*, uint32>                    GGeneratedCodeCRCs;
 TMap<UEnum*,  EPropertyType>             GEnumUnderlyingTypes;
+TSet<UClass*>                            GTemporaryClasses;
\ No newline at end of file
diff --git a/Engine/Source/Programs/UnrealHeaderTool/Private/ClassMaps.h b/Engine/Source/Programs/UnrealHeaderTool/Private/ClassMaps.h
--- a/Engine/Source/Programs/UnrealHeaderTool/Private/ClassMaps.h
+++ b/Engine/Source/Programs/UnrealHeaderTool/Private/ClassMaps.h
@@ -39,6 +39,8 @@ extern TMap<UProperty*, FString>                GArrayDimensions;
 extern TMap<UPackage*,  const FManifestModule*> GPackageToManifestModuleMap;
 extern TMap<UField*, uint32>                    GGeneratedCodeCRCs;
 extern TMap<UEnum*,  EPropertyType>             GEnumUnderlyingTypes;
+extern TSet<UClass*>                            GTemporaryClasses;
+
 
 /** Types access specifiers. */
 enum EAccessSpecifier
diff --git a/Engine/Source/Programs/UnrealHeaderTool/Private/CodeGenerator.cpp b/Engine/Source/Programs/UnrealHeaderTool/Private/CodeGenerator.cpp
--- a/Engine/Source/Programs/UnrealHeaderTool/Private/CodeGenerator.cpp
+++ b/Engine/Source/Programs/UnrealHeaderTool/Private/CodeGenerator.cpp
@@ -787,7 +787,7 @@ void FNativeClassHeaderGenerator::OutputProperties(FString& Meta, FOutputDevice&
 			OutputDevice.Logf( TEXT("#endif // WITH_EDITORONLY_DATA\r\n") );
 			bEmittedHasEditorOnlyMacro = false;
 		}
-		OutputProperty(Meta, OutputDevice, OuterString, Properties[Index], TEXT("            "));
+		OutputProperty(Meta, OutputDevice, OuterString, Properties[Index], Spaces);
 	}
 	if (bEmittedHasEditorOnlyMacro)
 	{
@@ -986,7 +986,7 @@ void FNativeClassHeaderGenerator::ExportNativeGeneratedInitCode(FClass* Class)
 
 	check(!FriendText.Len());
 	// Emit code to build the UObjects that used to be in .u files
-	bool bIsNoExport = Class->HasAnyClassFlags(CLASS_NoExport|CLASS_Temporary);
+	bool bIsNoExport = Class->HasAnyClassFlags(CLASS_NoExport) || FClassUtils::IsTemporaryClass(Class);
 	FStringOutputDevice BodyText;
 	FStringOutputDevice CallSingletons;
 	FString ApiString = FString::Printf(TEXT("%s_API "), *API);
@@ -1052,7 +1052,7 @@ void FNativeClassHeaderGenerator::ExportNativeGeneratedInitCode(FClass* Class)
 
 		uint32 EnumCrc = FCrc::MemCrc32(*GeneratedEnumRegisterFunctionText, GeneratedEnumRegisterFunctionText.Len() * sizeof(TCHAR));		
 		GGeneratedCodeCRCs.Add(Enum, EnumCrc);
-		CallSingletons.Logf(TEXT("            OuterClass->LinkChild(%s); // %u\r\n"), *SingletonName, EnumCrc);
+		CallSingletons.Logf(TEXT("                OuterClass->LinkChild(%s); // %u\r\n"), *SingletonName, EnumCrc);
 	}
 
 	// structs
@@ -1149,7 +1149,7 @@ void FNativeClassHeaderGenerator::ExportNativeGeneratedInitCode(FClass* Class)
 		GeneratedFunctionText += GeneratedStructRegisterFunctionText;
 		GeneratedFunctionText.Logf(TEXT("    uint32 Get_%s_CRC() { return %uU; }\r\n"), *SingletonName.Replace(TEXT("()"), TEXT("")), StructCrc);
 
-		CallSingletons.Logf(TEXT("            OuterClass->LinkChild(%s); // %u\r\n"), *SingletonName, StructCrc);
+		CallSingletons.Logf(TEXT("                OuterClass->LinkChild(%s); // %u\r\n"), *SingletonName, StructCrc);
 	}
 
 	TArray<UFunction*> FunctionsToExport;
@@ -1177,7 +1177,7 @@ void FNativeClassHeaderGenerator::ExportNativeGeneratedInitCode(FClass* Class)
 
 		const FString SingletonName(GetSingletonName(Function));
 
-		CallSingletons.Logf(TEXT("            OuterClass->LinkChild(%s);\r\n"), *SingletonName);
+		CallSingletons.Logf(TEXT("                OuterClass->LinkChild(%s);\r\n"), *SingletonName);
 		FString Extern = FString::Printf(TEXT("    %sclass UFunction* %s;\r\n"), *ApiString, *SingletonName);
 		SingletonNameToExternDecl.Add(SingletonName, Extern);
 		GeneratedFunctionDeclarations.Log(*Extern);
@@ -1280,7 +1280,7 @@ void FNativeClassHeaderGenerator::ExportNativeGeneratedInitCode(FClass* Class)
 	FStringOutputDevice GeneratedClassRegisterFunctionText;
 
 	// The class itself (unless it's a temporary UHT class).
-	if (!Class->HasAnyClassFlags(CLASS_Temporary))
+	if (!FClassUtils::IsTemporaryClass(Class))
 	{
 		// simple ::StaticClass wrapper to avoid header, link and DLL hell
 		{
@@ -1294,8 +1294,6 @@ void FNativeClassHeaderGenerator::ExportNativeGeneratedInitCode(FClass* Class)
 			GeneratedClassRegisterFunctionText.Logf(TEXT("    {\r\n"));
 			GeneratedClassRegisterFunctionText.Logf(TEXT("        return %s::StaticClass();\r\n"), NameLookupCPP.GetNameCPP(Class));
 			GeneratedClassRegisterFunctionText.Logf(TEXT("    }\r\n"));
-
-
 		}
 		FString SingletonName(GetSingletonName(Class));
 
@@ -1315,20 +1313,18 @@ void FNativeClassHeaderGenerator::ExportNativeGeneratedInitCode(FClass* Class)
 		}
 		GeneratedClassRegisterFunctionText.Logf(TEXT("            %s;\r\n"), *GetPackageSingletonName(CastChecked<UPackage>(Class->GetOutermost())));
 		GeneratedClassRegisterFunctionText.Logf(TEXT("            OuterClass = %s::StaticClass();\r\n"), NameLookupCPP.GetNameCPP(Class));
-		GeneratedClassRegisterFunctionText.Logf(TEXT("            UObjectForceRegistration(OuterClass);\r\n"));
-		uint32 Flags = Class->ClassFlags & CLASS_SaveInCompiledInClasses;
-		GeneratedClassRegisterFunctionText.Logf(TEXT("            OuterClass->ClassFlags |= 0x%08X;\r\n"), Flags);
-		//if (Class->IsChildOf(UObject::StaticClass()))
-		//{
-		//	GeneratedClassRegisterFunctionText.Logf(TEXT("            OuterClass->ClassAddReferencedObjects = &%s::AddReferencedObjects;\r\n"), NameLookupCPP.GetNameCPP(Class));
-		//}
+		GeneratedClassRegisterFunctionText.Logf(TEXT("            if (!(OuterClass->ClassFlags & CLASS_Constructed))\r\n"), NameLookupCPP.GetNameCPP(Class));
+		GeneratedClassRegisterFunctionText.Logf(TEXT("            {\r\n"), NameLookupCPP.GetNameCPP(Class));
+		GeneratedClassRegisterFunctionText.Logf(TEXT("                UObjectForceRegistration(OuterClass);\r\n"));
+		uint32 Flags = (Class->ClassFlags & CLASS_SaveInCompiledInClasses) | CLASS_Constructed;
+		GeneratedClassRegisterFunctionText.Logf(TEXT("                OuterClass->ClassFlags |= 0x%08X;\r\n"), Flags);
 		TheFlagAudit.Add(Class, TEXT("ClassFlags"), Flags);
 		GeneratedClassRegisterFunctionText.Logf(TEXT("\r\n"));
 		GeneratedClassRegisterFunctionText.Log(CallSingletons);
 		GeneratedClassRegisterFunctionText.Logf(TEXT("\r\n"));
 
 		FString OuterString = FString(TEXT("OuterClass"));
-		FString Meta = GetMetaDataCodeForObject(Class, *OuterString, TEXT("            "));
+		FString Meta = GetMetaDataCodeForObject(Class, *OuterString, TEXT("                "));
 		// properties
 		{
 			TArray<UProperty*> Props;
@@ -1336,7 +1332,7 @@ void FNativeClassHeaderGenerator::ExportNativeGeneratedInitCode(FClass* Class)
 			{
 				Props.Add(*ItInner);
 			}
-			OutputProperties(Meta, GeneratedClassRegisterFunctionText, OuterString, Props, TEXT("            "));
+			OutputProperties(Meta, GeneratedClassRegisterFunctionText, OuterString, Props, TEXT("                "));
 		}
 		// function table
 		{
@@ -1352,7 +1348,7 @@ void FNativeClassHeaderGenerator::ExportNativeGeneratedInitCode(FClass* Class)
 			// Emit code to construct each UFunction and rebuild the function map at runtime
 			for (UFunction* Function : FunctionsInMap)
 			{
-				GeneratedClassRegisterFunctionText.Logf(TEXT("            OuterClass->AddFunctionToFunctionMap(%s);%s\r\n"), *GetSingletonName(Function), *GetGeneratedCodeCRCTag(Function));
+				GeneratedClassRegisterFunctionText.Logf(TEXT("                OuterClass->AddFunctionToFunctionMap(%s);%s\r\n"), *GetSingletonName(Function), *GetGeneratedCodeCRCTag(Function));
 			}
 		}
 
@@ -1360,7 +1356,7 @@ void FNativeClassHeaderGenerator::ExportNativeGeneratedInitCode(FClass* Class)
 		//GeneratedClassRegisterFunctionText.Logf(TEXT("            OuterClass->ClassFlags = 0x%08X;\r\n"), Class->ClassFlags);
 		if (Class->ClassConfigName != NAME_None)
 		{
-			GeneratedClassRegisterFunctionText.Logf(TEXT("            OuterClass->ClassConfigName = FName(TEXT(\"%s\"));\r\n"), *Class->ClassConfigName.ToString());
+			GeneratedClassRegisterFunctionText.Logf(TEXT("                OuterClass->ClassConfigName = FName(TEXT(\"%s\"));\r\n"), *Class->ClassConfigName.ToString());
 		}
 
 		for (auto& Inter : Class->Interfaces)
@@ -1371,7 +1367,7 @@ void FNativeClassHeaderGenerator::ExportNativeGeneratedInitCode(FClass* Class)
 			{
 				OffsetString = FString::Printf(TEXT("VTABLE_OFFSET(%s, %s)"), NameLookupCPP.GetNameCPP(Class), NameLookupCPP.GetNameCPP(Inter.Class, true));
 			}
-			GeneratedClassRegisterFunctionText.Logf(TEXT("            OuterClass->Interfaces.Add(FImplementedInterface(%s, %s, %s ));\r\n"), 
+			GeneratedClassRegisterFunctionText.Logf(TEXT("                OuterClass->Interfaces.Add(FImplementedInterface(%s, %s, %s ));\r\n"), 
 				*GetSingletonName(Inter.Class, false),
 				*OffsetString,
 				Inter.bImplementedByK2 ? TEXT("true") : TEXT("false")
@@ -1380,19 +1376,20 @@ void FNativeClassHeaderGenerator::ExportNativeGeneratedInitCode(FClass* Class)
 		if (Class->ClassGeneratedBy)
 		{
 			UE_LOG(LogCompile, Fatal, TEXT("For intrinsic and compiled-in classes, ClassGeneratedBy should always be NULL"));
-			GeneratedClassRegisterFunctionText.Logf(TEXT("            OuterClass->ClassGeneratedBy = %s;\r\n"), *GetSingletonName(CastChecked<UClass>(Class->ClassGeneratedBy), false));
+			GeneratedClassRegisterFunctionText.Logf(TEXT("                OuterClass->ClassGeneratedBy = %s;\r\n"), *GetSingletonName(CastChecked<UClass>(Class->ClassGeneratedBy), false));
 		}
 
-		GeneratedClassRegisterFunctionText.Logf(TEXT("            OuterClass->StaticLink();\r\n"));
+		GeneratedClassRegisterFunctionText.Logf(TEXT("                OuterClass->StaticLink();\r\n"));
 
 		if (Meta.Len())
 		{
 			GeneratedClassRegisterFunctionText.Logf(TEXT("#if WITH_METADATA\r\n"));
-			GeneratedClassRegisterFunctionText.Logf(TEXT("            UMetaData* MetaData = OuterClass->GetOutermost()->GetMetaData();\r\n"));
+			GeneratedClassRegisterFunctionText.Logf(TEXT("                UMetaData* MetaData = OuterClass->GetOutermost()->GetMetaData();\r\n"));
 			GeneratedClassRegisterFunctionText.Log(*Meta);
 			GeneratedClassRegisterFunctionText.Logf(TEXT("#endif\r\n"));
 		}
 
+		GeneratedClassRegisterFunctionText.Logf(TEXT("            }\r\n"));
 		GeneratedClassRegisterFunctionText.Logf(TEXT("        }\r\n"));
 		GeneratedClassRegisterFunctionText.Logf(TEXT("        check(OuterClass->GetClass());\r\n"));
 		GeneratedClassRegisterFunctionText.Logf(TEXT("        return OuterClass;\r\n"));
@@ -1414,7 +1411,7 @@ void FNativeClassHeaderGenerator::ExportNativeGeneratedInitCode(FClass* Class)
 	SingletonName.ReplaceInline(TEXT("()"), TEXT("")); // function address
 
 	// Skip temporary UHT classes.
-	if (!Class->HasAnyClassFlags(CLASS_Temporary))
+	if (!FClassUtils::IsTemporaryClass(Class))
 	{	
 		auto ClassNameCPP = NameLookupCPP.GetNameCPP(Class);		
 		GeneratedFunctionText.Logf(TEXT("    static FCompiledInDefer Z_CompiledInDefer_UClass_%s(%s, TEXT(\"%s\"));\r\n"), 
@@ -1430,7 +1427,7 @@ void FNativeClassHeaderGenerator::ExportNativeGeneratedInitCode(FClass* Class)
 void FNativeClassHeaderGenerator::ExportNatives(FClass* Class)
 {
 	// Skip temporary UHT classes.
-	if (Class->HasAnyClassFlags(CLASS_NoExport|CLASS_Temporary))
+	if (FClassUtils::IsNoExportOrTemporaryClass(Class))
 		return;
 
 	GeneratedPackageCPP.Logf(TEXT("    void %s::StaticRegisterNatives%s()\r\n"),NameLookupCPP.GetNameCPP(Class),NameLookupCPP.GetNameCPP(Class));
@@ -1600,7 +1597,7 @@ void FNativeClassHeaderGenerator::ExportInterfaceCallFunctions( const TArray<UFu
 void FNativeClassHeaderGenerator::ExportNames()
 {
 	// Skip temporary UHT classes.
-	if (CurrentClass->HasAnyClassFlags(CLASS_NoExport|CLASS_Temporary))
+	if (FClassUtils::IsNoExportOrTemporaryClass(CurrentClass))
 	{
 		return;
 	}
@@ -2129,7 +2126,7 @@ void FNativeClassHeaderGenerator::ExportConstructorsMacros(FClass* Class)
 
 void FNativeClassHeaderGenerator::ExportClassHeaderWrapper( FClass* Class, bool bIsExportClass )
 {
-	const bool bTemporaryJunkClass = Class->HasAnyClassFlags(CLASS_Temporary);
+	const bool bTemporaryJunkClass = FClassUtils::IsTemporaryClass(Class);
 	const bool bValidClass = !bTemporaryJunkClass;
 
 	check(!GeneratedHeaderText.Len());
@@ -2271,7 +2268,7 @@ void FNativeClassHeaderGenerator::ExportClassHeaderWrapper( FClass* Class, bool
 
 void FNativeClassHeaderGenerator::ExportClassHeaderRecursive( FClasses& AllClasses, FClass* Class, TArray<FClass*>& DependencyChain, TSet<const UClass*>& VisitedSet, bool bCheckDependenciesOnly )
 {
-	bool bIsExportClass   = GClassStrippedHeaderTextMap.Contains(Class) && Class->HasAnyClassFlags(CLASS_Native|CLASS_Temporary) && !Class->HasAnyClassFlags(CLASS_NoExport);
+	bool bIsExportClass   = GClassStrippedHeaderTextMap.Contains(Class) && (Class->HasAnyClassFlags(CLASS_Native) || FClassUtils::IsTemporaryClass(Class)) && !Class->HasAnyClassFlags(CLASS_NoExport);
 	bool bIsCorrectHeader = Class->GetOuter() == Package;
 
 	// Check for circular header dependencies between export classes.
@@ -2348,7 +2345,7 @@ void FNativeClassHeaderGenerator::ExportClassHeaderRecursive( FClasses& AllClass
 	}
 
 	// Export class header.
-	if ( bIsCorrectHeader && !bCheckDependenciesOnly && !Class->HasAnyClassFlags(CLASS_Intrinsic) && Class->HasAnyClassFlags(CLASS_Native|CLASS_Temporary) )
+	if ( bIsCorrectHeader && !bCheckDependenciesOnly && !Class->HasAnyClassFlags(CLASS_Intrinsic) && (Class->HasAnyClassFlags(CLASS_Native) || FClassUtils::IsTemporaryClass(Class)) )
 	{
 		CurrentClass = Class;
 
@@ -3830,7 +3827,7 @@ void FNativeClassHeaderGenerator::ExportCallbackFunctions( const TArray<UFunctio
 		WriteEventFunctionEpilogue(RPCWrappersCPP, 4, Parameters, *ClassName, *FunctionName);
 	}
 
-	if (!CurrentClass->HasAnyClassFlags(CLASS_NoExport|CLASS_Temporary))
+	if (!FClassUtils::IsNoExportOrTemporaryClass(CurrentClass))
 	{
 		GeneratedPackageCPP.Log(*RPCWrappersCPP);
 	}
@@ -4017,8 +4014,8 @@ bool FNativeClassHeaderGenerator::FindInterDependencyRecursive( TMap<const FStri
 
 bool IsExportOrTemporaryClass(FClass* Class)
 {
-	bool bIsExportClass = Class->HasAnyClassFlags(CLASS_Native) && !Class->HasAnyClassFlags(CLASS_NoExport | CLASS_Intrinsic | CLASS_Temporary);
-	bool bIsTemporaryClass = Class->HasAnyClassFlags(CLASS_Temporary);
+	bool bIsTemporaryClass = FClassUtils::IsTemporaryClass(Class);
+	bool bIsExportClass = Class->HasAnyClassFlags(CLASS_Native) && !(Class->HasAnyClassFlags(CLASS_NoExport | CLASS_Intrinsic) || bIsTemporaryClass);
 
 	return bIsExportClass || bIsTemporaryClass;
 }
@@ -4078,13 +4075,13 @@ FNativeClassHeaderGenerator::FNativeClassHeaderGenerator( UPackage* InPackage, F
 		int32 ClassCount = 0;
 		for (FClass* Class : Classes)
 		{
-			if( Class->GetOuter()==Package && Class->HasAnyClassFlags(CLASS_Native|CLASS_Temporary))
+			if( Class->GetOuter()==Package && (Class->HasAnyClassFlags(CLASS_Native) || FClassUtils::IsTemporaryClass(Class)))
 			{
 				if (GClassStrippedHeaderTextMap.Contains(Class) && !Class->HasAnyClassFlags(CLASS_NoExport))
 				{
 					ClassCount++;
 					// Skip temporary classes autogenerated in struct-only headers.
-					if (!Class->HasAnyClassFlags(CLASS_Temporary))
+					if (!FClassUtils::IsTemporaryClass(Class))
 					{
 						Class->UnMark(OBJECTMARK_TagImp);
 						Class->Mark(OBJECTMARK_TagExp);
@@ -4159,7 +4156,7 @@ FNativeClassHeaderGenerator::FNativeClassHeaderGenerator( UPackage* InPackage, F
 		// If no headers are generated, check if there's any no export classes that need to have the inl file generated.
 		for (FClass* Class : Classes)
 		{
-			if (Class->GetOuter() == Package && !Class->HasAnyClassFlags(CLASS_Intrinsic | CLASS_Temporary))
+			if (Class->GetOuter() == Package && !(Class->HasAnyClassFlags(CLASS_Intrinsic) || FClassUtils::IsTemporaryClass(Class)))
 			{
 				ExportClassHeader(AllClasses, Class);
 			}
@@ -5069,7 +5066,8 @@ UClass* GenerateCodeForHeader
 		{
 			const FString DummyClassName = FHeaderParser::GenerateTemporaryClassName(*ResultClass->GetName());
 
-			ResultClass->ClassFlags |= CLASS_Temporary | CLASS_Native;
+			ResultClass->ClassFlags |= CLASS_Native;
+			FClassUtils::MarkAsTemporaryClass(ResultClass);
 			ResultClass->Rename(*DummyClassName, NULL, REN_DontCreateRedirectors);
 		}
 
diff --git a/Engine/Source/Programs/UnrealHeaderTool/Private/HeaderParser.cpp b/Engine/Source/Programs/UnrealHeaderTool/Private/HeaderParser.cpp
--- a/Engine/Source/Programs/UnrealHeaderTool/Private/HeaderParser.cpp
+++ b/Engine/Source/Programs/UnrealHeaderTool/Private/HeaderParser.cpp
@@ -610,7 +610,7 @@ UEnum* FHeaderParser::CompileEnum(UClass* Scope)
 	ParseFieldMetaData(EnumToken.MetaData, EnumToken.Identifier);
 
 	// Create enum definition.
-	UEnum* Enum = new(Scope->HasAnyClassFlags(CLASS_Temporary) ? Scope->GetOuter() : Scope, EnumToken.Identifier, RF_Public) UEnum(FObjectInitializer());
+	UEnum* Enum = new(FClassUtils::IsTemporaryClass(Scope) ? Scope->GetOuter() : Scope, EnumToken.Identifier, RF_Public) UEnum(FObjectInitializer());
 	Enum->Next = Scope->Children;
 	Scope->Children = Enum;
 
@@ -1196,7 +1196,7 @@ UScriptStruct* FHeaderParser::CompileStructDeclaration(FClasses& AllClasses, FCl
 	}
 
 	// Verify uniqueness (if declared within UClass).
-	if (!Scope->HasAnyClassFlags(CLASS_Temporary))
+	if (!FClassUtils::IsTemporaryClass(Scope))
 	{
 		UField* Existing = FindField(Scope, *EffectiveStructName);
 		if (Existing && Existing->GetOuter() == Scope)
@@ -1329,7 +1329,7 @@ UScriptStruct* FHeaderParser::CompileStructDeclaration(FClasses& AllClasses, FCl
 	}
 	
 	// If declared in a header without UClass, the outer is the class package.
-	UObject* StructOuter = Scope->HasAnyClassFlags(CLASS_Temporary) ? Scope->GetOuter() : Scope;
+	UObject* StructOuter = FClassUtils::IsTemporaryClass(Scope) ? Scope->GetOuter() : Scope;
 	// Create.
 	UScriptStruct* Struct = new(StructOuter, *EffectiveStructName, RF_Public) UScriptStruct(FObjectInitializer(), BaseStruct);
 	Struct->Next = Scope->Children;
@@ -6436,7 +6436,7 @@ ECompilationResult::Type FHeaderParser::ParseHeaderForOneClass(FClasses& AllClas
 		}
 
 		// mark temporary classes as native
-		if ( (Class->ClassFlags & CLASS_Temporary) )
+		if (FClassUtils::IsTemporaryClass(Class))
 		{
 			Class->ClassFlags |= CLASS_Native;
 		}
diff --git a/Engine/Source/Programs/UnrealHeaderTool/Private/ParserHelper.h b/Engine/Source/Programs/UnrealHeaderTool/Private/ParserHelper.h
--- a/Engine/Source/Programs/UnrealHeaderTool/Private/ParserHelper.h
+++ b/Engine/Source/Programs/UnrealHeaderTool/Private/ParserHelper.h
@@ -60,6 +60,26 @@ struct ERefQualifier
 	};
 };
 
+/** Misc helper functions */
+struct FClassUtils
+{
+	FORCEINLINE static bool IsTemporaryClass(UClass* InClass)
+	{
+		return GTemporaryClasses.Contains(InClass);
+	}
+	FORCEINLINE static void MarkAsTemporaryClass(UClass* InClass)
+	{
+		if (!IsTemporaryClass(InClass))
+		{
+			GTemporaryClasses.Add(InClass);
+		}
+	}
+	FORCEINLINE static bool IsNoExportOrTemporaryClass(UClass* InClass)
+	{
+		return InClass->HasAnyClassFlags(CLASS_NoExport) || IsTemporaryClass(InClass);
+	}
+};
+
 #ifndef CASE_TEXT
 #define CASE_TEXT(txt) case txt: return TEXT(#txt)
 #endif
@@ -1711,7 +1731,7 @@ struct FNameLookupCPP
 		FString DesiredStructName = Struct->GetName();
 		if (UClass* TestClass = Cast<UClass>(Struct))
 		{
-			if (TestClass->HasAnyClassFlags(CLASS_Temporary))
+			if (FClassUtils::IsTemporaryClass(TestClass))
 			{
 				DesiredStructName = GClassHeaderNameWithNoPathMap[TestClass];
 			}
diff --git a/Engine/Source/Runtime/CoreUObject/Public/UObject/ObjectBase.h b/Engine/Source/Runtime/CoreUObject/Public/UObject/ObjectBase.h
--- a/Engine/Source/Runtime/CoreUObject/Public/UObject/ObjectBase.h
+++ b/Engine/Source/Runtime/CoreUObject/Public/UObject/ObjectBase.h
@@ -173,8 +173,8 @@ enum EClassFlags
 	CLASS_HideDropDown		  = 0x04000000,
 	/** Class was declared directly in C++ and has no boilerplate generated by UnrealHeaderTool */
 	CLASS_Intrinsic			  = 0x10000000,
-	/** Internal, temporary UHT generated class. */
-	CLASS_Temporary		  = 0x20000000,
+	/** Class has already been constructed (maybe in a previous DLL version before hot-reload). */
+	CLASS_Constructed		  = 0x20000000,
 	/** Indicates that object configuration will not check against ini base/defaults when serialized */
 	CLASS_ConfigDoNotCheckDefaults = 0x40000000,
 	/** Class has been consigned to oblivion as part of a blueprint recompile, and a newer version currently exists. */
