Fix for generated code changes (for the Engine) when switching between targets due to uninitialized GeneratedBodyMacroAccessSpecifier member variable.

Added more safety checks and a more detailed error message.

diff --git a/Engine/Source/Programs/UnrealHeaderTool/Private/CodeGenerator.cpp b/Engine/Source/Programs/UnrealHeaderTool/Private/CodeGenerator.cpp
--- a/Engine/Source/Programs/UnrealHeaderTool/Private/CodeGenerator.cpp
+++ b/Engine/Source/Programs/UnrealHeaderTool/Private/CodeGenerator.cpp
@@ -1662,19 +1662,22 @@ FString GetGeneratedMacroDeprecationWarning(const TCHAR* MacroName)
  */
 FString GetPreservedAccessSpecifierString(FClass* Class)
 {
-	FString PreservedAccessSpecifier = "static_assert(false, \"Access specifier for GENERATED_BODY() macro is wrong.\");";
-
-	switch (GScriptHelper.FindClassData(Class)->GeneratedBodyMacroAccessSpecifier)
-	{
-	case EAccessSpecifier::ACCESS_Private:
-		PreservedAccessSpecifier = "private:";
-		break;
-	case EAccessSpecifier::ACCESS_Protected:
-		PreservedAccessSpecifier = "protected:";
-		break;
-	case EAccessSpecifier::ACCESS_Public:
-		PreservedAccessSpecifier = "public:";
-		break;
+	FString PreservedAccessSpecifier(FString::Printf(TEXT("static_assert(false, \"Unknown access specifier for GENERATED_BODY() macro in class %s.\");"), *GetNameSafe(Class)));
+	FClassMetaData* Data = GScriptHelper.FindClassData(Class);	
+	if (Data)
+	{
+		switch (Data->GeneratedBodyMacroAccessSpecifier)
+		{
+		case EAccessSpecifier::ACCESS_Private:
+			PreservedAccessSpecifier = "private:";
+			break;
+		case EAccessSpecifier::ACCESS_Protected:
+			PreservedAccessSpecifier = "protected:";
+			break;
+		case EAccessSpecifier::ACCESS_Public:
+			PreservedAccessSpecifier = "public:";
+			break;
+		}
 	}
 
 	return PreservedAccessSpecifier + LINE_TERMINATOR;
diff --git a/Engine/Source/Programs/UnrealHeaderTool/Private/ParserHelper.h b/Engine/Source/Programs/UnrealHeaderTool/Private/ParserHelper.h
--- a/Engine/Source/Programs/UnrealHeaderTool/Private/ParserHelper.h
+++ b/Engine/Source/Programs/UnrealHeaderTool/Private/ParserHelper.h
@@ -1327,10 +1327,11 @@ class FClassMetaData
 public:
 	/** Default constructor */
 	FClassMetaData()
-		: bContainsDelegates(false),
-		bConstructorDeclared(false),
-		bDefaultConstructorDeclared(false),
-		bObjectInitializerConstructorDeclared(false)
+		: bContainsDelegates(false)
+		, bConstructorDeclared(false)
+		, bDefaultConstructorDeclared(false)
+		, bObjectInitializerConstructorDeclared(false)
+		, GeneratedBodyMacroAccessSpecifier(ACCESS_NotAnAccessSpecifier)
 	{
 	}
 
