Don't prototype FObjectInitializer in generated code; it prevents Doxygen from matching up the funciton signatures.

diff --git a/Engine/Source/Programs/UnrealHeaderTool/Private/CodeGenerator.cpp b/Engine/Source/Programs/UnrealHeaderTool/Private/CodeGenerator.cpp
--- a/Engine/Source/Programs/UnrealHeaderTool/Private/CodeGenerator.cpp
+++ b/Engine/Source/Programs/UnrealHeaderTool/Private/CodeGenerator.cpp
@@ -2041,7 +2041,7 @@ void ExportStandardConstructorsMacro(FStringOutputDevice& Out, FClass* Class, co
 	if (!Class->HasAnyClassFlags(CLASS_CustomConstructor))
 	{
 		Out.Logf(TEXT("%s/** Standard constructor, called after all reflected properties have been initialized */\r\n"), FCString::Spc(4));
-		Out.Logf(TEXT("%s%s_API %s(const class FObjectInitializer& ObjectInitializer);\r\n"), FCString::Spc(4), *API, NameLookupCPP.GetNameCPP(Class));
+		Out.Logf(TEXT("%s%s_API %s(const FObjectInitializer& ObjectInitializer);\r\n"), FCString::Spc(4), *API, NameLookupCPP.GetNameCPP(Class));
 	}
 	Out.Logf(TEXT("%sDEFINE_DEFAULT_OBJECT_INITIALIZER_CONSTRUCTOR_CALL(%s)\r\n"), FCString::Spc(4), NameLookupCPP.GetNameCPP(Class));
 	ExportCopyConstructorDefinition(Out, Class, API);
