Error: "Static array cannot be exposed to blueprint" (previously it was a warning, but it was ignored)

diff --git a/Engine/Source/Programs/UnrealHeaderTool/Private/HeaderParser.cpp b/Engine/Source/Programs/UnrealHeaderTool/Private/HeaderParser.cpp
--- a/Engine/Source/Programs/UnrealHeaderTool/Private/HeaderParser.cpp
+++ b/Engine/Source/Programs/UnrealHeaderTool/Private/HeaderParser.cpp
@@ -6069,7 +6069,7 @@ void FHeaderParser::CompileVariableDeclaration(FClasses& AllClasses, UStruct* St
 
 		if (NewProperty->HasAnyPropertyFlags(CPF_BlueprintVisible) && (NewProperty->ArrayDim > 1))
 		{
-			UE_LOG(LogCompile, Warning, TEXT("Static array cannot be exposed to blueprint %s.%s\n"), *Struct->GetName(), *NewProperty->GetName());
+			FError::Throwf(TEXT("Static array cannot be exposed to blueprint %s.%s"), *Struct->GetName(), *NewProperty->GetName());
 		}
 
 	} while( MatchSymbol(TEXT(",")) );
