Don't automatically create function declarations in UHT that user has to implement.

diff --git a/Engine/Source/Programs/UnrealHeaderTool/Private/CodeGenerator.cpp b/Engine/Source/Programs/UnrealHeaderTool/Private/CodeGenerator.cpp
--- a/Engine/Source/Programs/UnrealHeaderTool/Private/CodeGenerator.cpp
+++ b/Engine/Source/Programs/UnrealHeaderTool/Private/CodeGenerator.cpp
@@ -1830,23 +1830,36 @@ void FNativeClassHeaderGenerator::ExportInterfaceClassDeclaration( FClass* Class
 
 		// Replication, add in the declaration for GetLifetimeReplicatedProps() automatically if there are any net flagged properties
 		bool bNeedsRep = false;
+		FStringOutputDevice StaticChecks;
+		auto ClassName = FString(NameLookupCPP.GetNameCPP(Class));
 		for( TFieldIterator<UProperty> It(Class,EFieldIteratorFlags::ExcludeSuper); It; ++It )
 		{
 			if( (It->PropertyFlags & CPF_Net) != 0 )
 			{
 				bNeedsRep = true;
+				StaticChecks.Logf(TEXT("\tstatic inline void StaticChecks()\r\n"));
+				StaticChecks.Logf(TEXT("\t{\r\n"));
+				StaticChecks.Logf(TEXT("\t\tstatic_assert(THasMemberFunction_GetLifetimeReplicatedProps<%s>::Value, \"Class %s has Net flagged properties and should declare member function: virtual void GetLifetimeReplicatedProps( TArray< FLifetimeProperty > & OutLifetimeProps ) const override;\");\r\n"), ClassCPPName, ClassCPPName, ClassCPPName);
+				StaticChecks.Logf(TEXT("\t}\r\n"));
 				break;
 			}
 		}
+		
+		auto NoPureDeclsSuffix = TEXT("_INCLASS_IINTERFACE_NO_PURE_DECLS");
+		auto NoPureDeclsMacroName = ClassName + NoPureDeclsSuffix;
+		auto Macroized = Macroize(*NoPureDeclsMacroName, *(InterfaceBoilerplate + StaticChecks));
+		GeneratedHeaderText.Log(*Macroized);
+
 		if (bNeedsRep)
 		{
-			InterfaceBoilerplate.Logf(TEXT("\tvirtual void GetLifetimeReplicatedProps( TArray< FLifetimeProperty > & OutLifetimeProps ) const override;\r\n"));
+			InterfaceBoilerplate.Logf(TEXT("\tvirtual void GetLifetimeReplicatedProps(TArray<FLifetimeProperty>& OutLifetimeProps) const override;\r\n"));
 		}
 		const TCHAR* Suffix = TEXT("_INCLASS_IINTERFACE");
-		FString MacroName = FString(NameLookupCPP.GetNameCPP(Class)) + Suffix;
-		FString Macroized = Macroize(*MacroName, *InterfaceBoilerplate);
+		FString MacroName = ClassName + Suffix;
+		Macroized = Macroize(*MacroName, *InterfaceBoilerplate);
 		GeneratedHeaderText.Log(*Macroized);
-		InClassMacroCalls.Logf( TEXT("%s%s\r\n"), FCString::Spc(4),  *MacroName);
+		InClassMacroCalls.Logf(TEXT("%s%s\r\n"), FCString::Spc(4), *MacroName);
+		InClassNoPureDeclsMacroCalls.Logf(TEXT("%s%s\r\n"), FCString::Spc(4), *NoPureDeclsMacroName);
 	}
 }
 
@@ -1987,21 +2000,39 @@ void FNativeClassHeaderGenerator::ExportClassHeaderInner(FClass* Class, bool bVa
 				ClassBoilerplate.Logf(TEXT("%sUObject* _getUObject() const { return const_cast<%s*>(this); }\r\n"), FCString::Spc(4), ClassCPPName);
 
 				// Replication, add in the declaration for GetLifetimeReplicatedProps() automatically if there are any net flagged properties
-				for( TFieldIterator<UProperty> It(Class,EFieldIteratorFlags::ExcludeSuper); It; ++It )
+				bool bHasNetFlaggedProperties = false;
+				FStringOutputDevice StaticChecks;
+				for (auto Field : TFieldRange<UProperty>(Class, EFieldIteratorFlags::ExcludeSuper))
 				{
-					if( (It->PropertyFlags & CPF_Net) != 0 )
+					if ((Field->PropertyFlags & CPF_Net) != 0)
 					{
-						ClassBoilerplate.Logf(TEXT("\tvirtual void GetLifetimeReplicatedProps( TArray< FLifetimeProperty > & OutLifetimeProps ) const override;\r\n"));
+						bHasNetFlaggedProperties = true;
+						StaticChecks.Logf(TEXT("\tstatic inline void StaticChecks()\r\n"));
+						StaticChecks.Logf(TEXT("\t{\r\n"));
+						StaticChecks.Logf(TEXT("\t\tstatic_assert(THasMemberFunction_GetLifetimeReplicatedProps<%s>::Value, \"Class %s has Net flagged properties and should have member virtual void GetLifetimeReplicatedProps( TArray< FLifetimeProperty > & OutLifetimeProps ) const override;\");\r\n"), ClassCPPName, ClassCPPName, ClassCPPName);
+						StaticChecks.Logf(TEXT("\t}\r\n"));
 						break;
 					}
 				}
 
+				auto ClassNameString = FString(ClassCPPName);
+				auto NoPureDeclsSuffix = TEXT("_INCLASS_NO_PURE_DECLS");
+				auto NoPureDeclsMacroName = ClassNameString + NoPureDeclsSuffix;
+				auto Macroized = Macroize(*NoPureDeclsMacroName, *(ClassBoilerplate + StaticChecks));
+				GeneratedHeaderText.Log(*Macroized);
+
+				if (bHasNetFlaggedProperties)
+				{
+					ClassBoilerplate.Logf(TEXT("\tvirtual void GetLifetimeReplicatedProps(TArray<FLifetimeProperty>& OutLifetimeProps) const override;\r\n"));
+				}
 
 				const TCHAR* Suffix = TEXT("_INCLASS");
-				FString MacroName = FString(NameLookupCPP.GetNameCPP(Class)) + Suffix;
-				FString Macroized = Macroize(*MacroName, *ClassBoilerplate);
+				FString MacroName = ClassNameString + Suffix;
+				Macroized = Macroize(*MacroName, *ClassBoilerplate);
 				GeneratedHeaderText.Log(*Macroized);
+
 				InClassMacroCalls.Logf(TEXT("%s%s\r\n"), FCString::Spc(4), *MacroName);
+				InClassNoPureDeclsMacroCalls.Logf(TEXT("%s%s\r\n"), FCString::Spc(4), *NoPureDeclsMacroName);
 
 				ExportConstructorsMacros(Class);
 			}
@@ -2190,7 +2221,7 @@ void FNativeClassHeaderGenerator::ExportClassHeaderWrapper( FClass* Class, bool
 			else
 			{
 				auto StandardWrappedInClassMacroCalls = Public + InClassMacroCalls + StandardUObjectConstructorsMacroCall + Public;
-				auto EnhancedWrappedInClassMacroCalls = Public + InClassMacroCalls + EnhancedUObjectConstructorsMacroCall + GetPreservedAccessSpecifierString(Class);
+				auto EnhancedWrappedInClassMacroCalls = Public + InClassNoPureDeclsMacroCalls + EnhancedUObjectConstructorsMacroCall + GetPreservedAccessSpecifierString(Class);
 
 				BodyMacros = Macroize(*MacroName, *(DeprecationWarning + DeprecationPushString + StandardWrappedInClassMacroCalls + DeprecationPopString)) +
 					Macroize(TEXT("GENERATED_BODY()"), *(FString(DeprecationPushString) + EnhancedWrappedInClassMacroCalls + DeprecationPopString));
@@ -2199,6 +2230,7 @@ void FNativeClassHeaderGenerator::ExportClassHeaderWrapper( FClass* Class, bool
 			GeneratedHeaderText.Log(*BodyMacros);
 
 			InClassMacroCalls.Empty();
+			InClassNoPureDeclsMacroCalls.Empty();
 			StandardUObjectConstructorsMacroCall.Empty();
 			EnhancedUObjectConstructorsMacroCall.Empty();
 		}
@@ -3385,6 +3417,99 @@ FString FNativeClassHeaderGenerator::GetNullParameterValue( UProperty* Prop, boo
 	return TEXT("");
 }
 
+
+FString FNativeClassHeaderGenerator::GetFunctionReturnString(UFunction* Function)
+{
+	if (auto Return = Function->GetReturnProperty())
+	{
+		FString ExtendedReturnType;
+		FString ReturnType = Return->GetCPPType(&ExtendedReturnType, CPPF_ArgumentOrReturnValue);
+		FStringOutputDevice ReplacementText(*ReturnType);
+		ApplyAlternatePropertyExportText(Return, ReplacementText);
+		return ReplacementText + ExtendedReturnType;
+	}
+
+	return TEXT("void");
+}
+
+/**
+* Gets string with function const modifier type.
+*
+* @param Function Function to get const modifier of.
+* @return Empty FString if function is non-const, FString("const") if function is const.
+*/
+FString GetFunctionConstModifierString(UFunction* Function)
+{
+	if (Function->HasAllFunctionFlags(FUNC_Const))
+	{
+		return TEXT("const");
+	}
+
+	return FString();
+}
+
+void FNativeClassHeaderGenerator::ExportFunctionChecks(const FFuncInfo& FunctionData, FStringOutputDevice& CheckClasses, FStringOutputDevice& StaticChecks, const FString& ClassName)
+{
+	auto Function = FunctionData.FunctionReference;
+	auto FunctionReturnType = GetFunctionReturnString(Function);
+	auto ConstModifier = GetFunctionConstModifierString(Function);
+
+	auto bIsNative = Function->HasAllFunctionFlags(FUNC_Native);
+	auto bIsNet = Function->HasAllFunctionFlags(FUNC_Net);
+	auto bIsNetValidate = Function->HasAllFunctionFlags(FUNC_NetValidate);
+	auto bIsNetResponse = Function->HasAllFunctionFlags(FUNC_NetResponse);
+	auto bIsBlueprintEvent = Function->HasAllFunctionFlags(FUNC_BlueprintEvent);
+
+	bool bNeedsImplementation = (bIsNet && !bIsNetResponse) || bIsBlueprintEvent || bIsNative;
+	bool bNeedsValidate = (bIsNative || bIsNet) && !bIsNetResponse && bIsNetValidate;
+
+	check(bNeedsImplementation || bNeedsValidate);
+
+	auto ParameterString = GetFunctionParameterString(Function);
+
+	//
+	// Get string with function specifiers, listing why we need _Implementation or _Validate functions.
+	//
+	TArray<FString> FunctionSpecifiers;
+	FunctionSpecifiers.Reserve(4);
+	if (bIsNative)			{ FunctionSpecifiers.Add(TEXT("Native"));			}
+	if (bIsNet)				{ FunctionSpecifiers.Add(TEXT("Net"));				}
+	if (bIsBlueprintEvent)	{ FunctionSpecifiers.Add(TEXT("BlueprintEvent"));	}
+	if (bIsNetValidate)		{ FunctionSpecifiers.Add(TEXT("NetValidate"));		}
+
+	check(FunctionSpecifiers.Num() > 0);
+
+	//
+	// Coin static_assert message
+	//
+	FStringOutputDevice AssertMessage;
+	AssertMessage.Logf(TEXT("Function %s was marked as %s"), *(Function->GetName()), *FunctionSpecifiers[0]);
+	for (int32 i = 1; i < FunctionSpecifiers.Num(); ++i)
+	{
+		AssertMessage.Logf(TEXT(", %s"), *FunctionSpecifiers[i]);
+	}
+
+	AssertMessage.Logf(TEXT("."));
+
+	//
+	// Add proper checks if necessary.
+	//
+	FString FunctionName;
+	if (bNeedsImplementation)
+	{
+		FunctionName = FunctionData.CppImplName;
+		StaticChecks.Logf(TEXT("\t\tstatic_assert(THasMemberFunction_%s<%s>::Value, \"%s Declare function virtual %s %s(%s) %s in class %s.\");\r\n"), *FunctionName, *ClassName, *AssertMessage, *FunctionReturnType, *FunctionName, *ParameterString, *ConstModifier, *ClassName);
+		CheckClasses.Logf(TEXT("GENERATE_MEMBER_FUNCTION_CHECK(%s, %s, %s, %s)\r\n"), *FunctionName, *FunctionReturnType, *ConstModifier, *ParameterString);
+	}
+	
+	if (bNeedsValidate)
+	{
+		FunctionName = FunctionData.CppValidationImplName;
+		StaticChecks.Logf(TEXT("\t\tstatic_assert(THasMemberFunction_%s<%s>::Value, \"%s Declare function virtual bool %s(%s) %s in class %s.\");\r\n"), *FunctionName, *ClassName, *AssertMessage, *FunctionReturnType, *FunctionName, *ParameterString, *ConstModifier, *ClassName);
+		CheckClasses.Logf(TEXT("GENERATE_MEMBER_FUNCTION_CHECK(%s, bool, %s, %s)\r\n"), *FunctionName, *ConstModifier, *ParameterString);
+	}
+}
+
 void FNativeClassHeaderGenerator::ExportNativeFunctionHeader( const FFuncInfo& FunctionData, FStringOutputDevice& HeaderOutput, EExportFunctionType::Type FunctionType, EExportFunctionHeaderStyle::Type FunctionHeaderStyle, const TCHAR* ExtraParam )
 {
 	UFunction* Function = FunctionData.FunctionReference;
@@ -3520,7 +3645,7 @@ void FNativeClassHeaderGenerator::ExportFunctionThunk(FStringOutputDevice& RPCWr
 {
 	// export the GET macro for this parameter
 	FString ParameterList;
-	for ( int32 ParameterIndex = 0; ParameterIndex < Parameters.Num(); ParameterIndex++ )
+	for (int32 ParameterIndex = 0; ParameterIndex < Parameters.Num(); ParameterIndex++)
 	{
 		UProperty* Param = Parameters[ParameterIndex];
 
@@ -3529,7 +3654,7 @@ void FNativeClassHeaderGenerator::ExportFunctionThunk(FStringOutputDevice& RPCWr
 		FString EvalParameterText;				// e.g. (UObject*,NULL)
 
 		FString TypeText;
-		if ( Param->ArrayDim > 1 )
+		if (Param->ArrayDim > 1)
 		{
 			EvalBaseText += TEXT("ARRAY");
 			TypeText = Param->GetCPPType();
@@ -3553,7 +3678,7 @@ void FNativeClassHeaderGenerator::ExportFunctionThunk(FStringOutputDevice& RPCWr
 		}
 
 		// if this property requires a specialization, add a comma to the type name so we can print it out easily
-		if ( TypeText != TEXT("") )
+		if (TypeText != TEXT(""))
 		{
 			TypeText += TCHAR(',');
 		}
@@ -3565,30 +3690,30 @@ void FNativeClassHeaderGenerator::ExportFunctionThunk(FStringOutputDevice& RPCWr
 		RPCWrappers.Logf(TEXT("%s%s%s%s;%s"), FCString::Spc(8), *EvalBaseText, *EvalModifierText, *EvalParameterText, LINE_TERMINATOR);
 
 		// add this property to the parameter list string
-		if ( ParameterList.Len() )
+		if (ParameterList.Len())
 		{
 			ParameterList += TCHAR(',');
 		}
 
 		{
-			UDelegateProperty* DelegateProp = Cast< UDelegateProperty >( Param );
-			if( DelegateProp != NULL )
+			UDelegateProperty* DelegateProp = Cast< UDelegateProperty >(Param);
+			if (DelegateProp != NULL)
 			{
 				// For delegates, add an explicit conversion to the specific type of delegate before passing it along
-				const FString FunctionName = DelegateProp->SignatureFunction->GetName().LeftChop( FString( HEADER_GENERATED_DELEGATE_SIGNATURE_SUFFIX ).Len() );
-				const FString CPPDelegateName = FString( TEXT("F" ) ) + FunctionName;
-				ParamName = FString::Printf( TEXT( "%s(%s)" ), *CPPDelegateName, *ParamName );
+				const FString FunctionName = DelegateProp->SignatureFunction->GetName().LeftChop(FString(HEADER_GENERATED_DELEGATE_SIGNATURE_SUFFIX).Len());
+				const FString CPPDelegateName = FString(TEXT("F")) + FunctionName;
+				ParamName = FString::Printf(TEXT("%s(%s)"), *CPPDelegateName, *ParamName);
 			}
 		}
 
 		{
-			UMulticastDelegateProperty* MulticastDelegateProp = Cast< UMulticastDelegateProperty >( Param );
-			if( MulticastDelegateProp != NULL )
+			UMulticastDelegateProperty* MulticastDelegateProp = Cast< UMulticastDelegateProperty >(Param);
+			if (MulticastDelegateProp != NULL)
 			{
 				// For delegates, add an explicit conversion to the specific type of delegate before passing it along
-				const FString FunctionName = MulticastDelegateProp->SignatureFunction->GetName().LeftChop( FString( HEADER_GENERATED_DELEGATE_SIGNATURE_SUFFIX ).Len() );
-				const FString CPPDelegateName = FString( TEXT("F" ) ) + FunctionName;
-				ParamName = FString::Printf( TEXT( "%s(%s)" ), *CPPDelegateName, *ParamName );
+				const FString FunctionName = MulticastDelegateProp->SignatureFunction->GetName().LeftChop(FString(HEADER_GENERATED_DELEGATE_SIGNATURE_SUFFIX).Len());
+				const FString CPPDelegateName = FString(TEXT("F")) + FunctionName;
+				ParamName = FString::Printf(TEXT("%s(%s)"), *CPPDelegateName, *ParamName);
 			}
 		}
 
@@ -3598,11 +3723,11 @@ void FNativeClassHeaderGenerator::ExportFunctionThunk(FStringOutputDevice& RPCWr
 			// For enums, add an explicit conversion 
 			if (!(ByteProp->PropertyFlags & CPF_OutParm))
 			{
-				ParamName = FString::Printf( TEXT( "%s(%s)" ), *ByteProp->Enum->CppType, *ParamName );
+				ParamName = FString::Printf(TEXT("%s(%s)"), *ByteProp->Enum->CppType, *ParamName);
 			}
 			else
 			{
-				ParamName = FString::Printf( TEXT( "(TEnumAsByte<%s>&)(%s)" ), *ByteProp->Enum->CppType, *ParamName );
+				ParamName = FString::Printf(TEXT("(TEnumAsByte<%s>&)(%s)"), *ByteProp->Enum->CppType, *ParamName);
 			}
 		}
 
@@ -3612,20 +3737,20 @@ void FNativeClassHeaderGenerator::ExportFunctionThunk(FStringOutputDevice& RPCWr
 	RPCWrappers.Logf(TEXT("%sP_FINISH;%s"), FCString::Spc(8), LINE_TERMINATOR);
 
 	// Call the validate function if there is one
-	if ( !( FunctionData.FunctionExportFlags & FUNCEXPORT_CppStatic ) && ( FunctionData.FunctionFlags & FUNC_NetValidate ) )
+	if (!(FunctionData.FunctionExportFlags & FUNCEXPORT_CppStatic) && (FunctionData.FunctionFlags & FUNC_NetValidate))
 	{
 		// Call the validate function, and if it fails, return (which will NOT execute the _Implementation function below)
 		// NOTE - We don't need to set "Result", since RPC calls don't return any values...
 		RPCWrappers.Logf(TEXT("%sif (!this->%s(%s))%s"), FCString::Spc(8), *FunctionData.CppValidationImplName, *ParameterList, LINE_TERMINATOR);
 		RPCWrappers.Logf(TEXT("%s{%s"), FCString::Spc(8), LINE_TERMINATOR);
-		RPCWrappers.Logf(TEXT("%sRPC_ValidateFailed(TEXT(\"%s\"));%s"), FCString::Spc(8+4), *FunctionData.CppValidationImplName, LINE_TERMINATOR);
-		RPCWrappers.Logf(TEXT("%sreturn;%s"), FCString::Spc(8+4), LINE_TERMINATOR);	// If we got here, the validation function check failed
+		RPCWrappers.Logf(TEXT("%sRPC_ValidateFailed(TEXT(\"%s\"));%s"), FCString::Spc(8 + 4), *FunctionData.CppValidationImplName, LINE_TERMINATOR);
+		RPCWrappers.Logf(TEXT("%sreturn;%s"), FCString::Spc(8 + 4), LINE_TERMINATOR);	// If we got here, the validation function check failed
 		RPCWrappers.Logf(TEXT("%s}%s"), FCString::Spc(8), LINE_TERMINATOR);
 	}
 
 	// write out the return value
 	RPCWrappers.Log(FCString::Spc(8));
-	if ( Return != NULL )
+	if (Return != NULL)
 	{
 		FString ReturnType, ReturnExtendedType;
 		ReturnType = Return->GetCPPType(&ReturnExtendedType);
@@ -3646,107 +3771,123 @@ void FNativeClassHeaderGenerator::ExportFunctionThunk(FStringOutputDevice& RPCWr
 	}
 }
 
+FString FNativeClassHeaderGenerator::GetFunctionParameterString(UFunction* Function)
+{
+	FString ParameterList;
+	FStringOutputDevice PropertyText;
+
+	for (auto Property : TFieldRange<UProperty>(Function))
+	{
+		if ((Property->PropertyFlags & (CPF_Parm | CPF_ReturnParm)) != CPF_Parm)
+		{
+			break;
+		}
+
+		if (ParameterList.Len())
+		{
+			ParameterList += TEXT(", ");
+		}
+
+		auto Dim = GArrayDimensions.Find(Property);
+		Property->ExportCppDeclaration(PropertyText, EExportedDeclaration::Parameter, Dim ? **Dim : NULL);
+		ApplyAlternatePropertyExportText(Property, PropertyText);
+
+		ParameterList += PropertyText;
+		PropertyText.Empty();
+	}
+
+	return ParameterList;
+}
+
 void FNativeClassHeaderGenerator::ExportNativeFunctions(const TArray<UFunction*>& NativeFunctions)
 {
 	FStringOutputDevice RPCWrappers;
+	FStringOutputDevice	AutogeneratedBlueprintFunctionDeclarations;
+	FStringOutputDevice	MemberFunctionCheckClasses;
+	FStringOutputDevice StaticChecks;
+	StaticChecks.Logf(TEXT("\tstatic inline void StaticChecks_Implementation_Validate()\r\n"));
+	StaticChecks.Logf(TEXT("\t{\r\n"));
 
 	// find the class info for this class
 	FClassMetaData* ClassData = GScriptHelper.FindClassData(CurrentClass);
 
-	//@todo: UCREMOVAL this multipass stuff is only used for diffs, remove it
-	TArray<UFunction*> SortedNativeFunctions = NativeFunctions;
-	SortedNativeFunctions.Sort();
+	auto ClassName = FString(NameLookupCPP.GetNameCPP(CurrentClass));
 
 	// export the C++ stubs
-	for (int32 Pass = 0; Pass < 2; Pass++)
+	for (auto Function : NativeFunctions)
 	{
-		for ( int32 i = NativeFunctions.Num() - 1; i >= 0; i-- )
-		{
-			UFunction*     Function     = Pass ? SortedNativeFunctions[i] : NativeFunctions[i];
-			FFunctionData* CompilerInfo = ClassData->FindFunctionData(Function);
-			check(CompilerInfo);
-
-			const FFuncInfo& FunctionData = CompilerInfo->GetFunctionData();
+		FFunctionData* CompilerInfo = ClassData->FindFunctionData(Function);
+		check(CompilerInfo);
 
-			// Custom thunks don't get any C++ stub function generated
-			if (FunctionData.FunctionExportFlags & FUNCEXPORT_CustomThunk)
-				continue;
+		const FFuncInfo& FunctionData = CompilerInfo->GetFunctionData();
 
-			// All functions that are declared in class headers are by definition native
-			check(Function->FunctionFlags & FUNC_Native);
+		// Custom thunks don't get any C++ stub function generated
+		if (FunctionData.FunctionExportFlags & FUNCEXPORT_CustomThunk)
+			continue;
 
-			// Should we emit these to RPC wrappers or just ignore them?
-			const bool bWillBeProgrammerTyped = FunctionData.CppImplName == Function->GetName();
+		// All functions that are declared in class headers are by definition native
+		check(Function->FunctionFlags & FUNC_Native);
 
-			FStringOutputDevice WrongPass;
-			FStringOutputDevice& DestinationForDecl = (!bWillBeProgrammerTyped && Pass) ? RPCWrappers : WrongPass;
+		// Should we emit these to RPC wrappers or just ignore them?
+		const bool bWillBeProgrammerTyped = FunctionData.CppImplName == Function->GetName();
 
+		if (!bWillBeProgrammerTyped)
+		{
 			// Declare validation function if needed
-			if ( FunctionData.FunctionFlags & FUNC_NetValidate ) 
+			if (FunctionData.FunctionFlags & FUNC_NetValidate)
 			{
-				FString ParameterList;
-
-				for ( TFieldIterator<UProperty> It(Function); It && (It->PropertyFlags&(CPF_Parm|CPF_ReturnParm))==CPF_Parm; ++It )
-				{
-					UProperty* Property = *It;
-
-					if ( ParameterList.Len() )
-					{
-						ParameterList += TEXT(", ");
-					}
-
-					FStringOutputDevice PropertyText;
-
-					const FString* Dim = GArrayDimensions.Find(Property);
-					Property->ExportCppDeclaration(PropertyText, EExportedDeclaration::Parameter, Dim ? **Dim : NULL);
-					ApplyAlternatePropertyExportText(Property, PropertyText);
-
-					ParameterList += PropertyText;
-				}
+				FString ParameterList = GetFunctionParameterString(Function);
 
 				// this is not an event, the function is not a static function and the function is not marked final
-				if (!FunctionData.FunctionReference->HasAnyFunctionFlags(FUNC_Static) && !(FunctionData.FunctionExportFlags & FUNCEXPORT_Final))
-				{
-					DestinationForDecl.Logf(TEXT("%s virtual bool %s(%s);\r\n"), FCString::Spc(4), *FunctionData.CppValidationImplName, *ParameterList);
-				}
-				else
-				{
-					DestinationForDecl.Logf(TEXT("%s bool %s(%s);\r\n"), FCString::Spc(4), *FunctionData.CppValidationImplName, *ParameterList);
-				}
+				auto Virtual = (!FunctionData.FunctionReference->HasAnyFunctionFlags(FUNC_Static) && !(FunctionData.FunctionExportFlags & FUNCEXPORT_Final)) ? TEXT("virtual") : TEXT("");
+				AutogeneratedBlueprintFunctionDeclarations.Logf(TEXT("\t%s bool %s(%s);\r\n"), Virtual, *FunctionData.CppValidationImplName, *ParameterList);
 			}
 
-			ExportNativeFunctionHeader(FunctionData, DestinationForDecl, EExportFunctionType::Function, EExportFunctionHeaderStyle::Declaration);
-
-			DestinationForDecl.Log(TEXT(";"));
-			if (bMultiLineUFUNCTION)
+			if ((FunctionData.bCppImplNameEndsWith_Implementation || FunctionData.bCppValidationImplNameEndsWith_Validate))
 			{
-				DestinationForDecl.Log(TEXT("\r\n"));
+				ExportNativeFunctionHeader(FunctionData, AutogeneratedBlueprintFunctionDeclarations, EExportFunctionType::Function, EExportFunctionHeaderStyle::Declaration);
+				AutogeneratedBlueprintFunctionDeclarations.Log(TEXT(";\r\n"));
+				ExportFunctionChecks(FunctionData, MemberFunctionCheckClasses, StaticChecks, ClassName);
 			}
-			DestinationForDecl.Log(TEXT("\r\n"));
+		}
 
-			// if this function was originally declared in a base class, and it isn't a static function,
-			// only the C++ function header will be exported
-			if (!Pass || !ShouldExportFunction(Function))
-			{
-				continue;
-			}
+		if (bMultiLineUFUNCTION)
+		{
+			RPCWrappers.Log(TEXT("\r\n"));
+		}
+		RPCWrappers.Log(TEXT("\r\n"));
 
-			// export the script wrappers
-			RPCWrappers.Logf(TEXT("%sDECLARE_FUNCTION(%s)"), FCString::Spc(4), *FunctionData.UnMarshallAndCallName);
-			RPCWrappers.Logf(TEXT("%s%s{%s"), LINE_TERMINATOR, FCString::Spc(4), LINE_TERMINATOR);
+		// if this function was originally declared in a base class, and it isn't a static function,
+		// only the C++ function header will be exported
+		if (!ShouldExportFunction(Function))
+		{
+			continue;
+		}
 
-			auto Parameters = GetFunctionParmsAndReturn(FunctionData.FunctionReference);
-			ExportFunctionThunk(RPCWrappers, FunctionData, Parameters.Parms, Parameters.Return);
+		// export the script wrappers
+		RPCWrappers.Logf(TEXT("%sDECLARE_FUNCTION(%s)"), FCString::Spc(4), *FunctionData.UnMarshallAndCallName);
+		RPCWrappers.Logf(TEXT("%s%s{%s"), LINE_TERMINATOR, FCString::Spc(4), LINE_TERMINATOR);
 
-			RPCWrappers.Logf(TEXT("%s}%s"), FCString::Spc(4), LINE_TERMINATOR);
-		}
+		auto Parameters = GetFunctionParmsAndReturn(FunctionData.FunctionReference);
+		ExportFunctionThunk(RPCWrappers, FunctionData, Parameters.Parms, Parameters.Return);
+
+		RPCWrappers.Logf(TEXT("%s}%s"), FCString::Spc(4), LINE_TERMINATOR);
 	}
 
-	const TCHAR* Suffix = TEXT("_RPC_WRAPPERS");
-	FString MacroName = FString(NameLookupCPP.GetNameCPP(CurrentClass)) + Suffix;
-	FString Macroized = Macroize(*MacroName, *RPCWrappers);
+	StaticChecks.Logf(TEXT("\t}\r\n"));
+
+	auto Suffix = FString(TEXT("_RPC_WRAPPERS"));
+	FString MacroName = ClassName + Suffix;
+	FString Macroized = Macroize(*MacroName, *(AutogeneratedBlueprintFunctionDeclarations + RPCWrappers));
 	GeneratedHeaderText.Log(*Macroized);
-	InClassMacroCalls.Logf( TEXT("%s%s\r\n"), FCString::Spc(4),  *MacroName);
+	InClassMacroCalls.Logf(TEXT("%s%s\r\n"), FCString::Spc(4), *MacroName);
+
+	auto NoPureDeclsSuffix = FString(TEXT("_RPC_WRAPPERS_NO_PURE_DECLS"));
+	MacroName = ClassName + NoPureDeclsSuffix;
+	// Put static checks before RPCWrappers to get proper messages from static asserts before compiler errors.
+	Macroized = Macroize(*MacroName, *(MemberFunctionCheckClasses + StaticChecks + RPCWrappers));
+	GeneratedHeaderText.Log(*Macroized);
+	InClassNoPureDeclsMacroCalls.Logf(TEXT("%s%s\r\n"), FCString::Spc(4), *MacroName);
 }
 
 /**
@@ -3842,6 +3983,7 @@ void FNativeClassHeaderGenerator::ExportCallbackFunctions( const TArray<UFunctio
 	FString Macroized = Macroize(*MacroName, *RPCWrappers);
 	GeneratedHeaderText.Log(*Macroized);
 	InClassMacroCalls.Logf( TEXT("%s%s\r\n"), FCString::Spc(4),  *MacroName);
+	InClassNoPureDeclsMacroCalls.Logf(TEXT("%s%s\r\n"), FCString::Spc(4), *MacroName);
 }
 
 
diff --git a/Engine/Source/Programs/UnrealHeaderTool/Private/NativeClassExporter.h b/Engine/Source/Programs/UnrealHeaderTool/Private/NativeClassExporter.h
--- a/Engine/Source/Programs/UnrealHeaderTool/Private/NativeClassExporter.h
+++ b/Engine/Source/Programs/UnrealHeaderTool/Private/NativeClassExporter.h
@@ -61,9 +61,11 @@ private:
 	FStringOutputDevice GeneratedMCPText;
 	FStringOutputDevice	PrologMacroCalls;
 	FStringOutputDevice	InClassMacroCalls;
+	FStringOutputDevice	InClassNoPureDeclsMacroCalls;
 	FStringOutputDevice	StandardUObjectConstructorsMacroCall;
 	FStringOutputDevice	EnhancedUObjectConstructorsMacroCall;
 	FStringOutputDevice	AllConstructors;
+	FStringOutputDevice StaticChecks;
 
 	/** Generated function implementations that belong in the cpp file, split into multiple files base on line count **/
 	TArray<FStringOutputDeviceCountLines> GeneratedFunctionBodyTextSplit;
@@ -357,6 +359,16 @@ private:
 	 */
 	void ExportNativeFunctionHeader( const FFuncInfo& FunctionData, FStringOutputDevice& HeaderOutput, EExportFunctionType::Type FunctionType, EExportFunctionHeaderStyle::Type FunctionHeaderStyle, const TCHAR* ExtraParam = NULL );
 
+	/**
+	* Exports checks if function exists
+	*
+	* @param	FunctionData			Data representing the function to export.
+	* @param	CheckClasses			Where to write the member check classes.
+	* @param	StaticChecks			Where to write the static asserts throwing errors when function doesn't exist.
+	* @param	ClassName				Name of currently parsed class.
+	*/
+	void ExportFunctionChecks(const FFuncInfo& FunctionData, FStringOutputDevice& CheckClasses, FStringOutputDevice& StaticChecks, const FString& ClassName);
+
 	/**
 	 * Exports the native stubs for the list of functions specified
 	 * 
@@ -476,4 +488,20 @@ public:
 
 	// Constructor
 	FNativeClassHeaderGenerator( UPackage* InPackage, FClasses& AllClasses, bool InAllowSaveExportedHeaders );
+
+	/**
+	 * Gets string with function return type.
+	 * 
+	 * @param Function Function to get return type of.
+	 * @return FString with function return type.
+	 */
+	FString GetFunctionReturnString(UFunction* Function);
+
+	/**
+	* Gets string with function parameters (with names).
+	*
+	* @param Function Function to get parameters of.
+	* @return FString with function parameters.
+	*/
+	FString GetFunctionParameterString(UFunction* Function);
 };
diff --git a/Engine/Source/Programs/UnrealHeaderTool/Private/ParserHelper.h b/Engine/Source/Programs/UnrealHeaderTool/Private/ParserHelper.h
--- a/Engine/Source/Programs/UnrealHeaderTool/Private/ParserHelper.h
+++ b/Engine/Source/Programs/UnrealHeaderTool/Private/ParserHelper.h
@@ -889,8 +889,12 @@ struct FFuncInfo
 	FString		MarshallAndCallName;
 	/** Name of the actual implementation **/
 	FString		CppImplName;
+	/** Cached info whether CppImplName ends with _Implementation suffix to avoid string comparisons. **/
+	bool		bCppImplNameEndsWith_Implementation;
 	/** Name of the actual validation implementation **/
 	FString		CppValidationImplName;
+	/** Cached info whether CppValidationImplName ends with _Validate suffix to avoid string comparisons. **/
+	bool		bCppValidationImplNameEndsWith_Validate;
 	/** Name for callback-style names **/
 	FString		UnMarshallAndCallName;
 	/** Identifier for an RPC call to a platform service */
@@ -911,6 +915,8 @@ struct FFuncInfo
 	,	RPCId               (0)
 	,	RPCResponseId       (0)
 	,	bSealedEvent        (false)
+	,	bCppImplNameEndsWith_Implementation(false)
+	,	bCppValidationImplNameEndsWith_Validate(false)
 	{}
 
 	FFuncInfo( const FFuncInfo& Other )
@@ -920,6 +926,8 @@ struct FFuncInfo
 	,	FunctionReference(Other.FunctionReference)
 	,	RPCId(Other.RPCId)
 	,	RPCResponseId(Other.RPCResponseId)
+	,	bCppImplNameEndsWith_Implementation(Other.bCppImplNameEndsWith_Implementation)
+	,	bCppValidationImplNameEndsWith_Validate(Other.bCppValidationImplNameEndsWith_Validate)
 	{
 		Function.Clone(Other.Function);
 		if (FunctionReference)
@@ -961,9 +969,11 @@ struct FFuncInfo
 			else
 			{
 				CppImplName = FunctionReference->GetName() + TEXT("_Implementation");
+				bCppImplNameEndsWith_Implementation = true;
 				if (FunctionReference->HasAllFunctionFlags(FUNC_NetValidate))
 				{
 					CppValidationImplName = FunctionReference->GetName() + TEXT("_Validate");
+					bCppValidationImplNameEndsWith_Validate = true;
 				}
 			}
 		}
@@ -977,6 +987,7 @@ struct FFuncInfo
 		{
 			MarshallAndCallName = FunctionName;
 			CppImplName = FunctionReference->GetName() + TEXT("_Implementation");
+			bCppImplNameEndsWith_Implementation = true;
 		}
 	}
 };
diff --git a/Engine/Source/Runtime/Core/Public/Misc/CoreMiscDefines.h b/Engine/Source/Runtime/Core/Public/Misc/CoreMiscDefines.h
--- a/Engine/Source/Runtime/Core/Public/Misc/CoreMiscDefines.h
+++ b/Engine/Source/Runtime/Core/Public/Misc/CoreMiscDefines.h
@@ -105,3 +105,6 @@ const FPlatformUserId PLATFORMUSERID_NONE = INDEX_NONE;
 // helpers to turn an preprocessor token into a real string (see UBT_COMPILED_PLATFORM)
 #define PREPROCESSOR_TO_STRING_INNER(x) #x
 #define PREPROCESSOR_TO_STRING(x) PREPROCESSOR_TO_STRING_INNER(x)
+
+#define PREPROCESSOR_JOIN_INNER(x, y) x##y
+#define PREPROCESSOR_JOIN(x, y) PREPROCESSOR_JOIN_INNER(x, y)
diff --git a/Engine/Source/Runtime/Core/Public/Templates/UnrealTypeTraits.h b/Engine/Source/Runtime/Core/Public/Templates/UnrealTypeTraits.h
--- a/Engine/Source/Runtime/Core/Public/Templates/UnrealTypeTraits.h
+++ b/Engine/Source/Runtime/Core/Public/Templates/UnrealTypeTraits.h
@@ -587,3 +587,14 @@ template <> struct TIsBitwiseConstructible<uint32,  int32> { enum { Value = true
 template <> struct TIsBitwiseConstructible< int32, uint32> { enum { Value = true }; };
 template <> struct TIsBitwiseConstructible<uint64,  int64> { enum { Value = true }; };
 template <> struct TIsBitwiseConstructible< int64, uint64> { enum { Value = true }; };
+
+#define GENERATE_MEMBER_FUNCTION_CHECK(MemberName, Result, ConstModifier, ...)									\
+template <typename T>																							\
+class THasMemberFunction_##MemberName																			\
+{																												\
+	template <typename U, Result(U::*)(__VA_ARGS__) ConstModifier> struct Check;								\
+	template <typename U> static char MemberTest(Check<U, &U::MemberName> *);									\
+	template <typename U> static int MemberTest(...);															\
+public:																											\
+	enum { Value = sizeof(MemberTest<T>(nullptr)) == sizeof(char) };											\
+};
\ No newline at end of file
diff --git a/Engine/Source/Runtime/CoreUObject/Private/UObject/Obj.cpp b/Engine/Source/Runtime/CoreUObject/Private/UObject/Obj.cpp
--- a/Engine/Source/Runtime/CoreUObject/Private/UObject/Obj.cpp
+++ b/Engine/Source/Runtime/CoreUObject/Private/UObject/Obj.cpp
@@ -54,7 +54,7 @@ UObject::UObject( EStaticConstructor, EObjectFlags InFlags )
 
 UObject* UObject::CreateDefaultSubobject(FName SubobjectFName, UClass* ReturnType, UClass* ClassToCreateByDefault, bool bIsRequired, bool bAbstract, bool bIsTransient)
 {
-	UE_CLOG(!GIsInConstructor, LogObj, Fatal, TEXT("CreateDefultSubobject can only be used inside of UObject constructors"));
+	UE_CLOG(!GIsInConstructor, LogObj, Fatal, TEXT("CreateDefultSubobject can only be used inside of UObject constructors. UObject constructing subobjects cannot be created using new or placement new operator."));
 	auto CurrentInitializer = FTlsObjectInitializers::Top();
 	UE_CLOG(!CurrentInitializer, LogObj, Fatal, TEXT("No object initializer found during construction."));
 	UE_CLOG(CurrentInitializer->Obj != this, LogObj, Fatal, TEXT("Using incorrect object initializer."));
diff --git a/Engine/Source/Runtime/CoreUObject/Public/UObject/ObjectBase.h b/Engine/Source/Runtime/CoreUObject/Public/UObject/ObjectBase.h
--- a/Engine/Source/Runtime/CoreUObject/Public/UObject/ObjectBase.h
+++ b/Engine/Source/Runtime/CoreUObject/Public/UObject/ObjectBase.h
@@ -1439,3 +1439,4 @@ private:
 	UClass* Class;
 };
 
+GENERATE_MEMBER_FUNCTION_CHECK(GetLifetimeReplicatedProps, void, const, TArray<FLifetimeProperty>&)
diff --git a/Engine/Source/Runtime/Engine/Classes/Components/ActorComponent.h b/Engine/Source/Runtime/Engine/Classes/Components/ActorComponent.h
--- a/Engine/Source/Runtime/Engine/Classes/Components/ActorComponent.h
+++ b/Engine/Source/Runtime/Engine/Classes/Components/ActorComponent.h
@@ -27,6 +27,8 @@ public:
 	 */
 	UActorComponent(const FObjectInitializer& ObjectInitializer);
 
+	virtual void GetLifetimeReplicatedProps(TArray<FLifetimeProperty>& OutLifetimeProps) const override;
+
 	/** Main tick function for the Actor */
 	UPROPERTY()
 	struct FActorComponentTickFunction PrimaryComponentTick;
diff --git a/Engine/Source/Runtime/Engine/Classes/Components/SceneComponent.h b/Engine/Source/Runtime/Engine/Classes/Components/SceneComponent.h
--- a/Engine/Source/Runtime/Engine/Classes/Components/SceneComponent.h
+++ b/Engine/Source/Runtime/Engine/Classes/Components/SceneComponent.h
@@ -87,6 +87,8 @@ public:
 	 */
 	USceneComponent(const FObjectInitializer& ObjectInitializer);
 
+	virtual void GetLifetimeReplicatedProps(TArray<FLifetimeProperty>& OutLifetimeProps) const override;
+
 	/** Current transform of this component, relative to the world */
 	FTransform ComponentToWorld;
 
diff --git a/Engine/Source/Runtime/Engine/Classes/GameFramework/Actor.h b/Engine/Source/Runtime/Engine/Classes/GameFramework/Actor.h
--- a/Engine/Source/Runtime/Engine/Classes/GameFramework/Actor.h
+++ b/Engine/Source/Runtime/Engine/Classes/GameFramework/Actor.h
@@ -75,6 +75,8 @@ public:
 	 */
 	AActor(const FObjectInitializer& ObjectInitializer);
 
+	void GetLifetimeReplicatedProps(TArray<FLifetimeProperty>& OutLifetimeProps) const override;
+
 	/**
 	 * Primary Actor tick function, which calls TickActor().
 	 * Tick functions can be configured to control whether ticking is enabled, at what time during a frame the update occurs, and to set up tick dependencies.
diff --git a/Engine/Source/Runtime/Engine/Classes/GameFramework/Character.h b/Engine/Source/Runtime/Engine/Classes/GameFramework/Character.h
--- a/Engine/Source/Runtime/Engine/Classes/GameFramework/Character.h
+++ b/Engine/Source/Runtime/Engine/Classes/GameFramework/Character.h
@@ -207,6 +207,7 @@ public:
 	 */
 	ACharacter(const FObjectInitializer& ObjectInitializer);
 
+	void GetLifetimeReplicatedProps(TArray<FLifetimeProperty>& OutLifetimeProps) const override;
 private_subobject:
 	/** The main skeletal mesh associated with this Character (optional sub-object). */
 	DEPRECATED_FORGAME(4.6, "Mesh should not be accessed directly, please use GetMesh() function instead. Mesh will soon be private and your code will not compile.")
@@ -426,6 +427,7 @@ protected:
 
 	UFUNCTION(BlueprintNativeEvent, Category="Pawn|Character|InternalEvents", meta=(FriendlyName="CanJump"))
 	bool CanJumpInternal() const;
+	virtual bool CanJumpInternal_Implementation() const;
 
 public:
 
@@ -479,6 +481,7 @@ public:
 	/** Event fired when the character has just started jumping */
 	UFUNCTION(BlueprintNativeEvent, Category="Pawn|Character")
 	void OnJumped();
+	void OnJumped_Implementation();
 
 	/** Called when the character's movement enters falling */
 	virtual void Falling() {}
@@ -513,6 +516,7 @@ public:
 	 */
 	UFUNCTION(BlueprintNativeEvent, Category="Pawn|Character")
 	void OnWalkingOffLedge();
+	virtual void OnWalkingOffLedge_Implementation();
 
 	/** Called when pawn's movement is blocked
 		@PARAM Impact describes the blocking hit. */
@@ -636,12 +640,15 @@ public:
 
 	UFUNCTION(reliable, client)
 	void ClientCheatWalk();
+	void ClientCheatWalk_Implementation();
 
 	UFUNCTION(reliable, client)
 	void ClientCheatFly();
+	void ClientCheatFly_Implementation();
 
 	UFUNCTION(reliable, client)
 	void ClientCheatGhost();
+	void ClientCheatGhost_Implementation();
 
 	// Root Motion
 public:
diff --git a/Engine/Source/Runtime/Engine/Classes/GameFramework/CharacterMovementComponent.h b/Engine/Source/Runtime/Engine/Classes/GameFramework/CharacterMovementComponent.h
--- a/Engine/Source/Runtime/Engine/Classes/GameFramework/CharacterMovementComponent.h
+++ b/Engine/Source/Runtime/Engine/Classes/GameFramework/CharacterMovementComponent.h
@@ -1608,31 +1608,40 @@ public:
 	/** Replicated function sent by client to server - contains client movement and view info. */
 	UFUNCTION(unreliable, server, WithValidation)
 	virtual void ServerMove(float TimeStamp, FVector_NetQuantize10 InAccel, FVector_NetQuantize100 ClientLoc, uint8 CompressedMoveFlags, uint8 ClientRoll, uint32 View, UPrimitiveComponent* ClientMovementBase, FName ClientBaseBoneName, uint8 ClientMovementMode);
+	virtual void ServerMove_Implementation(float TimeStamp, FVector_NetQuantize10 InAccel, FVector_NetQuantize100 ClientLoc, uint8 CompressedMoveFlags, uint8 ClientRoll, uint32 View, UPrimitiveComponent* ClientMovementBase, FName ClientBaseBoneName, uint8 ClientMovementMode);
+	virtual bool ServerMove_Validate(float TimeStamp, FVector_NetQuantize10 InAccel, FVector_NetQuantize100 ClientLoc, uint8 CompressedMoveFlags, uint8 ClientRoll, uint32 View, UPrimitiveComponent* ClientMovementBase, FName ClientBaseBoneName, uint8 ClientMovementMode);
 
 	/** Replicated function sent by client to server - contains client movement and view info for two moves. */
 	UFUNCTION(unreliable, server, WithValidation)
 	virtual void ServerMoveDual(float TimeStamp0, FVector_NetQuantize10 InAccel0, uint8 PendingFlags, uint32 View0, float TimeStamp, FVector_NetQuantize10 InAccel, FVector_NetQuantize100 ClientLoc, uint8 NewFlags, uint8 ClientRoll, uint32 View, UPrimitiveComponent* ClientMovementBase, FName ClientBaseBoneName, uint8 ClientMovementMode);
-	
+	virtual void ServerMoveDual_Implementation(float TimeStamp0, FVector_NetQuantize10 InAccel0, uint8 PendingFlags, uint32 View0, float TimeStamp, FVector_NetQuantize10 InAccel, FVector_NetQuantize100 ClientLoc, uint8 NewFlags, uint8 ClientRoll, uint32 View, UPrimitiveComponent* ClientMovementBase, FName ClientBaseBoneName, uint8 ClientMovementMode);
+	virtual bool ServerMoveDual_Validate(float TimeStamp0, FVector_NetQuantize10 InAccel0, uint8 PendingFlags, uint32 View0, float TimeStamp, FVector_NetQuantize10 InAccel, FVector_NetQuantize100 ClientLoc, uint8 NewFlags, uint8 ClientRoll, uint32 View, UPrimitiveComponent* ClientMovementBase, FName ClientBaseBoneName, uint8 ClientMovementMode);
+
 	/* Resending an (important) old move. Process it if not already processed. */
 	UFUNCTION(unreliable, server, WithValidation)
 	virtual void ServerMoveOld(float OldTimeStamp, FVector_NetQuantize10 OldAccel, uint8 OldMoveFlags);
+	virtual void ServerMoveOld_Implementation(float OldTimeStamp, FVector_NetQuantize10 OldAccel, uint8 OldMoveFlags);
+	virtual bool ServerMoveOld_Validate(float OldTimeStamp, FVector_NetQuantize10 OldAccel, uint8 OldMoveFlags);
 	
 	/** If no client adjustment is needed after processing received ServerMove(), ack the good move so client can remove it from SavedMoves */
 	UFUNCTION(unreliable, client)
 	virtual void ClientAckGoodMove(float TimeStamp);
+	virtual void ClientAckGoodMove_Implementation(float TimeStamp);
 
 	/** Replicate position correction to client, associated with a timestamped servermove.  Client will replay subsequent moves after applying adjustment.  */
 	UFUNCTION(unreliable, client)
 	virtual void ClientAdjustPosition(float TimeStamp, FVector NewLoc, FVector NewVel, UPrimitiveComponent* NewBase, FName NewBaseBoneName, bool bHasBase, bool bBaseRelativePosition, uint8 ServerMovementMode);
+	virtual void ClientAdjustPosition_Implementation(float TimeStamp, FVector NewLoc, FVector NewVel, UPrimitiveComponent* NewBase, FName NewBaseBoneName, bool bHasBase, bool bBaseRelativePosition, uint8 ServerMovementMode);
 
 	/* Bandwidth saving version, when velocity is zeroed */
 	UFUNCTION(unreliable, client)
 	virtual void ClientVeryShortAdjustPosition(float TimeStamp, FVector NewLoc, UPrimitiveComponent* NewBase, FName NewBaseBoneName, bool bHasBase, bool bBaseRelativePosition, uint8 ServerMovementMode);
+	virtual void ClientVeryShortAdjustPosition_Implementation(float TimeStamp, FVector NewLoc, UPrimitiveComponent* NewBase, FName NewBaseBoneName, bool bHasBase, bool bBaseRelativePosition, uint8 ServerMovementMode);
 	
 	/** Replicate position correction to client when using root motion for movement. */
 	UFUNCTION(unreliable, client)
 	void ClientAdjustRootMotionPosition(float TimeStamp, float ServerMontageTrackPosition, FVector ServerLoc, FVector_NetQuantizeNormal ServerRotation, float ServerVelZ, UPrimitiveComponent* ServerBase, FName ServerBoneName, bool bHasBase, bool bBaseRelativePosition, uint8 ServerMovementMode);
-
+	void ClientAdjustRootMotionPosition_Implementation(float TimeStamp, float ServerMontageTrackPosition, FVector ServerLoc, FVector_NetQuantizeNormal ServerRotation, float ServerVelZ, UPrimitiveComponent* ServerBase, FName ServerBoneName, bool bHasBase, bool bBaseRelativePosition, uint8 ServerMovementMode);
 
 	// Root Motion
 public:
diff --git a/Engine/Source/Runtime/Engine/Classes/GameFramework/Pawn.h b/Engine/Source/Runtime/Engine/Classes/GameFramework/Pawn.h
--- a/Engine/Source/Runtime/Engine/Classes/GameFramework/Pawn.h
+++ b/Engine/Source/Runtime/Engine/Classes/GameFramework/Pawn.h
@@ -38,6 +38,8 @@ public:
 	 */
 	APawn(const FObjectInitializer& ObjectInitializer);
 
+	virtual void GetLifetimeReplicatedProps(TArray<FLifetimeProperty>& OutLifetimeProps) const;
+
 	/** Return our PawnMovementComponent, if we have one. By default, returns the first PawnMovementComponent found. Native classes that create their own movement component should override this method for more efficiency. */
 	UFUNCTION(BlueprintCallable, meta=(Tooltip="Return our PawnMovementComponent, if we have one."), Category="Pawn")
 	virtual UPawnMovementComponent* GetMovementComponent() const;
