XBoxOne CIS fix.

diff --git a/Engine/Source/Programs/UnrealHeaderTool/Private/CodeGenerator.cpp b/Engine/Source/Programs/UnrealHeaderTool/Private/CodeGenerator.cpp
--- a/Engine/Source/Programs/UnrealHeaderTool/Private/CodeGenerator.cpp
+++ b/Engine/Source/Programs/UnrealHeaderTool/Private/CodeGenerator.cpp
@@ -1847,7 +1847,7 @@ void FNativeClassHeaderGenerator::ExportInterfaceClassDeclaration( FClass* Class
 		
 		auto NoPureDeclsSuffix = TEXT("_INCLASS_IINTERFACE_NO_PURE_DECLS");
 		auto NoPureDeclsMacroName = ClassName + NoPureDeclsSuffix;
-		auto Macroized = Macroize(*NoPureDeclsMacroName, *(InterfaceBoilerplate + StaticChecks));
+		auto Macroized = Macroize(*NoPureDeclsMacroName, *(InterfaceBoilerplate));
 		GeneratedHeaderText.Log(*Macroized);
 
 		if (bNeedsRep)
@@ -3847,7 +3847,7 @@ void FNativeClassHeaderGenerator::ExportNativeFunctions(const TArray<UFunction*>
 			{
 				ExportNativeFunctionHeader(FunctionData, AutogeneratedBlueprintFunctionDeclarations, EExportFunctionType::Function, EExportFunctionHeaderStyle::Declaration);
 				AutogeneratedBlueprintFunctionDeclarations.Log(TEXT(";\r\n"));
-				ExportFunctionChecks(FunctionData, MemberFunctionCheckClasses, StaticChecks, ClassName);
+				//ExportFunctionChecks(FunctionData, MemberFunctionCheckClasses, StaticChecks, ClassName);
 			}
 		}
 
