Emit warning if GetLifetimeReplicatedProps and functions with _Implementation and _Validate suffixes aren't declared.

diff --git a/Engine/Source/Programs/UnrealHeaderTool/Private/CodeGenerator.cpp b/Engine/Source/Programs/UnrealHeaderTool/Private/CodeGenerator.cpp
--- a/Engine/Source/Programs/UnrealHeaderTool/Private/CodeGenerator.cpp
+++ b/Engine/Source/Programs/UnrealHeaderTool/Private/CodeGenerator.cpp
@@ -1602,7 +1602,7 @@ void FNativeClassHeaderGenerator::ExportNames()
 	{
 		return;
 	}
-	GeneratedHeaderText.Logf(
+	GeneratedHeaderTextBeforeForwardDeclarations.Logf(
 		TEXT("#ifdef %s_%s_generated_h")													LINE_TERMINATOR
 		TEXT("#error \"%s.generated.h already included, missing '#pragma once' in %s.h\"")	LINE_TERMINATOR
 		TEXT("#endif")																		LINE_TERMINATOR
@@ -1635,7 +1635,7 @@ void FNativeClassHeaderGenerator::ExportNames()
 
 	for( int32 NameIndex=0; NameIndex<Names.Num(); NameIndex++ )
 	{
-		GeneratedHeaderText.Logf( TEXT("extern %s_API FName %s_%s;") LINE_TERMINATOR, *API, *API, *Names[NameIndex] );
+		GeneratedHeaderTextBeforeForwardDeclarations.Logf(TEXT("extern %s_API FName %s_%s;") LINE_TERMINATOR, *API, *API, *Names[NameIndex]);
 	}
 
 }
@@ -1775,7 +1775,7 @@ void FNativeClassHeaderGenerator::ExportInterfaceClassDeclaration( FClass* Class
 		auto DeprecationWarning = GetGeneratedMacroDeprecationWarning(TEXT("GENERATED_UINTERFACE_BODY"));
 		
 		auto DeprecationPushString = TEXT("PRAGMA_DISABLE_DEPRECATION_WARNINGS") LINE_TERMINATOR;
-		auto DeprecationPopString = TEXT("PRAGMA_ENABLE_DEPRECATION_WARNINGS") LINE_TERMINATOR;
+		auto DeprecationPopString = TEXT("PRAGMA_POP") LINE_TERMINATOR;
 		auto Offset = FCString::Spc(4);
 
 		GeneratedHeaderText.Log(TEXT("#undef GENERATED_UINTERFACE_BODY\r\n"));
@@ -1839,24 +1839,36 @@ void FNativeClassHeaderGenerator::ExportInterfaceClassDeclaration( FClass* Class
 				bNeedsRep = true;
 				StaticChecks.Logf(TEXT("\tstatic inline void StaticChecks()\r\n"));
 				StaticChecks.Logf(TEXT("\t{\r\n"));
-				StaticChecks.Logf(TEXT("\t\tstatic_assert(THasMemberFunction_GetLifetimeReplicatedProps<%s>::Value, \"Class %s has Net flagged properties and should declare member function: virtual void GetLifetimeReplicatedProps( TArray< FLifetimeProperty > & OutLifetimeProps ) const override;\");\r\n"), ClassCPPName, ClassCPPName, ClassCPPName);
+				StaticChecks.Logf(TEXT("\t\tstatic_assert(THasMemberFunction_GetLifetimeReplicatedProps<%s>::Value, \"Class %s has Net flagged properties and should declare member function: void GetLifetimeReplicatedProps( TArray< FLifetimeProperty > & OutLifetimeProps ) const override;\");\r\n"), ClassCPPName, ClassCPPName, ClassCPPName);
 				StaticChecks.Logf(TEXT("\t}\r\n"));
 				break;
 			}
 		}
 		
+		auto ClassRange = ClassDefinitionRange();
+		if (ClassDefinitionRanges.Contains(CurrentClass))
+		{
+			ClassRange = ClassDefinitionRanges[CurrentClass];
+		}
+
+		auto ClassStart = ClassRange.Start;
+		auto ClassEnd = ClassRange.End;
+		auto ClassDefinition = FString(ClassEnd - ClassStart, ClassStart);
+
+		bool bHasGetLifetimeReplicatedProps = ClassDefinition.Contains(TEXT("GetLifetimeReplicatedProps"), ESearchCase::CaseSensitive);
+
+		if (bNeedsRep && !bHasGetLifetimeReplicatedProps)
+		{
+			InterfaceBoilerplate.Logf(TEXT("\void GetLifetimeReplicatedProps(TArray<FLifetimeProperty>& OutLifetimeProps) const;\r\n"));
+		}
 		auto NoPureDeclsSuffix = TEXT("_INCLASS_IINTERFACE_NO_PURE_DECLS");
 		auto NoPureDeclsMacroName = ClassName + NoPureDeclsSuffix;
 		auto Macroized = Macroize(*NoPureDeclsMacroName, *(InterfaceBoilerplate));
 		GeneratedHeaderText.Log(*Macroized);
 
-		if (bNeedsRep)
-		{
-			InterfaceBoilerplate.Logf(TEXT("\tvirtual void GetLifetimeReplicatedProps(TArray<FLifetimeProperty>& OutLifetimeProps) const override;\r\n"));
-		}
 		const TCHAR* Suffix = TEXT("_INCLASS_IINTERFACE");
 		FString MacroName = ClassName + Suffix;
-		Macroized = Macroize(*MacroName, *InterfaceBoilerplate);
+		Macroized = Macroize(*MacroName, *(InterfaceBoilerplate + StaticChecks));
 		GeneratedHeaderText.Log(*Macroized);
 		InClassMacroCalls.Logf(TEXT("%s%s\r\n"), FCString::Spc(4), *MacroName);
 		InClassNoPureDeclsMacroCalls.Logf(TEXT("%s%s\r\n"), FCString::Spc(4), *NoPureDeclsMacroName);
@@ -2009,23 +2021,35 @@ void FNativeClassHeaderGenerator::ExportClassHeaderInner(FClass* Class, bool bVa
 						bHasNetFlaggedProperties = true;
 						StaticChecks.Logf(TEXT("\tstatic inline void StaticChecks()\r\n"));
 						StaticChecks.Logf(TEXT("\t{\r\n"));
-						StaticChecks.Logf(TEXT("\t\tstatic_assert(THasMemberFunction_GetLifetimeReplicatedProps<%s>::Value, \"Class %s has Net flagged properties and should have member virtual void GetLifetimeReplicatedProps( TArray< FLifetimeProperty > & OutLifetimeProps ) const override;\");\r\n"), ClassCPPName, ClassCPPName, ClassCPPName);
+						StaticChecks.Logf(TEXT("\t\tstatic_assert(THasMemberFunction_GetLifetimeReplicatedProps<%s>::Value, \"Class %s has Net flagged properties and should have member void GetLifetimeReplicatedProps( TArray< FLifetimeProperty > & OutLifetimeProps ) const override;\");\r\n"), ClassCPPName, ClassCPPName, ClassCPPName);
 						StaticChecks.Logf(TEXT("\t}\r\n"));
 						break;
 					}
 				}
 
+				auto ClassRange = ClassDefinitionRange();
+				if (ClassDefinitionRanges.Contains(CurrentClass))
+				{
+					ClassRange = ClassDefinitionRanges[CurrentClass];
+				}
+
+				auto ClassStart = ClassRange.Start;
+				auto ClassEnd = ClassRange.End;
+				auto ClassDefinition = FString(ClassEnd - ClassStart, ClassStart);
+
+				bool bHasGetLifetimeReplicatedProps = ClassDefinition.Contains(TEXT("GetLifetimeReplicatedProps"), ESearchCase::CaseSensitive);
+
+				if (bHasNetFlaggedProperties && !bHasGetLifetimeReplicatedProps)
+				{
+					ClassBoilerplate.Logf(TEXT("\tvoid GetLifetimeReplicatedProps(TArray<FLifetimeProperty>& OutLifetimeProps) const override;\r\n"));
+				}
+
 				auto ClassNameString = FString(ClassCPPName);
 				auto NoPureDeclsSuffix = TEXT("_INCLASS_NO_PURE_DECLS");
 				auto NoPureDeclsMacroName = ClassNameString + NoPureDeclsSuffix;
 				auto Macroized = Macroize(*NoPureDeclsMacroName, *(ClassBoilerplate + StaticChecks));
 				GeneratedHeaderText.Log(*Macroized);
 
-				if (bHasNetFlaggedProperties)
-				{
-					ClassBoilerplate.Logf(TEXT("\tvirtual void GetLifetimeReplicatedProps(TArray<FLifetimeProperty>& OutLifetimeProps) const override;\r\n"));
-				}
-
 				const TCHAR* Suffix = TEXT("_INCLASS");
 				FString MacroName = ClassNameString + Suffix;
 				Macroized = Macroize(*MacroName, *ClassBoilerplate);
@@ -2208,7 +2232,7 @@ void FNativeClassHeaderGenerator::ExportClassHeaderWrapper( FClass* Class, bool
 			auto DeprecationWarning = bIsIInterface ? FString(TEXT("")) : GetGeneratedMacroDeprecationWarning(*MacroName);
 
 			auto DeprecationPushString = TEXT("PRAGMA_DISABLE_DEPRECATION_WARNINGS") LINE_TERMINATOR;
-			auto DeprecationPopString = TEXT("PRAGMA_ENABLE_DEPRECATION_WARNINGS") LINE_TERMINATOR;
+			auto DeprecationPopString = TEXT("PRAGMA_POP") LINE_TERMINATOR;
 
 			auto Public = TEXT("public:" LINE_TERMINATOR);
 
@@ -2216,7 +2240,7 @@ void FNativeClassHeaderGenerator::ExportClassHeaderWrapper( FClass* Class, bool
 			
 			if (bIsIInterface)
 			{
-				BodyMacros = Macroize(*MacroName, *(DeprecationPushString + FString(Public) + InClassMacroCalls + Public + DeprecationPopString));
+				BodyMacros = Macroize(*MacroName, *(FString(Public) + InClassMacroCalls + Public));
 			}
 			else
 			{
@@ -2224,7 +2248,7 @@ void FNativeClassHeaderGenerator::ExportClassHeaderWrapper( FClass* Class, bool
 				auto EnhancedWrappedInClassMacroCalls = Public + InClassNoPureDeclsMacroCalls + EnhancedUObjectConstructorsMacroCall + GetPreservedAccessSpecifierString(Class);
 
 				BodyMacros = Macroize(*MacroName, *(DeprecationWarning + DeprecationPushString + StandardWrappedInClassMacroCalls + DeprecationPopString)) +
-					Macroize(TEXT("GENERATED_BODY()"), *(FString(DeprecationPushString) + EnhancedWrappedInClassMacroCalls + DeprecationPopString));
+					Macroize(TEXT("GENERATED_BODY()"), *(DeprecationPushString + FString(EnhancedWrappedInClassMacroCalls) + DeprecationPopString));
 			}
 
 			GeneratedHeaderText.Log(*BodyMacros);
@@ -2267,6 +2291,24 @@ void FNativeClassHeaderGenerator::ExportClassHeaderWrapper( FClass* Class, bool
 			LINE_TERMINATOR
 			TEXT("#include \"ObjectBase.h\"") LINE_TERMINATOR
 			LINE_TERMINATOR);
+		GeneratedHeaderTextWithCopyright.Log(*GeneratedHeaderTextBeforeForwardDeclarations);
+
+		TSet<FString> ForwardDeclarationStrings;
+		for (auto Property : ForwardDeclarations)
+		{
+			auto FWDecl = Property->GetCPPTypeForwardDeclaration();
+			if (FWDecl.Len() > 0)
+			{
+				ForwardDeclarationStrings.Add(FWDecl);
+			}
+		}
+
+		for (auto FWDeclString : ForwardDeclarationStrings)
+		{
+			GeneratedForwardDeclarations.Logf(TEXT("%s\r\n"), *FWDeclString);
+		}
+
+		GeneratedHeaderTextWithCopyright.Log(*GeneratedForwardDeclarations);
 		GeneratedHeaderTextWithCopyright.Log(*GeneratedHeaderText);
 
 		auto ClassHeaderInfo = GClassGeneratedFileMap.Find(Class);
@@ -2297,7 +2339,9 @@ void FNativeClassHeaderGenerator::ExportClassHeaderWrapper( FClass* Class, bool
 			ListOfPublicClassesUObjectHeaderGroupIncludes.Logf(TEXT("#include \"%s\"\r\n"), *NewFileName);
 		}
 	}
-
+	ForwardDeclarations.Empty();
+	GeneratedHeaderTextBeforeForwardDeclarations.Empty();
+	GeneratedForwardDeclarations.Empty();
 	GeneratedHeaderText.Empty();
 	EnumForeachText.Empty();
 }
@@ -2968,6 +3012,7 @@ void FNativeClassHeaderGenerator::ExportProtoMessage(const TArray<UFunction*>& I
 				for(; CommentIt && (CommentIt->PropertyFlags & CPF_Parm); ++CommentIt)
 				{
 					UProperty* Param = *CommentIt;
+					ForwardDeclarations.Add(Param);
 					FString TypeText, ExtendedTypeText;
 					TypeText = Param->GetCPPType(&ExtendedTypeText, CPPF_None);
 					FString ParamName = FString::Printf(TEXT("%s%s %s"), *TypeText, *ExtendedTypeText, *Param->GetName());
@@ -3199,6 +3244,7 @@ void FNativeClassHeaderGenerator::ExportMCPMessage(const TArray<UFunction*>& InC
 				for(; CommentIt && (CommentIt->PropertyFlags & CPF_Parm); ++CommentIt)
 				{
 					UProperty* Param = *CommentIt;
+					ForwardDeclarations.Add(Param);
 					FString TypeText, ExtendedTypeText;
 					TypeText = Param->GetCPPType(&ExtendedTypeText, CPPF_None);
 					FString ParamName = FString::Printf(TEXT("%s%s %s"), *TypeText, *ExtendedTypeText, *Param->GetName());
@@ -3264,7 +3310,7 @@ void FNativeClassHeaderGenerator::ExportEventParms( const TArray<UFunction*>& In
 		for( TFieldIterator<UProperty> It(Function); It && (It->PropertyFlags&CPF_Parm); ++It )
 		{
 			UProperty* Prop = *It;
-
+			ForwardDeclarations.Add(Prop);
 			FStringOutputDevice PropertyText;
 			PropertyText.Log( FCString::Spc(4 + Indent) );
 
@@ -3423,6 +3469,7 @@ FString FNativeClassHeaderGenerator::GetFunctionReturnString(UFunction* Function
 	if (auto Return = Function->GetReturnProperty())
 	{
 		FString ExtendedReturnType;
+		ForwardDeclarations.Add(Return);
 		FString ReturnType = Return->GetCPPType(&ExtendedReturnType, CPPF_ArgumentOrReturnValue);
 		FStringOutputDevice ReplacementText(*ReturnType);
 		ApplyAlternatePropertyExportText(Return, ReplacementText);
@@ -3498,14 +3545,14 @@ void FNativeClassHeaderGenerator::ExportFunctionChecks(const FFuncInfo& Function
 	if (bNeedsImplementation)
 	{
 		FunctionName = FunctionData.CppImplName;
-		StaticChecks.Logf(TEXT("\t\tstatic_assert(THasMemberFunction_%s<%s>::Value, \"%s Declare function virtual %s %s(%s) %s in class %s.\");\r\n"), *FunctionName, *ClassName, *AssertMessage, *FunctionReturnType, *FunctionName, *ParameterString, *ConstModifier, *ClassName);
+		StaticChecks.Logf(TEXT("\t\tstatic_assert(THasMemberFunction_%s<%s>::Value, \"%s Declare function %s %s(%s) %s in class %s.\");\r\n"), *FunctionName, *ClassName, *AssertMessage, *FunctionReturnType, *FunctionName, *ParameterString, *ConstModifier, *ClassName);
 		CheckClasses.Logf(TEXT("GENERATE_MEMBER_FUNCTION_CHECK(%s, %s, %s, %s)\r\n"), *FunctionName, *FunctionReturnType, *ConstModifier, *ParameterString);
 	}
 	
 	if (bNeedsValidate)
 	{
 		FunctionName = FunctionData.CppValidationImplName;
-		StaticChecks.Logf(TEXT("\t\tstatic_assert(THasMemberFunction_%s<%s>::Value, \"%s Declare function virtual bool %s(%s) %s in class %s.\");\r\n"), *FunctionName, *ClassName, *AssertMessage, *FunctionReturnType, *FunctionName, *ParameterString, *ConstModifier, *ClassName);
+		StaticChecks.Logf(TEXT("\t\tstatic_assert(THasMemberFunction_%s<%s>::Value, \"%s Declare function bool %s(%s) %s in class %s.\");\r\n"), *FunctionName, *ClassName, *AssertMessage, *FunctionReturnType, *FunctionName, *ParameterString, *ConstModifier, *ClassName);
 		CheckClasses.Logf(TEXT("GENERATE_MEMBER_FUNCTION_CHECK(%s, bool, %s, %s)\r\n"), *FunctionName, *ConstModifier, *ParameterString);
 	}
 }
@@ -3560,6 +3607,7 @@ void FNativeClassHeaderGenerator::ExportNativeFunctionHeader( const FFuncInfo& F
 	{
 		FString ExtendedReturnType;
 		FString ReturnType = Return->GetCPPType(&ExtendedReturnType, (FunctionHeaderStyle == EExportFunctionHeaderStyle::Definition && (FunctionType != EExportFunctionType::Interface) ? CPPF_Implementation : 0) | CPPF_ArgumentOrReturnValue);
+		ForwardDeclarations.Add(Return);
 		FStringOutputDevice ReplacementText(*ReturnType);
 		ApplyAlternatePropertyExportText(Return, ReplacementText);
 		HeaderOutput.Logf(TEXT("%s%s"), *ReplacementText, *ExtendedReturnType);
@@ -3602,7 +3650,7 @@ void FNativeClassHeaderGenerator::ExportNativeFunctionHeader( const FFuncInfo& F
 	for( TFieldIterator<UProperty> It(Function); It && (It->PropertyFlags&(CPF_Parm|CPF_ReturnParm))==CPF_Parm; ++It )
 	{
 		UProperty* Property = *It;
-
+		ForwardDeclarations.Add(Property);
 		if( ParmCount++ )
 		{
 			HeaderOutput.Log(TEXT(", "));
@@ -3648,7 +3696,7 @@ void FNativeClassHeaderGenerator::ExportFunctionThunk(FStringOutputDevice& RPCWr
 	for (int32 ParameterIndex = 0; ParameterIndex < Parameters.Num(); ParameterIndex++)
 	{
 		UProperty* Param = Parameters[ParameterIndex];
-
+		ForwardDeclarations.Add(Param);
 		FString EvalBaseText = TEXT("P_GET_");	// e.g. P_GET_STR
 		FString EvalModifierText;				// e.g. _REF
 		FString EvalParameterText;				// e.g. (UObject*,NULL)
@@ -3736,11 +3784,53 @@ void FNativeClassHeaderGenerator::ExportFunctionThunk(FStringOutputDevice& RPCWr
 
 	RPCWrappers.Logf(TEXT("%sP_FINISH;%s"), FCString::Spc(8), LINE_TERMINATOR);
 
+	const auto DeprecationPushString = TEXT("PRAGMA_ENABLE_DEPRECATION_WARNINGS") LINE_TERMINATOR;
+	const auto DeprecationPopString = TEXT("PRAGMA_POP") LINE_TERMINATOR;
+
+	auto ClassRange = ClassDefinitionRange();
+	if (ClassDefinitionRanges.Contains(CurrentClass))
+	{
+		ClassRange = ClassDefinitionRanges[CurrentClass];
+	}
+
+	auto ClassStart = ClassRange.Start;
+	auto ClassEnd = ClassRange.End;
+	auto ClassDefinition = FString(ClassEnd - ClassStart, ClassStart);
+	auto ClassName = CurrentClass->GetName();
+
+	auto Function = FunctionData.FunctionReference;
+	auto FunctionName = Function->GetName();
+
+	bool bHasImplementation = ClassDefinition.Contains(FunctionName + TEXT("_Implementation"), ESearchCase::CaseSensitive);
+	bool bHasValidate = ClassDefinition.Contains(FunctionName + TEXT("_Validate"), ESearchCase::CaseSensitive);
+
+	auto bShouldEnableImplementationDeprecation =
+		// Enable deprecation warnings only if GENERATED_BODY is used inside class or interface (not GENERATED_UCLASS_BODY etc.)
+		ClassRange.bHasGeneratedBody
+		// And _Implementation function is called, but not the one declared by user
+		&& (FunctionData.CppImplName.EndsWith(TEXT("_Implementation")) && !bHasImplementation);
+
+	auto bShouldEnableValidateDeprecation =
+		// Enable deprecation warnings only if GENERATED_BODY is used inside class or interface (not GENERATED_UCLASS_BODY etc.)
+		ClassRange.bHasGeneratedBody
+		// or _Validate function is called, but not the one declared by user
+		&& FunctionData.CppImplName.EndsWith(TEXT("_Validate")) && !bHasValidate;
+
+	//Emit warning here if necessary
+	FStringOutputDevice FunctionDeclaration;
+	ExportNativeFunctionHeader(FunctionData, FunctionDeclaration, EExportFunctionType::Function, EExportFunctionHeaderStyle::Declaration);
+	FunctionDeclaration.Trim();
+
+
 	// Call the validate function if there is one
 	if (!(FunctionData.FunctionExportFlags & FUNCEXPORT_CppStatic) && (FunctionData.FunctionFlags & FUNC_NetValidate))
 	{
 		// Call the validate function, and if it fails, return (which will NOT execute the _Implementation function below)
 		// NOTE - We don't need to set "Result", since RPC calls don't return any values...
+		if (bShouldEnableValidateDeprecation)
+		{
+			RPCWrappers.Logf(TEXT("EMIT_CUSTOM_WARNING(\"Function %s::%s needs validation by %s due to its properties. Currently %s declaration is autogenerated by UHT. Since next release you'll have to provide declaration on your own. Please update your code before upgrading to the next release, otherwise your project will no longer compile.\")\r\n"), *ClassName, *FunctionName, *FunctionDeclaration, *FunctionData.CppValidationImplName);
+		}
 		RPCWrappers.Logf(TEXT("%sif (!this->%s(%s))%s"), FCString::Spc(8), *FunctionData.CppValidationImplName, *ParameterList, LINE_TERMINATOR);
 		RPCWrappers.Logf(TEXT("%s{%s"), FCString::Spc(8), LINE_TERMINATOR);
 		RPCWrappers.Logf(TEXT("%sRPC_ValidateFailed(TEXT(\"%s\"));%s"), FCString::Spc(8 + 4), *FunctionData.CppValidationImplName, LINE_TERMINATOR);
@@ -3748,11 +3838,16 @@ void FNativeClassHeaderGenerator::ExportFunctionThunk(FStringOutputDevice& RPCWr
 		RPCWrappers.Logf(TEXT("%s}%s"), FCString::Spc(8), LINE_TERMINATOR);
 	}
 
+	if (bShouldEnableImplementationDeprecation)
+	{
+		RPCWrappers.Logf(TEXT("EMIT_CUSTOM_WARNING(\"Function %s::%s needs native implementation by %s due to its properties. Currently %s declaration is autogenerated by UHT. Since next release you'll have to provide declaration on your own. Please update your code before upgrading to the next release, otherwise your project will no longer compile.\")\r\n"), *ClassName, *FunctionName, *FunctionDeclaration, *FunctionData.CppImplName);
+	}
 	// write out the return value
 	RPCWrappers.Log(FCString::Spc(8));
 	if (Return != NULL)
 	{
 		FString ReturnType, ReturnExtendedType;
+		ForwardDeclarations.Add(Return);
 		ReturnType = Return->GetCPPType(&ReturnExtendedType);
 		FStringOutputDevice ReplacementText(*ReturnType);
 		ApplyAlternatePropertyExportText(Return, ReplacementText);
@@ -3778,6 +3873,8 @@ FString FNativeClassHeaderGenerator::GetFunctionParameterString(UFunction* Funct
 
 	for (auto Property : TFieldRange<UProperty>(Function))
 	{
+		ForwardDeclarations.Add(Property);
+
 		if ((Property->PropertyFlags & (CPF_Parm | CPF_ReturnParm)) != CPF_Parm)
 		{
 			break;
@@ -3803,15 +3900,21 @@ void FNativeClassHeaderGenerator::ExportNativeFunctions(const TArray<UFunction*>
 {
 	FStringOutputDevice RPCWrappers;
 	FStringOutputDevice	AutogeneratedBlueprintFunctionDeclarations;
+	FStringOutputDevice	AutogeneratedBlueprintFunctionDeclarationsOnlyNotDeclared;
 	FStringOutputDevice	MemberFunctionCheckClasses;
 	FStringOutputDevice StaticChecks;
 	StaticChecks.Logf(TEXT("\tstatic inline void StaticChecks_Implementation_Validate()\r\n"));
 	StaticChecks.Logf(TEXT("\t{\r\n"));
-
+	
 	// find the class info for this class
 	FClassMetaData* ClassData = GScriptHelper.FindClassData(CurrentClass);
 
 	auto ClassName = FString(NameLookupCPP.GetNameCPP(CurrentClass));
+	auto ClassRange = ClassDefinitionRange();
+	if (ClassDefinitionRanges.Contains(CurrentClass))
+	{
+		ClassRange = ClassDefinitionRanges[CurrentClass];
+	}
 
 	// export the C++ stubs
 	for (auto Function : NativeFunctions)
@@ -3833,6 +3936,21 @@ void FNativeClassHeaderGenerator::ExportNativeFunctions(const TArray<UFunction*>
 
 		if (!bWillBeProgrammerTyped)
 		{
+			auto ClassStart = ClassRange.Start;
+			auto ClassEnd = ClassRange.End;
+			auto ClassDefinition = FString(ClassEnd - ClassStart, ClassStart);
+
+			auto Function = FunctionData.FunctionReference;
+			auto FunctionName = Function->GetName();
+
+			bool bHasImplementation = ClassDefinition.Contains(FunctionName + TEXT("_Implementation"), ESearchCase::CaseSensitive);
+			bool bHasValidate = ClassDefinition.Contains(FunctionName + TEXT("_Validate"), ESearchCase::CaseSensitive);
+
+			//Emit warning here if necessary
+			FStringOutputDevice FunctionDeclaration;
+			ExportNativeFunctionHeader(FunctionData, FunctionDeclaration, EExportFunctionType::Function, EExportFunctionHeaderStyle::Declaration);
+			FunctionDeclaration.Log(TEXT(";\r\n"));
+
 			// Declare validation function if needed
 			if (FunctionData.FunctionFlags & FUNC_NetValidate)
 			{
@@ -3841,14 +3959,19 @@ void FNativeClassHeaderGenerator::ExportNativeFunctions(const TArray<UFunction*>
 				// this is not an event, the function is not a static function and the function is not marked final
 				auto Virtual = (!FunctionData.FunctionReference->HasAnyFunctionFlags(FUNC_Static) && !(FunctionData.FunctionExportFlags & FUNCEXPORT_Final)) ? TEXT("virtual") : TEXT("");
 				AutogeneratedBlueprintFunctionDeclarations.Logf(TEXT("\t%s bool %s(%s);\r\n"), Virtual, *FunctionData.CppValidationImplName, *ParameterList);
+				if (!bHasValidate)
+				{
+					AutogeneratedBlueprintFunctionDeclarationsOnlyNotDeclared.Logf(TEXT("\t%s bool %s(%s);\r\n"), Virtual, *FunctionData.CppValidationImplName, *ParameterList);
+				}
 			}
 
-			if ((FunctionData.bCppImplNameEndsWith_Implementation || FunctionData.bCppValidationImplNameEndsWith_Validate))
+			AutogeneratedBlueprintFunctionDeclarations.Log(*FunctionDeclaration);
+			if ((FunctionData.bCppImplNameEndsWith_Implementation && !bHasImplementation) || (FunctionData.bCppValidationImplNameEndsWith_Validate && !bHasValidate))
 			{
-				ExportNativeFunctionHeader(FunctionData, AutogeneratedBlueprintFunctionDeclarations, EExportFunctionType::Function, EExportFunctionHeaderStyle::Declaration);
-				AutogeneratedBlueprintFunctionDeclarations.Log(TEXT(";\r\n"));
-				//ExportFunctionChecks(FunctionData, MemberFunctionCheckClasses, StaticChecks, ClassName);
+				AutogeneratedBlueprintFunctionDeclarationsOnlyNotDeclared.Log(*FunctionDeclaration);
 			}
+
+			ExportFunctionChecks(FunctionData, MemberFunctionCheckClasses, StaticChecks, ClassName);
 		}
 
 		if (bMultiLineUFUNCTION)
@@ -3884,8 +4007,9 @@ void FNativeClassHeaderGenerator::ExportNativeFunctions(const TArray<UFunction*>
 
 	auto NoPureDeclsSuffix = FString(TEXT("_RPC_WRAPPERS_NO_PURE_DECLS"));
 	MacroName = ClassName + NoPureDeclsSuffix;
+
 	// Put static checks before RPCWrappers to get proper messages from static asserts before compiler errors.
-	Macroized = Macroize(*MacroName, *(MemberFunctionCheckClasses + StaticChecks + RPCWrappers));
+	Macroized = Macroize(*MacroName, *(AutogeneratedBlueprintFunctionDeclarationsOnlyNotDeclared + MemberFunctionCheckClasses + StaticChecks + RPCWrappers));
 	GeneratedHeaderText.Log(*Macroized);
 	InClassNoPureDeclsMacroCalls.Logf(TEXT("%s%s\r\n"), FCString::Spc(4), *MacroName);
 }
diff --git a/Engine/Source/Programs/UnrealHeaderTool/Private/HeaderParser.cpp b/Engine/Source/Programs/UnrealHeaderTool/Private/HeaderParser.cpp
--- a/Engine/Source/Programs/UnrealHeaderTool/Private/HeaderParser.cpp
+++ b/Engine/Source/Programs/UnrealHeaderTool/Private/HeaderParser.cpp
@@ -22,7 +22,7 @@ static FUObjectAnnotationSparseBool FailedClassesAnnotation;
 enum {MAX_ARRAY_SIZE=2048};
 
 static const FName NAME_ToolTip(TEXT("ToolTip"));
-
+TMap<FClass*, ClassDefinitionRange> ClassDefinitionRanges;
 /**
  * Dirty hack global variable to allow different result codes passed through
  * exceptions. Needs to be fixed in future versions of UHT.
@@ -3840,6 +3840,12 @@ bool FHeaderParser::CompileDeclaration( FClasses& AllClasses, FToken& Token )
 		{
 			FError::Throwf(TEXT("%s must occur inside the class or interface definition"), Token.Identifier);
 		}
+		if (!ClassDefinitionRanges.Contains(Class))
+		{
+			ClassDefinitionRanges.Add(Class, ClassDefinitionRange());
+		}
+
+		ClassDefinitionRanges[Class].bHasGeneratedBody = true;
 
 		RequireSymbol(TEXT("("), Token.Identifier);
 		RequireSymbol(TEXT(")"), Token.Identifier);
@@ -3942,6 +3948,10 @@ bool FHeaderParser::CompileDeclaration( FClasses& AllClasses, FToken& Token )
 	}
 	else if (bEncounteredNewStyleClass_UnmatchedBrackets && Token.Matches(TEXT("}")))
 	{
+		if (ClassDefinitionRanges.Contains(Class))
+		{
+			ClassDefinitionRanges[Class].End = &Input[InputPos];
+		}
 		MatchSemi();
 
 		// Closing brace for class declaration
@@ -4277,7 +4287,7 @@ void FHeaderParser::CompileClassDeclaration(FClasses& AllClasses)
 	FString DeclaredClassName;
 	FString RequiredAPIMacroIfPresent;
 	ParseClassNameDeclaration(AllClasses, /*out*/ DeclaredClassName, /*out*/ RequiredAPIMacroIfPresent);
-
+	ClassDefinitionRanges.Add(Class, ClassDefinitionRange(&Input[InputPos], nullptr));
 	// Record that this class is RequiredAPI if the CORE_API style macro was present
 	if (!RequiredAPIMacroIfPresent.IsEmpty())
 	{
diff --git a/Engine/Source/Programs/UnrealHeaderTool/Private/HeaderParser.h b/Engine/Source/Programs/UnrealHeaderTool/Private/HeaderParser.h
--- a/Engine/Source/Programs/UnrealHeaderTool/Private/HeaderParser.h
+++ b/Engine/Source/Programs/UnrealHeaderTool/Private/HeaderParser.h
@@ -63,6 +63,27 @@ struct FIndexRange
 	int32 Count;
 };
 
+
+struct ClassDefinitionRange
+{
+	ClassDefinitionRange(const TCHAR* InStart, const TCHAR* InEnd)
+		: Start(InStart)
+		, End(InEnd)
+		, bHasGeneratedBody(false)
+	{ }
+
+	ClassDefinitionRange()
+		: Start(nullptr)
+		, End(nullptr)
+		, bHasGeneratedBody(false)
+	{ }
+
+	const TCHAR* Start;
+	const TCHAR* End;
+	bool bHasGeneratedBody;
+};
+
+extern TMap<FClass*, ClassDefinitionRange> ClassDefinitionRanges;
 /////////////////////////////////////////////////////
 // FHeaderParser
 
@@ -548,6 +569,7 @@ private:
 	 */
 	bool TryToMatchConstructorParameterList(FToken Token);
 	void SkipDeprecatedMacroIfNecessary();
+
 };
 
 /////////////////////////////////////////////////////
diff --git a/Engine/Source/Programs/UnrealHeaderTool/Private/NativeClassExporter.h b/Engine/Source/Programs/UnrealHeaderTool/Private/NativeClassExporter.h
--- a/Engine/Source/Programs/UnrealHeaderTool/Private/NativeClassExporter.h
+++ b/Engine/Source/Programs/UnrealHeaderTool/Private/NativeClassExporter.h
@@ -35,7 +35,6 @@ class FModuleClasses;
 struct FNativeClassHeaderGenerator
 {
 private:
-
 	FClass*				CurrentClass;
 	TArray<FClass*>		Classes;
 	FString				API;
@@ -55,6 +54,8 @@ private:
 	FStringOutputDevice FriendText;
 	FStringOutputDevice	GeneratedPackageCPP;
 	FStringOutputDevice	GeneratedHeaderText;
+	FStringOutputDevice	GeneratedHeaderTextBeforeForwardDeclarations;
+	FStringOutputDevice	GeneratedForwardDeclarations;
     /** output generated for a .proto file */
 	FStringOutputDevice GeneratedProtoText;
 	/** output generated for a mcp .java file */
@@ -102,6 +103,9 @@ private:
 	/** If true, any change in the generated headers will result in UHT failure. */
 	bool bFailIfGeneratedCodeChanges;
 
+	/** All properties that need to be forward declared. */
+	TArray<UProperty*>	ForwardDeclarations;
+
 	// This exists because it makes debugging much easier on VC2010, since the visualizers can't properly understand templates with templated args
 	struct HeaderDependents : TArray<const FString*>
 	{
diff --git a/Engine/Source/Runtime/Core/Public/Android/AndroidPlatformCompilerPreSetup.h b/Engine/Source/Runtime/Core/Public/Android/AndroidPlatformCompilerPreSetup.h
--- a/Engine/Source/Runtime/Core/Public/Android/AndroidPlatformCompilerPreSetup.h
+++ b/Engine/Source/Runtime/Core/Public/Android/AndroidPlatformCompilerPreSetup.h
@@ -32,6 +32,15 @@
 		_Pragma ("clang diagnostic ignored \"-Wdeprecated-declarations\"")
 
 	#define PRAGMA_ENABLE_DEPRECATION_WARNINGS \
+		_Pragma ("clang diagnostic push") \
+		_Pragma ("clang diagnostic warning \"-Wdeprecated-declarations\"")
+
+	#define PRAGMA_POP \
 		_Pragma("clang diagnostic pop")
 
 #endif // DISABLE_DEPRECATION
+
+#ifndef EMIT_CUSTOM_WARNING
+#define EMIT_CUSTOM_WARNING(Warning) \
+	_Pragma(PREPROCESSOR_TO_STRING(message(WARNING_LOCATION Warning)))
+#endif // EMIT_CUSTOM_WARNING
\ No newline at end of file
diff --git a/Engine/Source/Runtime/Core/Public/GenericPlatform/GenericPlatformCompilerPreSetup.h b/Engine/Source/Runtime/Core/Public/GenericPlatform/GenericPlatformCompilerPreSetup.h
--- a/Engine/Source/Runtime/Core/Public/GenericPlatform/GenericPlatformCompilerPreSetup.h
+++ b/Engine/Source/Runtime/Core/Public/GenericPlatform/GenericPlatformCompilerPreSetup.h
@@ -37,3 +37,7 @@
 #ifndef EMIT_DEPRECATED_WARNING_MESSAGE
 	#define EMIT_DEPRECATED_WARNING_MESSAGE(Msg)
 #endif // EMIT_DEPRECATED_WARNING_MESSAGE
+
+#ifndef EMIT_CUSTOM_WARNING
+	#define EMIT_CUSTOM_WARNING(Warning)
+#endif // EMIT_CUSTOM_WARNING
\ No newline at end of file
diff --git a/Engine/Source/Runtime/Core/Public/HTML5/HTML5PlatformCompilerPreSetup.h b/Engine/Source/Runtime/Core/Public/HTML5/HTML5PlatformCompilerPreSetup.h
--- a/Engine/Source/Runtime/Core/Public/HTML5/HTML5PlatformCompilerPreSetup.h
+++ b/Engine/Source/Runtime/Core/Public/HTML5/HTML5PlatformCompilerPreSetup.h
@@ -34,7 +34,15 @@
 		_Pragma ("clang diagnostic ignored \"-Wdeprecated-declarations\"")
 
 	#define PRAGMA_ENABLE_DEPRECATION_WARNINGS \
-		_Pragma("clang diagnostic pop")
+		_Pragma ("clang diagnostic push") \
+		_Pragma ("clang diagnostic warning \"-Wdeprecated-declarations\"")
 
+	#define PRAGMA_POP \
+		_Pragma("clang diagnostic pop")
 #endif 
 #endif // DISABLE_DEPRECATION
+
+#ifndef EMIT_CUSTOM_WARNING
+#define EMIT_CUSTOM_WARNING(Warning) \
+	_Pragma(PREPROCESSOR_TO_STRING(message(WARNING_LOCATION Warning)))
+#endif // EMIT_CUSTOM_WARNING
diff --git a/Engine/Source/Runtime/Core/Public/IOS/IOSPlatformCompilerPreSetup.h b/Engine/Source/Runtime/Core/Public/IOS/IOSPlatformCompilerPreSetup.h
--- a/Engine/Source/Runtime/Core/Public/IOS/IOSPlatformCompilerPreSetup.h
+++ b/Engine/Source/Runtime/Core/Public/IOS/IOSPlatformCompilerPreSetup.h
@@ -32,6 +32,14 @@
 		_Pragma ("clang diagnostic ignored \"-Wdeprecated-declarations\"")
 
 	#define PRAGMA_ENABLE_DEPRECATION_WARNINGS \
-		_Pragma("clang diagnostic pop")
+		_Pragma ("clang diagnostic push") \
+		_Pragma ("clang diagnostic warning \"-Wdeprecated-declarations\"")
 
+	#define PRAGMA_POP \
+		_Pragma("clang diagnostic pop")
 #endif // DISABLE_DEPRECATION
+
+#ifndef EMIT_CUSTOM_WARNING
+#define EMIT_CUSTOM_WARNING(Warning) \
+	_Pragma(PREPROCESSOR_TO_STRING(message(WARNING_LOCATION Warning)))
+#endif // EMIT_CUSTOM_WARNING
diff --git a/Engine/Source/Runtime/Core/Public/Linux/LinuxPlatformCompilerPreSetup.h b/Engine/Source/Runtime/Core/Public/Linux/LinuxPlatformCompilerPreSetup.h
--- a/Engine/Source/Runtime/Core/Public/Linux/LinuxPlatformCompilerPreSetup.h
+++ b/Engine/Source/Runtime/Core/Public/Linux/LinuxPlatformCompilerPreSetup.h
@@ -32,6 +32,14 @@
 		_Pragma ("clang diagnostic ignored \"-Wdeprecated-declarations\"")
 
 	#define PRAGMA_ENABLE_DEPRECATION_WARNINGS \
-		_Pragma("clang diagnostic pop")
+		_Pragma ("clang diagnostic push") \
+		_Pragma ("clang diagnostic warning \"-Wdeprecated-declarations\"")
 
+	#define PRAGMA_POP \
+		_Pragma("clang diagnostic pop")
 #endif // DISABLE_DEPRECATION
+
+#ifndef EMIT_CUSTOM_WARNING
+#define EMIT_CUSTOM_WARNING(Warning) \
+	_Pragma(PREPROCESSOR_TO_STRING(message(WARNING_LOCATION Warning)))
+#endif // EMIT_CUSTOM_WARNING
diff --git a/Engine/Source/Runtime/Core/Public/Mac/MacPlatformCompilerPreSetup.h b/Engine/Source/Runtime/Core/Public/Mac/MacPlatformCompilerPreSetup.h
--- a/Engine/Source/Runtime/Core/Public/Mac/MacPlatformCompilerPreSetup.h
+++ b/Engine/Source/Runtime/Core/Public/Mac/MacPlatformCompilerPreSetup.h
@@ -32,6 +32,14 @@
 		_Pragma ("clang diagnostic ignored \"-Wdeprecated-declarations\"")
 
 	#define PRAGMA_ENABLE_DEPRECATION_WARNINGS \
-		_Pragma("clang diagnostic pop")
+		_Pragma ("clang diagnostic push") \
+		_Pragma ("clang diagnostic warning \"-Wdeprecated-declarations\"")
 
+	#define PRAGMA_POP \
+		_Pragma("clang diagnostic pop")
 #endif // DISABLE_DEPRECATION
+
+#ifndef EMIT_CUSTOM_WARNING
+#define EMIT_CUSTOM_WARNING(Warning) \
+	_Pragma(PREPROCESSOR_TO_STRING(message(WARNING_LOCATION Warning)))
+#endif // EMIT_CUSTOM_WARNING
diff --git a/Engine/Source/Runtime/Core/Public/Misc/CoreMiscDefines.h b/Engine/Source/Runtime/Core/Public/Misc/CoreMiscDefines.h
--- a/Engine/Source/Runtime/Core/Public/Misc/CoreMiscDefines.h
+++ b/Engine/Source/Runtime/Core/Public/Misc/CoreMiscDefines.h
@@ -108,3 +108,6 @@ const FPlatformUserId PLATFORMUSERID_NONE = INDEX_NONE;
 
 #define PREPROCESSOR_JOIN_INNER(x, y) x##y
 #define PREPROCESSOR_JOIN(x, y) PREPROCESSOR_JOIN_INNER(x, y)
+
+// When passed to pragma message will result in clickable warning in VS
+#define WARNING_LOCATION __FILE__ "(" PREPROCESSOR_TO_STRING(__LINE__) ")"
diff --git a/Engine/Source/Runtime/Core/Public/WinRT/WinRTPlatformCompilerPreSetup.h b/Engine/Source/Runtime/Core/Public/WinRT/WinRTPlatformCompilerPreSetup.h
--- a/Engine/Source/Runtime/Core/Public/WinRT/WinRTPlatformCompilerPreSetup.h
+++ b/Engine/Source/Runtime/Core/Public/WinRT/WinRTPlatformCompilerPreSetup.h
@@ -37,7 +37,22 @@
 		__pragma (warning(disable:4995)) \
 		__pragma (warning(disable:4996))
 
-	#define PRAGMA_ENABLE_DEPRECATION_WARNINGS \
+	#define PRAGMA_POP \
 		__pragma(warning(pop))
 
+	#define PRAGMA_ENABLE_DEPRECATION_WARNINGS \
+		__pragma (warning(push)) \
+		__pragma (warning(default:4995)) \
+		__pragma (warning(default:4996))
+
 #endif // DISABLE_DEPRECATION
+
+#if _MSC_VER
+#define EMIT_CUSTOM_WARNING(Warning) \
+	__pragma(message(WARNING_LOCATION "(\": warning C4996: " Warning))
+#elif __clang__
+#define EMIT_CUSTOM_WARNING(Warning) \
+	_Pragma(PREPROCESSOR_TO_STRING(message(WARNING_LOCATION Warning)))
+#else
+#error Unknown compiler
+#endif
diff --git a/Engine/Source/Runtime/Core/Public/Windows/WindowsPlatformCompilerPreSetup.h b/Engine/Source/Runtime/Core/Public/Windows/WindowsPlatformCompilerPreSetup.h
--- a/Engine/Source/Runtime/Core/Public/Windows/WindowsPlatformCompilerPreSetup.h
+++ b/Engine/Source/Runtime/Core/Public/Windows/WindowsPlatformCompilerPreSetup.h
@@ -37,7 +37,22 @@
 		__pragma (warning(disable:4995)) \
 		__pragma (warning(disable:4996))
 
-	#define PRAGMA_ENABLE_DEPRECATION_WARNINGS \
+	#define PRAGMA_POP \
 		__pragma(warning(pop))
 
+	#define PRAGMA_ENABLE_DEPRECATION_WARNINGS \
+		__pragma (warning(push)) \
+		__pragma (warning(default:4995)) \
+		__pragma (warning(default:4996))
+
 #endif // DISABLE_DEPRECATION
+
+#if _MSC_VER
+#define EMIT_CUSTOM_WARNING(Warning) \
+	__pragma(message(WARNING_LOCATION "(: warning C4996: " Warning))
+#elif __clang__
+#define EMIT_CUSTOM_WARNING(Warning) \
+	_Pragma(PREPROCESSOR_TO_STRING(message(WARNING_LOCATION Warning)))
+#else
+#error Unknown compiler
+#endif
diff --git a/Engine/Source/Runtime/CoreUObject/Private/UObject/PropertyArray.cpp b/Engine/Source/Runtime/CoreUObject/Private/UObject/PropertyArray.cpp
--- a/Engine/Source/Runtime/CoreUObject/Private/UObject/PropertyArray.cpp
+++ b/Engine/Source/Runtime/CoreUObject/Private/UObject/PropertyArray.cpp
@@ -118,6 +118,12 @@ FString UArrayProperty::GetCPPType( FString* ExtendedTypeText/*=NULL*/, uint32 C
 	}
 	return TEXT("TArray");
 }
+
+FString UArrayProperty::GetCPPTypeForwardDeclaration() const
+{
+	checkSlow(Inner);
+	return Inner->GetCPPTypeForwardDeclaration();
+}
 FString UArrayProperty::GetCPPMacroType( FString& ExtendedTypeText ) const
 {
 	checkSlow(Inner);
diff --git a/Engine/Source/Runtime/CoreUObject/Private/UObject/PropertyAssetClass.cpp b/Engine/Source/Runtime/CoreUObject/Private/UObject/PropertyAssetClass.cpp
--- a/Engine/Source/Runtime/CoreUObject/Private/UObject/PropertyAssetClass.cpp
+++ b/Engine/Source/Runtime/CoreUObject/Private/UObject/PropertyAssetClass.cpp
@@ -8,7 +8,7 @@
 
 FString UAssetClassProperty::GetCPPType( FString* ExtendedTypeText/*=NULL*/, uint32 CPPExportFlags/*=0*/ ) const
 {
-	return FString::Printf(TEXT("TAssetSubclassOf<class %s%s> "),MetaClass->GetPrefixCPP(),*MetaClass->GetName());
+	return FString::Printf(TEXT("TAssetSubclassOf<%s%s> "),MetaClass->GetPrefixCPP(),*MetaClass->GetName());
 }
 FString UAssetClassProperty::GetCPPMacroType( FString& ExtendedTypeText ) const
 {
diff --git a/Engine/Source/Runtime/CoreUObject/Private/UObject/PropertyAssetObject.cpp b/Engine/Source/Runtime/CoreUObject/Private/UObject/PropertyAssetObject.cpp
--- a/Engine/Source/Runtime/CoreUObject/Private/UObject/PropertyAssetObject.cpp
+++ b/Engine/Source/Runtime/CoreUObject/Private/UObject/PropertyAssetObject.cpp
@@ -8,7 +8,7 @@
 
 FString UAssetObjectProperty::GetCPPType( FString* ExtendedTypeText/*=NULL*/, uint32 CPPExportFlags/*=0*/ ) const
 {
-	return FString::Printf( TEXT("TAssetPtr<class %s%s>"), PropertyClass->GetPrefixCPP(), *PropertyClass->GetName() );
+	return FString::Printf( TEXT("TAssetPtr<%s%s>"), PropertyClass->GetPrefixCPP(), *PropertyClass->GetName() );
 }
 FString UAssetObjectProperty::GetCPPMacroType( FString& ExtendedTypeText ) const
 {
diff --git a/Engine/Source/Runtime/CoreUObject/Private/UObject/PropertyBool.cpp b/Engine/Source/Runtime/CoreUObject/Private/UObject/PropertyBool.cpp
--- a/Engine/Source/Runtime/CoreUObject/Private/UObject/PropertyBool.cpp
+++ b/Engine/Source/Runtime/CoreUObject/Private/UObject/PropertyBool.cpp
@@ -142,6 +142,12 @@ FString UBoolProperty::GetCPPType( FString* ExtendedTypeText/*=NULL*/, uint32 CP
 	}
 	return TEXT("uint32");
 }
+
+FString UBoolProperty::GetCPPTypeForwardDeclaration() const
+{
+	return FString();
+}
+
 FString UBoolProperty::GetCPPMacroType( FString& ExtendedTypeText ) const
 {
 	check(FieldSize != 0);
diff --git a/Engine/Source/Runtime/CoreUObject/Private/UObject/PropertyClass.cpp b/Engine/Source/Runtime/CoreUObject/Private/UObject/PropertyClass.cpp
--- a/Engine/Source/Runtime/CoreUObject/Private/UObject/PropertyClass.cpp
+++ b/Engine/Source/Runtime/CoreUObject/Private/UObject/PropertyClass.cpp
@@ -50,9 +50,14 @@ const TCHAR* UClassProperty::ImportText_Internal( const TCHAR* Buffer, void* Dat
 
 FString UClassProperty::GetCPPType( FString* ExtendedTypeText/*=NULL*/, uint32 CPPExportFlags/*=0*/ ) const
 {
-//	return FString::Printf( TEXT("class %s%s*"), PropertyClass->GetPrefixCPP(), *PropertyClass->GetName() );
-	return FString::Printf(TEXT("TSubclassOf<class %s%s> "),MetaClass->GetPrefixCPP(),*MetaClass->GetName());
+	return FString::Printf(TEXT("TSubclassOf<%s%s> "),MetaClass->GetPrefixCPP(),*MetaClass->GetName());
 }
+
+FString UClassProperty::GetCPPTypeForwardDeclaration() const
+{
+	return FString::Printf(TEXT("class %s%s;"), MetaClass->GetPrefixCPP(), *MetaClass->GetName());
+}
+
 FString UClassProperty::GetCPPMacroType( FString& ExtendedTypeText ) const
 {
 	ExtendedTypeText = TEXT("UClass");
diff --git a/Engine/Source/Runtime/CoreUObject/Private/UObject/PropertyDelegate.cpp b/Engine/Source/Runtime/CoreUObject/Private/UObject/PropertyDelegate.cpp
--- a/Engine/Source/Runtime/CoreUObject/Private/UObject/PropertyDelegate.cpp
+++ b/Engine/Source/Runtime/CoreUObject/Private/UObject/PropertyDelegate.cpp
@@ -81,6 +81,10 @@ FString UDelegateProperty::GetCPPType( FString* ExtendedTypeText/*=NULL*/, uint3
 	return FString(TEXT("F")) + UnmangledFunctionName;
 }
 
+FString UDelegateProperty::GetCPPTypeForwardDeclaration() const
+{
+	return FString();
+}
 
 void UDelegateProperty::ExportTextItem( FString& ValueStr, const void* PropertyValue, const void* DefaultValue, UObject* Parent, int32 PortFlags, UObject* ExportRootScope ) const
 {
diff --git a/Engine/Source/Runtime/CoreUObject/Private/UObject/PropertyInterface.cpp b/Engine/Source/Runtime/CoreUObject/Private/UObject/PropertyInterface.cpp
--- a/Engine/Source/Runtime/CoreUObject/Private/UObject/PropertyInterface.cpp
+++ b/Engine/Source/Runtime/CoreUObject/Private/UObject/PropertyInterface.cpp
@@ -47,12 +47,26 @@ FString UInterfaceProperty::GetCPPType( FString* ExtendedTypeText/*=NULL*/, uint
 		check(ExportClass);
 		check(ExportClass->HasAnyClassFlags(CLASS_Interface));
 
-		*ExtendedTypeText = FString::Printf(TEXT("<class I%s>"), *ExportClass->GetName());
+		*ExtendedTypeText = FString::Printf(TEXT("<I%s>"), *ExportClass->GetName());
 	}
 
 	return TEXT("TScriptInterface");
 }
 
+FString UInterfaceProperty::GetCPPTypeForwardDeclaration() const
+{
+	checkSlow(InterfaceClass);
+	UClass* ExportClass = InterfaceClass;
+	while (ExportClass && !ExportClass->HasAnyClassFlags(CLASS_Native))
+	{
+		ExportClass = ExportClass->GetSuperClass();
+	}
+	check(ExportClass);
+	check(ExportClass->HasAnyClassFlags(CLASS_Interface));
+
+	return FString::Printf(TEXT("class I%s;"), *ExportClass->GetName());
+}
+
 void UInterfaceProperty::LinkInternal(FArchive& Ar)
 {
 	// for now, we won't support instancing of interface properties...it might be possible, but for the first pass we'll keep it simple
diff --git a/Engine/Source/Runtime/CoreUObject/Private/UObject/PropertyLazyObjectPtr.cpp b/Engine/Source/Runtime/CoreUObject/Private/UObject/PropertyLazyObjectPtr.cpp
--- a/Engine/Source/Runtime/CoreUObject/Private/UObject/PropertyLazyObjectPtr.cpp
+++ b/Engine/Source/Runtime/CoreUObject/Private/UObject/PropertyLazyObjectPtr.cpp
@@ -9,7 +9,7 @@
 
 FString ULazyObjectProperty::GetCPPType( FString* ExtendedTypeText/*=NULL*/, uint32 CPPExportFlags/*=0*/ ) const
 {
-	return FString::Printf( TEXT("TLazyObjectPtr<class %s%s>"), PropertyClass->GetPrefixCPP(), *PropertyClass->GetName() );
+	return FString::Printf( TEXT("TLazyObjectPtr<%s%s>"), PropertyClass->GetPrefixCPP(), *PropertyClass->GetName() );
 }
 FString ULazyObjectProperty::GetCPPMacroType( FString& ExtendedTypeText ) const
 {
diff --git a/Engine/Source/Runtime/CoreUObject/Private/UObject/PropertyObject.cpp b/Engine/Source/Runtime/CoreUObject/Private/UObject/PropertyObject.cpp
--- a/Engine/Source/Runtime/CoreUObject/Private/UObject/PropertyObject.cpp
+++ b/Engine/Source/Runtime/CoreUObject/Private/UObject/PropertyObject.cpp
@@ -8,8 +8,14 @@
 
 FString UObjectProperty::GetCPPType( FString* ExtendedTypeText/*=NULL*/, uint32 CPPExportFlags/*=0*/ ) const
 {
-	return FString::Printf( TEXT("class %s%s*"), PropertyClass->GetPrefixCPP(), *PropertyClass->GetName() );
+	return FString::Printf( TEXT("%s%s*"), PropertyClass->GetPrefixCPP(), *PropertyClass->GetName() );
 }
+
+FString UObjectProperty::GetCPPTypeForwardDeclaration() const
+{
+	return FString::Printf(TEXT("class %s%s;"), PropertyClass->GetPrefixCPP(), *PropertyClass->GetName());
+}
+
 FString UObjectProperty::GetCPPMacroType( FString& ExtendedTypeText ) const
 {
 	ExtendedTypeText = FString::Printf(TEXT("%s%s"), PropertyClass->GetPrefixCPP(), *PropertyClass->GetName());
diff --git a/Engine/Source/Runtime/CoreUObject/Private/UObject/PropertyStruct.cpp b/Engine/Source/Runtime/CoreUObject/Private/UObject/PropertyStruct.cpp
--- a/Engine/Source/Runtime/CoreUObject/Private/UObject/PropertyStruct.cpp
+++ b/Engine/Source/Runtime/CoreUObject/Private/UObject/PropertyStruct.cpp
@@ -215,15 +215,39 @@ bool UStructProperty::HasNoOpConstructor() const
 
 FString UStructProperty::GetCPPType( FString* ExtendedTypeText/*=NULL*/, uint32 CPPExportFlags/*=0*/ ) const
 {
-	bool bExportForwardDeclaration = 
-		(CPPExportFlags&CPPF_OptionalValue) == 0 &&
-		(Struct->GetOwnerClass() == NULL || 
-		!(Struct->GetOwnerClass()->HasAnyClassFlags(CLASS_NoExport) || (Struct->StructFlags&STRUCT_Native) == 0));
-
-	return FString::Printf( TEXT("%sF%s"), 
-		bExportForwardDeclaration ? TEXT("struct ") : TEXT(""),
-		*Struct->GetName() );
+	return FString::Printf(TEXT("F%s"), *Struct->GetName());
 }
+
+FString UStructProperty::GetCPPTypeForwardDeclaration() const
+{
+	// The following types are classes, but reflection system sees them as structs.
+	static FName ClassNames[] =
+	{
+		NAME_Box,
+		NAME_Color,
+		FName(TEXT("Guid")),
+		NAME_Plane,
+		NAME_Quat,
+		FName(TEXT("RandomStream")),
+		FName(TEXT("Rotator")),
+		FName(TEXT("Timespan")),
+		NAME_Transform,
+		NAME_Vector
+	};
+
+	auto Name = Struct->GetFName();
+
+	for (int32 i = 0; i < ARRAY_COUNT(ClassNames); ++i)
+	{
+		if (ClassNames[i] == Name)
+		{
+			return FString::Printf(TEXT("class F%s;"), *Struct->GetName());
+		}
+	}
+
+	return FString::Printf(TEXT("struct F%s;"), *Struct->GetName());
+}
+
 FString UStructProperty::GetCPPMacroType( FString& ExtendedTypeText ) const
 {
 	ExtendedTypeText = GetCPPType(NULL, CPPF_None);
diff --git a/Engine/Source/Runtime/CoreUObject/Private/UObject/PropertyWeakObjectPtr.cpp b/Engine/Source/Runtime/CoreUObject/Private/UObject/PropertyWeakObjectPtr.cpp
--- a/Engine/Source/Runtime/CoreUObject/Private/UObject/PropertyWeakObjectPtr.cpp
+++ b/Engine/Source/Runtime/CoreUObject/Private/UObject/PropertyWeakObjectPtr.cpp
@@ -10,9 +10,14 @@ FString UWeakObjectProperty::GetCPPType( FString* ExtendedTypeText/*=NULL*/, uin
 {
 	if (PropertyFlags & CPF_AutoWeak)
 	{
-		return FString::Printf( TEXT("TAutoWeakObjectPtr<class %s%s>"), PropertyClass->GetPrefixCPP(), *PropertyClass->GetName() );
+		return FString::Printf( TEXT("TAutoWeakObjectPtr<%s%s>"), PropertyClass->GetPrefixCPP(), *PropertyClass->GetName() );
 	}
-	return FString::Printf( TEXT("TWeakObjectPtr<class %s%s>"), PropertyClass->GetPrefixCPP(), *PropertyClass->GetName() );
+	return FString::Printf( TEXT("TWeakObjectPtr<%s%s>"), PropertyClass->GetPrefixCPP(), *PropertyClass->GetName() );
+}
+
+FString UWeakObjectProperty::GetCPPTypeForwardDeclaration() const
+{
+	return FString::Printf(TEXT("class %s%s;"), PropertyClass->GetPrefixCPP(), *PropertyClass->GetName());
 }
 
 FString UWeakObjectProperty::GetCPPMacroType( FString& ExtendedTypeText ) const
diff --git a/Engine/Source/Runtime/CoreUObject/Public/UObject/UTextProperty.h b/Engine/Source/Runtime/CoreUObject/Public/UObject/UTextProperty.h
--- a/Engine/Source/Runtime/CoreUObject/Public/UObject/UTextProperty.h
+++ b/Engine/Source/Runtime/CoreUObject/Public/UObject/UTextProperty.h
@@ -25,5 +25,9 @@ public:
 	virtual void SerializeItem( FArchive& Ar, void* Value, int32 MaxReadBytes, void const* Defaults ) const override;
 	virtual void ExportTextItem( FString& ValueStr, const void* PropertyValue, const void* DefaultValue, UObject* Parent, int32 PortFlags, UObject* ExportRootScope ) const override;
 	virtual const TCHAR* ImportText_Internal( const TCHAR* Buffer, void* Data, int32 PortFlags, UObject* OwnerObject, FOutputDevice* ErrorText ) const override;
+	virtual FString GetCPPTypeForwardDeclaration() const override
+	{
+		return FString();
+	}
 	// End of UProperty interface
 };
\ No newline at end of file
diff --git a/Engine/Source/Runtime/CoreUObject/Public/UObject/UnrealType.h b/Engine/Source/Runtime/CoreUObject/Public/UObject/UnrealType.h
--- a/Engine/Source/Runtime/CoreUObject/Public/UObject/UnrealType.h
+++ b/Engine/Source/Runtime/CoreUObject/Public/UObject/UnrealType.h
@@ -201,6 +201,8 @@ public:
 	 * @param	CPPExportFlags		flags for modifying the behavior of the export
 	 */
 	virtual FString GetCPPType( FString* ExtendedTypeText=NULL, uint32 CPPExportFlags=0 ) const PURE_VIRTUAL(UProperty::GetCPPType,return TEXT(""););
+
+	virtual FString GetCPPTypeForwardDeclaration() const PURE_VIRTUAL(UProperty::GetCPPTypeForwardDeclaration, return TEXT(""););
 	// End of UHT interface
 
 private:
@@ -1205,6 +1207,11 @@ public:
 	{
 	}
 
+	FString GetCPPTypeForwardDeclaration() const override
+	{
+		return FString();
+	}
+
 	// UNumericProperty interface.
 
 	virtual bool IsFloatingPoint() const override
@@ -1506,6 +1513,7 @@ public:
 	// UHT interface
 	virtual FString GetCPPType( FString* ExtendedTypeText, uint32 CPPExportFlags ) const override;
 	virtual FString GetCPPMacroType( FString& ExtendedTypeText ) const override;
+	virtual FString GetCPPTypeForwardDeclaration() const override;
 	// End of UHT interface
 
 	// UProperty interface.
@@ -1766,6 +1774,7 @@ class COREUOBJECT_API UObjectProperty : public TUObjectPropertyBase<UObject*>
 	// UHT interface
 	virtual FString GetCPPMacroType( FString& ExtendedTypeText ) const  override;
 	virtual FString GetCPPType( FString* ExtendedTypeText, uint32 CPPExportFlags ) const override;
+	virtual FString GetCPPTypeForwardDeclaration() const override;
 	// End of UHT interface
 
 	// UProperty interface
@@ -1799,6 +1808,7 @@ class COREUOBJECT_API UWeakObjectProperty : public TUObjectPropertyBase<FWeakObj
 	// UHT interface
 	virtual FString GetCPPMacroType( FString& ExtendedTypeText ) const  override;
 	virtual FString GetCPPType( FString* ExtendedTypeText, uint32 CPPExportFlags ) const override;
+	virtual FString GetCPPTypeForwardDeclaration() const override;
 	// End of UHT interface
 
 	// UProperty interface
@@ -1924,6 +1934,7 @@ public:
 	// UHT interface
 	virtual FString GetCPPMacroType( FString& ExtendedTypeText ) const  override;
 	virtual FString GetCPPType( FString* ExtendedTypeText, uint32 CPPExportFlags ) const override;
+	virtual FString GetCPPTypeForwardDeclaration() const override;
 	// End of UHT interface
 
 	// UProperty interface
@@ -2000,6 +2011,7 @@ public:
 	// UHT interface
 	virtual FString GetCPPMacroType( FString& ExtendedTypeText ) const  override;
 	virtual FString GetCPPType( FString* ExtendedTypeText, uint32 CPPExportFlags ) const override;
+	virtual FString GetCPPTypeForwardDeclaration() const override;
 	// End of UHT interface
 
 	// UProperty interface
@@ -2045,6 +2057,10 @@ public:
 	// UProperty interface
 	virtual void ExportTextItem( FString& ValueStr, const void* PropertyValue, const void* DefaultValue, UObject* Parent, int32 PortFlags, UObject* ExportRootScope ) const override;
 	virtual const TCHAR* ImportText_Internal( const TCHAR* Buffer, void* Data, int32 PortFlags, UObject* OwnerObject, FOutputDevice* ErrorText ) const override;
+	virtual FString GetCPPTypeForwardDeclaration() const override
+	{
+		return FString();
+	}
 	// End of UProperty interface
 };
 
@@ -2074,6 +2090,10 @@ public:
 	// UProperty interface
 	virtual void ExportTextItem( FString& ValueStr, const void* PropertyValue, const void* DefaultValue, UObject* Parent, int32 PortFlags, UObject* ExportRootScope ) const override;
 	virtual const TCHAR* ImportText_Internal( const TCHAR* Buffer, void* Data, int32 PortFlags, UObject* OwnerObject, FOutputDevice* ErrorText ) const override;
+	virtual FString GetCPPTypeForwardDeclaration() const
+	{
+		return FString();
+	}
 	// End of UProperty interface
 };
 
@@ -2116,6 +2136,7 @@ public:
 	// UProperty interface
 	virtual FString GetCPPMacroType( FString& ExtendedTypeText ) const  override;
 	virtual FString GetCPPType( FString* ExtendedTypeText, uint32 CPPExportFlags ) const override;
+	virtual FString GetCPPTypeForwardDeclaration() const override;
 	virtual void LinkInternal(FArchive& Ar) override;
 	virtual bool Identical( const void* A, const void* B, uint32 PortFlags ) const override;
 	virtual void SerializeItem( FArchive& Ar, void* Value, int32 MaxReadBytes, void const* Defaults ) const override;
@@ -2481,6 +2502,7 @@ public:
 	// UProperty interface
 	virtual FString GetCPPMacroType( FString& ExtendedTypeText ) const  override;
 	virtual FString GetCPPType( FString* ExtendedTypeText, uint32 CPPExportFlags ) const override;
+	virtual FString GetCPPTypeForwardDeclaration() const override;
 	virtual void LinkInternal(FArchive& Ar) override;
 	virtual bool Identical( const void* A, const void* B, uint32 PortFlags ) const override;
 	virtual void SerializeItem( FArchive& Ar, void* Value, int32 MaxReadBytes, void const* Defaults ) const override;
@@ -2553,6 +2575,7 @@ public:
 
 	// UProperty interface
 	virtual FString GetCPPType( FString* ExtendedTypeText, uint32 CPPExportFlags ) const override;
+	virtual FString GetCPPTypeForwardDeclaration() const override;
 	virtual bool Identical( const void* A, const void* B, uint32 PortFlags ) const override;
 	virtual void SerializeItem( FArchive& Ar, void* Value, int32 MaxReadBytes, void const* Defaults ) const override;
 	virtual bool NetSerializeItem( FArchive& Ar, UPackageMap* Map, void* Data, TArray<uint8> * MetaData = NULL ) const override;
diff --git a/Engine/Source/Runtime/Engine/Classes/Animation/AnimNotifies/AnimNotifyState_TimedParticleEffect.h b/Engine/Source/Runtime/Engine/Classes/Animation/AnimNotifies/AnimNotifyState_TimedParticleEffect.h
--- a/Engine/Source/Runtime/Engine/Classes/Animation/AnimNotifies/AnimNotifyState_TimedParticleEffect.h
+++ b/Engine/Source/Runtime/Engine/Classes/Animation/AnimNotifies/AnimNotifyState_TimedParticleEffect.h
@@ -56,7 +56,7 @@ class UAnimNotifyState_TimedParticleEffect : public UAnimNotifyState
 	virtual void NotifyEnd(class USkeletalMeshComponent * MeshComp, class UAnimSequenceBase * Animation) override;
 
 	// Overriden from UAnimNotifyState to provide custom notify name.
-	virtual FString GetNotifyName_Implementation() const override;
+	FString GetNotifyName_Implementation() const;
 
 private:
 	bool ValidateParameters(USkeletalMeshComponent* MeshComp);
diff --git a/Engine/Source/Runtime/Engine/Classes/GameFramework/Character.h b/Engine/Source/Runtime/Engine/Classes/GameFramework/Character.h
--- a/Engine/Source/Runtime/Engine/Classes/GameFramework/Character.h
+++ b/Engine/Source/Runtime/Engine/Classes/GameFramework/Character.h
@@ -402,8 +402,7 @@ public:
 	/**
 	 * Check if the character can jump in the current state.
 	 *
-	 * The default implementation may be overridden or extended by implementing the custom CanJump event in Blueprints.
-	 * C++ code can override CanJumpInternal_Implementation() for custom behavior.
+	 * The default implementation may be overridden or extended by implementing the custom CanJump event in Blueprints. 
 	 * 
 	 * @Return Whether the character can jump in the current state. 
 	 */
@@ -420,7 +419,6 @@ protected:
 	 * As well as returning true when on the ground, it also returns true when GetMaxJumpTime is more
 	 * than zero and IsJumping returns true.
 	 * 
-	 * (This generates a virtual CanJumpInternal_Implementation function).
 	 *
 	 * @Return Whether the character can jump in the current state. 
 	 */
diff --git a/Engine/Source/Runtime/Engine/Classes/GameFramework/Controller.h b/Engine/Source/Runtime/Engine/Classes/GameFramework/Controller.h
--- a/Engine/Source/Runtime/Engine/Classes/GameFramework/Controller.h
+++ b/Engine/Source/Runtime/Engine/Classes/GameFramework/Controller.h
@@ -169,6 +169,8 @@ public:
 
 public:
 
+	virtual void GetLifetimeReplicatedProps(TArray<FLifetimeProperty>& OutLifetimeProps) const;
+
 	// Begin AActor interface
 	virtual void TickActor( float DeltaTime, enum ELevelTick TickType, FActorTickFunction& ThisTickFunction ) override;
 	virtual void K2_DestroyActor() override;
