Fix for GUID calculation on split generated files.

diff --git a/Engine/Source/Programs/UnrealHeaderTool/Private/CodeGenerator.cpp b/Engine/Source/Programs/UnrealHeaderTool/Private/CodeGenerator.cpp
--- a/Engine/Source/Programs/UnrealHeaderTool/Private/CodeGenerator.cpp
+++ b/Engine/Source/Programs/UnrealHeaderTool/Private/CodeGenerator.cpp
@@ -958,7 +958,23 @@ void FNativeClassHeaderGenerator::ExportGeneratedPackageInitCode(UPackage* Packa
 	TheFlagAudit.Add(Package, TEXT("PackageFlags"), Package->PackageFlags);
 	{
 		FGuid Guid;
-		Guid.A = GenerateTextCRC(*GeneratedFunctionText        .ToUpper());
+
+		uint32 CombinedCRC = 0;
+		for (auto& Split : GeneratedFunctionBodyTextSplit)
+		{
+			uint32 SplitCRC = GenerateTextCRC(*Split.ToUpper());
+			if (CombinedCRC == 0)
+			{
+				// Don't combine in the first case because it keeps GUID backwards compatibility
+				CombinedCRC = SplitCRC;
+			}
+			else
+			{
+				CombinedCRC = HashCombine(SplitCRC, CombinedCRC);
+			}
+		}
+
+		Guid.A = CombinedCRC;
 		Guid.B = GenerateTextCRC(*GeneratedFunctionDeclarations.ToUpper());
 		GeneratedFunctionText.Logf(TEXT("            FGuid Guid;\r\n"));
 		GeneratedFunctionText.Logf(TEXT("            Guid.A = 0x%08X;\r\n"), Guid.A);
