[UECORE-65]: Generated code for deprecated class is corrupt when there's an enum at the top of the file.

Prevented a token check against the class name until parsing of the class definition has started.
Added a test case.

diff --git a/Engine/Source/Programs/UnrealHeaderTool/Private/HeaderParser.cpp b/Engine/Source/Programs/UnrealHeaderTool/Private/HeaderParser.cpp
--- a/Engine/Source/Programs/UnrealHeaderTool/Private/HeaderParser.cpp
+++ b/Engine/Source/Programs/UnrealHeaderTool/Private/HeaderParser.cpp
@@ -3994,7 +3994,7 @@ bool FHeaderParser::CompileDeclaration( FClasses& AllClasses, FToken& Token )
 			FError::Throwf(TEXT("Extra ';' before end of file") );
 		}
 	}
-	else if (Token.Matches(NameLookupCPP.GetNameCPP(Class)))
+	else if (bEncounteredNewStyleClass_UnmatchedBrackets && Token.Matches(NameLookupCPP.GetNameCPP(Class)))
 	{
 		if(!TryToMatchConstructorParameterList(Token))
 		{
diff --git a/Engine/Source/Programs/UnrealHeaderTool/Resources/UHTDebugging.manifest b/Engine/Source/Programs/UnrealHeaderTool/Resources/UHTDebugging.manifest
--- a/Engine/Source/Programs/UnrealHeaderTool/Resources/UHTDebugging.manifest
+++ b/Engine/Source/Programs/UnrealHeaderTool/Resources/UHTDebugging.manifest
@@ -26,6 +26,7 @@
 			"IncludeBase": ".\\UHTDebugging",
 			"OutputDirectory": "..\\..\\..\\..\\Intermediate/Build/Inc\\UHTDebugging\\",
 			"ClassesHeaders": [
+				".\\UHTDebugging\\TestDeprecatedObject.h",
 				".\\UHTDebugging\\TestObject.h",
 				".\\UHTDebugging\\TestInterface.h",
 				".\\UHTDebugging\\TestInterfaceObject.h",
diff --git a/Engine/Source/Programs/UnrealHeaderTool/Resources/UHTDebugging/TestDeprecatedObject.h b/Engine/Source/Programs/UnrealHeaderTool/Resources/UHTDebugging/TestDeprecatedObject.h
new file mode 100644
--- /dev/null
+++ b/Engine/Source/Programs/UnrealHeaderTool/Resources/UHTDebugging/TestDeprecatedObject.h
@@ -0,0 +1,19 @@
+// Copyright 1998-2015 Epic Games, Inc. All Rights Reserved.
+
+#pragma once
+
+#include "TestDeprecatedObject.generated.h"
+
+namespace EFoo
+{
+	enum Bar
+	{
+		Example,
+	};
+}
+
+UCLASS(deprecated)
+class UDEPRECATED_TestDeprecatedObject : public UObject
+{
+	GENERATED_UCLASS_BODY()
+};
