Core: UnrealHeaderTool fails on const member variables

UE-7101

diff --git a/Engine/Source/Programs/UnrealHeaderTool/Private/HeaderParser.cpp b/Engine/Source/Programs/UnrealHeaderTool/Private/HeaderParser.cpp
--- a/Engine/Source/Programs/UnrealHeaderTool/Private/HeaderParser.cpp
+++ b/Engine/Source/Programs/UnrealHeaderTool/Private/HeaderParser.cpp
@@ -3208,7 +3208,14 @@ bool FHeaderParser::GetVarType
 
 	if (bUnconsumedConstKeyword)
 	{
-		FError::Throwf(TEXT("Inappropriate keyword 'const' on variable of type '%s'"), VarType.Identifier );
+		if (VariableCategory == EVariableCategory::Member)
+		{
+			FError::Throwf(TEXT("Const properties are not supported."));
+		}
+		else
+		{
+			FError::Throwf(TEXT("Inappropriate keyword 'const' on variable of type '%s'"), VarType.Identifier);
+		}
 	}
 
 	if (bUnconsumedClassKeyword)
