UHT: Splitting generated CPP files fix.

UECORE-73

diff --git a/Engine/Source/Programs/UnrealHeaderTool/Private/CodeGenerator.cpp b/Engine/Source/Programs/UnrealHeaderTool/Private/CodeGenerator.cpp
--- a/Engine/Source/Programs/UnrealHeaderTool/Private/CodeGenerator.cpp
+++ b/Engine/Source/Programs/UnrealHeaderTool/Private/CodeGenerator.cpp
@@ -930,11 +930,7 @@ void FNativeClassHeaderGenerator::ExportGeneratedPackageInitCode(UPackage* Packa
 	FString ApiString = FString::Printf(TEXT("%s_API "), *API);
 	FString SingletonName(GetPackageSingletonName(Package));
 
-	if( GeneratedFunctionBodyTextSplit.Num() == 0 )
-	{
-		new(GeneratedFunctionBodyTextSplit) FStringOutputDeviceCountLines;
-	}
-	FStringOutputDevice& GeneratedFunctionText = GeneratedFunctionBodyTextSplit[GeneratedFunctionBodyTextSplit.Num()-1];
+	FStringOutputDevice& GeneratedFunctionText = GetGeneratedFunctionTextDevice();
 
 	GeneratedFunctionDeclarations.Logf(TEXT("    %sclass UPackage* %s;\r\n"), *ApiString, *SingletonName);
 
@@ -992,14 +988,7 @@ void FNativeClassHeaderGenerator::ExportGeneratedPackageInitCode(UPackage* Packa
 
 void FNativeClassHeaderGenerator::ExportNativeGeneratedInitCode(FClass* Class)
 {
-	int32 MaxLinesPerCpp = 30000;
-
-	if ((GeneratedFunctionBodyTextSplit.Num() == 0) || (GeneratedFunctionBodyTextSplit[GeneratedFunctionBodyTextSplit.Num()-1].GetLineCount() > MaxLinesPerCpp))
-	{
-		new(GeneratedFunctionBodyTextSplit) FStringOutputDeviceCountLines;
-	}
-	FStringOutputDevice& GeneratedFunctionText = GeneratedFunctionBodyTextSplit[GeneratedFunctionBodyTextSplit.Num()-1];
-
+	FStringOutputDevice& GeneratedFunctionText = GetGeneratedFunctionTextDevice();
 	check(!FriendText.Len());
 	// Emit code to build the UObjects that used to be in .u files
 	bool bIsNoExport = Class->HasAnyClassFlags(CLASS_NoExport) || FClassUtils::IsTemporaryClass(Class);
@@ -4314,6 +4303,18 @@ bool IsExportOrTemporaryClass(FClass* Class)
 	return bIsExportClass || bIsTemporaryClass;
 }
 
+FStringOutputDevice& FNativeClassHeaderGenerator::GetGeneratedFunctionTextDevice()
+{
+	int32 MaxLinesPerCpp = 30000;
+
+	if ((GeneratedFunctionBodyTextSplit.Num() == 0) || (GeneratedFunctionBodyTextSplit[GeneratedFunctionBodyTextSplit.Num() - 1].GetLineCount() > MaxLinesPerCpp))
+	{
+		new(GeneratedFunctionBodyTextSplit)FStringOutputDeviceCountLines;
+	}
+
+	return GeneratedFunctionBodyTextSplit[GeneratedFunctionBodyTextSplit.Num() - 1];
+}
+
 // Constructor.
 FNativeClassHeaderGenerator::FNativeClassHeaderGenerator( UPackage* InPackage, FClasses& AllClasses, bool InAllowSaveExportedHeaders )
 	: CurrentClass                          (NULL)
diff --git a/Engine/Source/Programs/UnrealHeaderTool/Private/NativeClassExporter.h b/Engine/Source/Programs/UnrealHeaderTool/Private/NativeClassExporter.h
--- a/Engine/Source/Programs/UnrealHeaderTool/Private/NativeClassExporter.h
+++ b/Engine/Source/Programs/UnrealHeaderTool/Private/NativeClassExporter.h
@@ -111,6 +111,11 @@ private:
 	{
 	};
 
+	/**
+	 * Gets generated function text device.
+	 */
+	FStringOutputDevice& GetGeneratedFunctionTextDevice();
+
 	/**
 	 * Sorts the list of header files being exported from a package according to their dependency on each other.
 	 *
diff --git a/Engine/Source/Runtime/Core/Public/Containers/UnrealString.h b/Engine/Source/Runtime/Core/Public/Containers/UnrealString.h
--- a/Engine/Source/Runtime/Core/Public/Containers/UnrealString.h
+++ b/Engine/Source/Runtime/Core/Public/Containers/UnrealString.h
@@ -1666,7 +1666,7 @@ public:
 	}
 	virtual void Serialize( const TCHAR* InData, ELogVerbosity::Type Verbosity, const class FName& Category ) override
 	{
-		*this += (TCHAR*)InData;
+		FString::operator+=((TCHAR*)InData);
 		if(bAutoEmitLineTerminator)
 		{
 			*this += LINE_TERMINATOR;
@@ -1718,6 +1718,11 @@ public:
 	#endif
 
 #endif
+	// Make += operator virtual.
+	virtual FString& operator+=(const FString& Other)
+	{
+		return FString::operator+=(Other);
+	}
 };
 
 //
@@ -1734,27 +1739,50 @@ public:
 	,	LineCount(0)
 	{}
 
-	virtual void Serialize( const TCHAR* InData, ELogVerbosity::Type Verbosity, const class FName& Category ) override
+	virtual void Serialize(const TCHAR* InData, ELogVerbosity::Type Verbosity, const class FName& Category) override
 	{
 		Super::Serialize(InData, Verbosity, Category);
 		int32 TermLength = FCString::Strlen(LINE_TERMINATOR);
-		for(;;)
+		for (;;)
 		{
 			InData = FCString::Strstr(InData, LINE_TERMINATOR);
-			if( !InData )
+			if (!InData)
 			{
 				break;
 			}
 			LineCount++;
 			InData += TermLength;
-		} while(InData);
+		} while (InData);
 
-		if(bAutoEmitLineTerminator)
+		if (bAutoEmitLineTerminator)
 		{
 			LineCount++;
 		}
 	}
 
+	/**
+	 * Appends other FStringOutputDeviceCountLines object to this one.
+	 */
+	virtual FStringOutputDeviceCountLines& operator+=(const FStringOutputDeviceCountLines& Other)
+	{
+		FString::operator+=(Other);
+
+		LineCount += Other.GetLineCount();
+
+		return *this;
+	}
+
+	/**
+	 * Appends other FString (as well as it's specializations like FStringOutputDevice)
+	 * object to this.
+	 */
+	virtual FString& operator+=(const FString& Other) override
+	{
+		Log(Other);
+
+		return *this;
+	}
+
 	int32 GetLineCount() const
 	{
 		return LineCount;
