UHT: Allows defining multiple UCLASSes in one header file.

diff --git a/Engine/Source/Programs/UnrealHeaderTool/Private/CodeGenerator.cpp b/Engine/Source/Programs/UnrealHeaderTool/Private/CodeGenerator.cpp
--- a/Engine/Source/Programs/UnrealHeaderTool/Private/CodeGenerator.cpp
+++ b/Engine/Source/Programs/UnrealHeaderTool/Private/CodeGenerator.cpp
@@ -25,7 +25,7 @@ static bool bVerifyContents = false;
 static const bool bMultiLineUFUNCTION = true;
 static const bool bMultiLineUPROPERTY = true;
 
-static UClass* GenerateCodeForHeader(UObject* InParent, const TCHAR* Name, EObjectFlags Flags, const TCHAR* Buffer, int32& OutClassDeclLine);
+static TSharedRef<FUnrealSourceFile> GenerateCodeForHeader(UPackage* InParent, const FString& FileName, EObjectFlags Flags, const TCHAR* Buffer);
 
 FCompilerMetadataManager GScriptHelper;
 
@@ -389,22 +389,6 @@ static FString GetMetaDataCodeForObject(UObject* Object, const TCHAR* SymbolName
 	return Result;
 }
 
-/**
- * Returns a C++ code string for declaring a class or variable as an export, if the class is configured for that
- *
- * @return	The API declaration string
- */
-FString FNativeClassHeaderGenerator::MakeAPIDecl() const
-{
-	FString APIDecl;
-	if( CurrentClass->HasAnyClassFlags( CLASS_RequiredAPI ) )
-	{
-		APIDecl = FString::Printf( TEXT( "%s_API " ), *API );
-	}
-
-	return APIDecl;
-}
-
 /**
  * Exports the struct's C++ properties to the specified output device and adds special
  * compiler directives for GCC to pack as we expect.
@@ -489,7 +473,7 @@ void FNativeClassHeaderGenerator::ExportProperties( UStruct* Struct, int32 TextI
 			if (bAccessSpecifiers)
 			{
 				// find the class info for this class
-				FClassMetaData* ClassData = GScriptHelper.FindClassData(CurrentClass);
+				FClassMetaData* ClassData = GScriptHelper.FindClassData(Struct);
 
 				// find the compiler token for this property
 				FTokenData* PropData = ClassData->FindTokenData(Current);
@@ -641,6 +625,7 @@ FString FNativeClassHeaderGenerator::GetSingletonName(UField* Item, bool bRequir
 		}
 		Outer = Outer->GetOuter();
 	}
+
 	FString ClassString = NameLookupCPP.GetNameCPP(Item->GetClass());
 	// Can't use long package names in function names.
 	if (Result.StartsWith(TEXT("/Script/")))
@@ -795,9 +780,23 @@ void FNativeClassHeaderGenerator::OutputProperties(FString& Meta, FOutputDevice&
 	}
 }
 
-inline FString GetEventStructParamsName(const TCHAR* ClassName, const TCHAR* FunctionName)
+inline FString GetEventStructParamsName(UObject* Outer, const TCHAR* FunctionName)
 {
-	FString Result = FString::Printf(TEXT("%s_event%s_Parms"), ClassName, FunctionName);
+	FString OuterName;
+	if (Outer->IsA<UClass>())
+	{
+		OuterName = ((UClass*)Outer)->GetName();
+	}
+	else if (Outer->IsA<UPackage>())
+	{
+		OuterName = ((UPackage*)Outer)->GetName().Replace(TEXT("/"), TEXT("_"));
+	}
+	else
+	{
+		FError::Throwf(TEXT("Unrecognized outer type"));
+	}
+
+	FString Result = FString::Printf(TEXT("%s_event%s_Parms"), *OuterName, FunctionName);
 	return Result;
 }
 
@@ -818,7 +817,7 @@ void FNativeClassHeaderGenerator::OutputProperty(FString& Meta, FOutputDevice& O
 				FunctionName = FunctionName.LeftChop( FString( HEADER_GENERATED_DELEGATE_SIGNATURE_SUFFIX ).Len() );
 			}
 
-			SourceStruct = GetEventStructParamsName(*CastChecked<UClass>(Function->GetOuter())->GetName(), *FunctionName);
+			SourceStruct = GetEventStructParamsName(Function->GetOuter(), *FunctionName);
 		}
 		else
 		{			
@@ -927,7 +926,7 @@ FString GetPackageSingletonName(UPackage* Package)
 
 void FNativeClassHeaderGenerator::ExportGeneratedPackageInitCode(UPackage* Package)
 {
-	FString ApiString = FString::Printf(TEXT("%s_API "), *API);
+	FString ApiString = GetAPIString();
 	FString SingletonName(GetPackageSingletonName(Package));
 
 	FStringOutputDevice& GeneratedFunctionText = GetGeneratedFunctionTextDevice();
@@ -978,6 +977,24 @@ void FNativeClassHeaderGenerator::ExportGeneratedPackageInitCode(UPackage* Packa
 		GeneratedFunctionText.Logf(TEXT("            Guid.C = 0x%08X;\r\n"), Guid.C);
 		GeneratedFunctionText.Logf(TEXT("            Guid.D = 0x%08X;\r\n"), Guid.D);
 		GeneratedFunctionText.Logf(TEXT("            ReturnPackage->SetGuid(Guid);\r\n"), Guid.D);
+
+		GeneratedFunctionText.Log(TEXT("\r\n"));
+
+		for (UField* ScriptType : TObjectRange<UField>())
+		{
+			if (ScriptType->GetOutermost() != Package)
+			{
+				continue;
+			}
+
+			if (
+				(ScriptType->IsA<UScriptStruct>() && (((UScriptStruct*)ScriptType)->StructFlags & STRUCT_NoExport) != 0) ||
+				ScriptType->IsA<UDelegateFunction>()
+				)
+			{
+				GeneratedFunctionText.Logf(TEXT("            %s;\r\n"), *GetSingletonName(ScriptType));
+			}
+		}
 	}
 
 	GeneratedFunctionText.Logf(TEXT("        }\r\n"));
@@ -986,176 +1003,15 @@ void FNativeClassHeaderGenerator::ExportGeneratedPackageInitCode(UPackage* Packa
 
 }
 
-void FNativeClassHeaderGenerator::ExportNativeGeneratedInitCode(FClass* Class)
+void FNativeClassHeaderGenerator::ExportNativeGeneratedInitCode(FClass* Class, FStringOutputDevice& OutFriendText)
 {
 	FStringOutputDevice& GeneratedFunctionText = GetGeneratedFunctionTextDevice();
-	check(!FriendText.Len());
+	check(!OutFriendText.Len());
 	// Emit code to build the UObjects that used to be in .u files
-	bool bIsNoExport = Class->HasAnyClassFlags(CLASS_NoExport) || FClassUtils::IsTemporaryClass(Class);
+	bool bIsNoExport = Class->HasAnyClassFlags(CLASS_NoExport);
 	FStringOutputDevice BodyText;
 	FStringOutputDevice CallSingletons;
-	FString ApiString = FString::Printf(TEXT("%s_API "), *API);
-
-	// enums
-	for ( TFieldIterator<UEnum> It(Class, EFieldIteratorFlags::ExcludeSuper); It; ++It )
-	{
-		UEnum* Enum = *It;
-
-		const FString SingletonName(GetSingletonName(Enum));		
-
-		FString Extern = FString::Printf(TEXT("    %sclass UEnum* %s;\r\n"), *ApiString, *SingletonName);
-		SingletonNameToExternDecl.Add(SingletonName, Extern);
-		GeneratedFunctionDeclarations.Log(*Extern);
-			
-		FStringOutputDevice GeneratedEnumRegisterFunctionText;
-
-		GeneratedEnumRegisterFunctionText.Logf(TEXT("    UEnum* %s\r\n"), *SingletonName);
-		GeneratedEnumRegisterFunctionText.Logf(TEXT("    {\r\n"));
-		// Enums can either have a UClass or UPackage as outer (if declared in non-UClass header).
-		if (Enum->GetOuter()->IsA(UStruct::StaticClass()))
-		{
-			GeneratedEnumRegisterFunctionText.Logf(TEXT("        UClass* Outer=%s;\r\n"), *GetSingletonName(CastChecked<UStruct>(Enum->GetOuter())));
-		}
-		else
-		{
-			GeneratedEnumRegisterFunctionText.Logf(TEXT("        UPackage* Outer=%s;\r\n"), *GetPackageSingletonName(CastChecked<UPackage>(Enum->GetOuter())));
-		}
-		GeneratedEnumRegisterFunctionText.Logf(TEXT("        static UEnum* ReturnEnum = NULL;\r\n"));
-		GeneratedEnumRegisterFunctionText.Logf(TEXT("        if (!ReturnEnum)\r\n"));
-		GeneratedEnumRegisterFunctionText.Logf(TEXT("        {\r\n"));
-
-		GeneratedEnumRegisterFunctionText.Logf(TEXT("            ReturnEnum = new(Outer, TEXT(\"%s\"), RF_Public|RF_Transient|RF_Native) UEnum(FObjectInitializer());\r\n"), *Enum->GetName());
-		GeneratedEnumRegisterFunctionText.Logf(TEXT("            TArray<FName> EnumNames;\r\n"));
-		for (int32 Index = 0; Index < Enum->NumEnums(); Index++)
-		{
-			GeneratedEnumRegisterFunctionText.Logf(TEXT("            EnumNames.Add(FName(TEXT(\"%s\")));\r\n"), *Enum->GetEnum(Index).ToString());
-		}
-
-		FString EnumTypeStr;
-		switch (Enum->GetCppForm())
-		{
-			case UEnum::ECppForm::Regular:    EnumTypeStr = TEXT("UEnum::ECppForm::Regular");    break;
-			case UEnum::ECppForm::Namespaced: EnumTypeStr = TEXT("UEnum::ECppForm::Namespaced"); break;
-			case UEnum::ECppForm::EnumClass:  EnumTypeStr = TEXT("UEnum::ECppForm::EnumClass");  break;
-		}
-		GeneratedEnumRegisterFunctionText.Logf(TEXT("            ReturnEnum->SetEnums(EnumNames, %s);\r\n"), *EnumTypeStr);
-
-		FString Meta = GetMetaDataCodeForObject(Enum, TEXT("ReturnEnum"), TEXT("            "));
-		if (Meta.Len())
-		{
-			GeneratedEnumRegisterFunctionText.Logf(TEXT("#if WITH_METADATA\r\n"));
-			GeneratedEnumRegisterFunctionText.Logf(TEXT("            UMetaData* MetaData = ReturnEnum->GetOutermost()->GetMetaData();\r\n"));
-			GeneratedEnumRegisterFunctionText.Log(*Meta);
-			GeneratedEnumRegisterFunctionText.Logf(TEXT("#endif\r\n"));
-		}
-
-		GeneratedEnumRegisterFunctionText.Logf(TEXT("        }\r\n"));
-		GeneratedEnumRegisterFunctionText.Logf(TEXT("        return ReturnEnum;\r\n"));
-		GeneratedEnumRegisterFunctionText.Logf(TEXT("    }\r\n"));
-
-		GeneratedFunctionText += GeneratedEnumRegisterFunctionText;
-
-		uint32 EnumCrc = GenerateTextCRC(*GeneratedEnumRegisterFunctionText);
-		GGeneratedCodeCRCs.Add(Enum, EnumCrc);
-		CallSingletons.Logf(TEXT("                OuterClass->LinkChild(%s); // %u\r\n"), *SingletonName, EnumCrc);
-	}
-
-	// structs
-	for ( TFieldIterator<UScriptStruct> It(Class, EFieldIteratorFlags::ExcludeSuper); It; ++It )
-	{
-		UScriptStruct* ScriptStruct = *It;
-		UStruct* BaseStruct = ScriptStruct->GetSuperStruct();
-
-		const FString SingletonName(GetSingletonName(ScriptStruct));
-		
-		FString Extern = FString::Printf(TEXT("    %sclass UScriptStruct* %s;\r\n"), *ApiString, *SingletonName);
-		SingletonNameToExternDecl.Add(SingletonName, Extern);
-		GeneratedFunctionDeclarations.Log(*Extern);
-
-		FStringOutputDevice GeneratedStructRegisterFunctionText;
-
-		GeneratedStructRegisterFunctionText.Logf(TEXT("    UScriptStruct* %s\r\n"), *SingletonName);
-		GeneratedStructRegisterFunctionText.Logf(TEXT("    {\r\n"));
-
-		// if this is a no export struct, we will put a local struct here for offset determination
-		TArray<UScriptStruct*> Structs;
-		FindNoExportStructs(Structs, ScriptStruct);
-		if (Structs.Num())
-		{
-			TGuardValue<bool> Guard(bIsExportingForOffsetDeterminationOnly,true);
-			ExportMirrorsForNoexportStructs(Structs, 8, GeneratedStructRegisterFunctionText);
-		}
-
-		// Structs can either have a UClass or UPackage as outer (if delcared in non-UClass header).
-		if (ScriptStruct->GetOuter()->IsA(UStruct::StaticClass()))
-		{
-			GeneratedStructRegisterFunctionText.Logf(TEXT("        UStruct* Outer=%s;\r\n"), *GetSingletonName(CastChecked<UStruct>(ScriptStruct->GetOuter())));
-		}
-		else
-		{
-			GeneratedStructRegisterFunctionText.Logf(TEXT("        UPackage* Outer=%s;\r\n"), *GetPackageSingletonName(CastChecked<UPackage>(ScriptStruct->GetOuter())));
-		}
-		GeneratedStructRegisterFunctionText.Logf(TEXT("        static UScriptStruct* ReturnStruct = NULL;\r\n"));
-		GeneratedStructRegisterFunctionText.Logf(TEXT("        if (!ReturnStruct)\r\n"));
-		GeneratedStructRegisterFunctionText.Logf(TEXT("        {\r\n"));
-		FString BaseStructString(TEXT("NULL"));
-		if (BaseStruct)
-		{
-			CastChecked<UScriptStruct>(BaseStruct); // this better actually be a script struct
-			BaseStructString = GetSingletonName(BaseStruct);
-		}
-		FString CppStructOpsString(TEXT("NULL"));
-		FString ExplicitSizeString;
-		FString ExplicitAlignmentString;
-		if ( (ScriptStruct->StructFlags&STRUCT_Native) != 0 )
-		{
-			//@todo .u we don't need the auto register versions of these (except for hotreload, which should be fixed)
-			CppStructOpsString = FString::Printf( TEXT("new UScriptStruct::TCppStructOps<%s>"), NameLookupCPP.GetNameCPP(ScriptStruct) );
-		}
-		else
-		{
-			ExplicitSizeString = FString::Printf( TEXT(", sizeof(%s), ALIGNOF(%s)"), NameLookupCPP.GetNameCPP(ScriptStruct), NameLookupCPP.GetNameCPP(ScriptStruct) );
-		}
-
-		GeneratedStructRegisterFunctionText.Logf(TEXT("            ReturnStruct = new(Outer, TEXT(\"%s\"), RF_Public|RF_Transient|RF_Native) UScriptStruct(FObjectInitializer(), %s, %s, EStructFlags(0x%08X)%s);\r\n"),
-			*ScriptStruct->GetName(),
-			*BaseStructString,
-			*CppStructOpsString,
-			(uint32)(ScriptStruct->StructFlags & ~STRUCT_ComputedFlags),
-			*ExplicitSizeString
-			);
-		TheFlagAudit.Add(ScriptStruct, TEXT("StructFlags"), (uint64)(ScriptStruct->StructFlags & ~STRUCT_ComputedFlags));
-
-		TArray<UProperty*> Props;
-		for ( TFieldIterator<UProperty> ItInner(ScriptStruct, EFieldIteratorFlags::ExcludeSuper); ItInner; ++ItInner )
-		{
-			Props.Add(*ItInner);
-		}
-		FString OuterString = FString(TEXT("ReturnStruct"));
-		FString Meta = GetMetaDataCodeForObject(ScriptStruct, *OuterString, TEXT("            "));
-		OutputProperties(Meta, GeneratedStructRegisterFunctionText, OuterString, Props, TEXT("            "));
-		GeneratedStructRegisterFunctionText.Logf(TEXT("            ReturnStruct->StaticLink();\r\n"));
-
-		if (Meta.Len())
-		{
-			GeneratedStructRegisterFunctionText.Logf(TEXT("#if WITH_METADATA\r\n"));
-			GeneratedStructRegisterFunctionText.Logf(TEXT("            UMetaData* MetaData = ReturnStruct->GetOutermost()->GetMetaData();\r\n"));
-			GeneratedStructRegisterFunctionText.Log(*Meta);
-			GeneratedStructRegisterFunctionText.Logf(TEXT("#endif\r\n"));
-		}
-
-		GeneratedStructRegisterFunctionText.Logf(TEXT("        }\r\n"));
-		GeneratedStructRegisterFunctionText.Logf(TEXT("        return ReturnStruct;\r\n"));
-		GeneratedStructRegisterFunctionText.Logf(TEXT("    }\r\n"));
-		
-		uint32 StructCrc = GenerateTextCRC(*GeneratedStructRegisterFunctionText);
-		GGeneratedCodeCRCs.Add(ScriptStruct, StructCrc);
-
-		GeneratedFunctionText += GeneratedStructRegisterFunctionText;
-		GeneratedFunctionText.Logf(TEXT("    uint32 Get_%s_CRC() { return %uU; }\r\n"), *SingletonName.Replace(TEXT("()"), TEXT("")), StructCrc);
-
-		CallSingletons.Logf(TEXT("                OuterClass->LinkChild(%s); // %u\r\n"), *SingletonName, StructCrc);
-	}
+	FString ApiString = GetAPIString();
 
 	TArray<UFunction*> FunctionsToExport;
 	for( TFieldIterator<UFunction> Function(Class,EFieldIteratorFlags::ExcludeSuper); Function; ++Function )
@@ -1173,119 +1029,13 @@ void FNativeClassHeaderGenerator::ExportNativeGeneratedInitCode(FClass* Class)
 	for (int32 FuncIndex = 0; FuncIndex < FunctionsToExport.Num(); FuncIndex++)
 	{
 		UFunction* Function = FunctionsToExport[FuncIndex];
-		UFunction* SuperFunction = Function->GetSuperFunction();
-
-		FFunctionData* CompilerInfo = ClassData->FindFunctionData(Function);
-		check(CompilerInfo);
-
-		const FFuncInfo& FunctionData = CompilerInfo->GetFunctionData();
-
-		const FString SingletonName(GetSingletonName(Function));
-
-		CallSingletons.Logf(TEXT("                OuterClass->LinkChild(%s);\r\n"), *SingletonName);
-		FString Extern = FString::Printf(TEXT("    %sclass UFunction* %s;\r\n"), *ApiString, *SingletonName);
-		SingletonNameToExternDecl.Add(SingletonName, Extern);
-		GeneratedFunctionDeclarations.Log(*Extern);
-
-		FStringOutputDevice CurrentFunctionText;
-
-		CurrentFunctionText.Logf(TEXT("    UFunction* %s\r\n"), *SingletonName);
-		CurrentFunctionText.Logf(TEXT("    {\r\n"));
-
-		if (bIsNoExport || !(Function->FunctionFlags&FUNC_Event))  // non-events do not export a params struct, so lets do that locally for offset determination
-		{
-			TGuardValue<bool> Guard(bIsExportingForOffsetDeterminationOnly,true);
-			TArray<UScriptStruct*> Structs;
-			FindNoExportStructs(Structs, Function);
-			if (Structs.Num())
-			{
-				ExportMirrorsForNoexportStructs(Structs, 8, CurrentFunctionText);
-			}
-			TArray<UFunction*> CallbackFunctions;
-			CallbackFunctions.Add(Function);
-			ExportEventParms(CallbackFunctions, NULL, 8, false, &CurrentFunctionText);
-		}
-
-		CurrentFunctionText.Logf(TEXT("        UClass* OuterClass=%s;\r\n"), *GetSingletonName(Class));
-		CurrentFunctionText.Logf(TEXT("        static UFunction* ReturnFunction = NULL;\r\n"));
-		CurrentFunctionText.Logf(TEXT("        if (!ReturnFunction)\r\n"));
-		CurrentFunctionText.Logf(TEXT("        {\r\n"));
-		FString SuperFunctionString(TEXT("NULL"));
-		if (SuperFunction)
-		{
-			SuperFunctionString = GetSingletonName(SuperFunction);
-		}
-		TArray<UProperty*> Props;
-		for ( TFieldIterator<UProperty> ItInner(Function,EFieldIteratorFlags::ExcludeSuper); ItInner; ++ItInner )
-		{
-			Props.Add(*ItInner);
-		}
-		FString StructureSize;
-		if (Props.Num())
-		{
-			UFunction* TempFunction = Function;
-			while (TempFunction->GetSuperFunction())
-			{
-				TempFunction = TempFunction->GetSuperFunction();
-			}
-			FString FunctionName = TempFunction->GetName();
-			if( TempFunction->HasAnyFunctionFlags( FUNC_Delegate ) )
-			{
-				FunctionName = FunctionName.LeftChop( FString( HEADER_GENERATED_DELEGATE_SIGNATURE_SUFFIX ).Len() );
-			}
-
-			StructureSize = FString::Printf(TEXT(", sizeof(%s)"), *GetEventStructParamsName(*CastChecked<UClass>(TempFunction->GetOuter())->GetName(), *FunctionName));
-		}
-
-		CurrentFunctionText.Logf(TEXT("            ReturnFunction = new(OuterClass, TEXT(\"%s\"), RF_Public|RF_Transient|RF_Native) UFunction(FObjectInitializer(), %s, 0x%08X, %d%s);\r\n"), 
-			*Function->GetName(),
-			*SuperFunctionString,
-			Function->FunctionFlags,
-			(uint32)Function->RepOffset,
-			*StructureSize
-			);
-		TheFlagAudit.Add(Function, TEXT("FunctionFlags"), Function->FunctionFlags);
-		
-
-		FString OuterString = FString(TEXT("ReturnFunction"));
-		FString Meta        = GetMetaDataCodeForObject(Function, *OuterString, TEXT("            "));
-
-		for (int32 Index = Props.Num() - 1; Index >= 0; Index--)
-		{
-			OutputProperty(Meta, CurrentFunctionText, OuterString, Props[Index], TEXT("            "));
-		}
-
-		if (FunctionData.FunctionFlags & (FUNC_NetRequest | FUNC_NetResponse))
-		{
-			CurrentFunctionText.Logf(TEXT("            ReturnFunction->RPCId=%d;\r\n"), FunctionData.RPCId);
-			CurrentFunctionText.Logf(TEXT("            ReturnFunction->RPCResponseId=%d;\r\n"), FunctionData.RPCResponseId);
-		}
-
-		CurrentFunctionText.Logf(TEXT("            ReturnFunction->Bind();\r\n"));
-		CurrentFunctionText.Logf(TEXT("            ReturnFunction->StaticLink();\r\n"));
-
-		if (Meta.Len())
-		{
-			CurrentFunctionText.Logf(TEXT("#if WITH_METADATA\r\n"));
-			CurrentFunctionText.Logf(TEXT("            UMetaData* MetaData = ReturnFunction->GetOutermost()->GetMetaData();\r\n"));
-			CurrentFunctionText.Log(*Meta);
-			CurrentFunctionText.Logf(TEXT("#endif\r\n"));
-		}
-
-		CurrentFunctionText.Logf(TEXT("        }\r\n"));
-		CurrentFunctionText.Logf(TEXT("        return ReturnFunction;\r\n"));
-		CurrentFunctionText.Logf(TEXT("    }\r\n"));
-
-		GeneratedFunctionText += CurrentFunctionText;
-
-		uint32 FunctionCrc = GenerateTextCRC(*CurrentFunctionText);
-		GGeneratedCodeCRCs.Add(Function, FunctionCrc);
+		ExportFunction(Function, &FScope::GetTypeScope(Class).Get(), bIsNoExport);
+		CallSingletons.Logf(TEXT("                OuterClass->LinkChild(%s);\r\n"), *GetSingletonName(Function));
 	}
 
 	FStringOutputDevice GeneratedClassRegisterFunctionText;
 
-	// The class itself (unless it's a temporary UHT class).
-	if (!FClassUtils::IsTemporaryClass(Class))
+	// The class itself.
 	{
 		// simple ::StaticClass wrapper to avoid header, link and DLL hell
 		{
@@ -1302,7 +1052,7 @@ void FNativeClassHeaderGenerator::ExportNativeGeneratedInitCode(FClass* Class)
 		}
 		FString SingletonName(GetSingletonName(Class));
 
-		FriendText.Logf(TEXT("    friend %sclass UClass* %s;\r\n"), *ApiString, *SingletonName);
+		OutFriendText.Logf(TEXT("    friend %sclass UClass* %s;\r\n"), *ApiString, *SingletonName);
 		FString Extern = FString::Printf(TEXT("    %sclass UClass* %s;\r\n"), *ApiString, *SingletonName);
 		SingletonNameToExternDecl.Add(SingletonName, Extern);
 		GeneratedFunctionDeclarations.Log(*Extern);
@@ -1402,23 +1152,19 @@ void FNativeClassHeaderGenerator::ExportNativeGeneratedInitCode(FClass* Class)
 
 		GeneratedFunctionText += GeneratedClassRegisterFunctionText;
 	}
-	if (FriendText.Len())
+
+	if (OutFriendText.Len() && bIsNoExport)
 	{
-		if (bIsNoExport)
-		{
-			GeneratedFunctionText.Logf(TEXT("    /* friend declarations for pasting into noexport class %s\r\n"), NameLookupCPP.GetNameCPP(Class));
-			GeneratedFunctionText.Log(FriendText);
-			GeneratedFunctionText.Logf(TEXT("    */\r\n"));
-			FriendText.Empty();
-		}
+		GeneratedFunctionText.Logf(TEXT("    /* friend declarations for pasting into noexport class %s\r\n"), NameLookupCPP.GetNameCPP(Class));
+		GeneratedFunctionText.Log(OutFriendText);
+		GeneratedFunctionText.Logf(TEXT("    */\r\n"));
+		OutFriendText.Empty();
 	}
 	FString SingletonName(GetSingletonName(Class));
 	SingletonName.ReplaceInline(TEXT("()"), TEXT("")); // function address
 
-	// Skip temporary UHT classes.
-	if (!FClassUtils::IsTemporaryClass(Class))
 	{	
-		auto ClassNameCPP = NameLookupCPP.GetNameCPP(Class);		
+		auto ClassNameCPP = NameLookupCPP.GetNameCPP(Class);
 		GeneratedFunctionText.Logf(TEXT("    static FCompiledInDefer Z_CompiledInDefer_UClass_%s(%s, TEXT(\"%s\"));\r\n"), 
 			ClassNameCPP, *SingletonName, ClassNameCPP);
 		
@@ -1429,11 +1175,135 @@ void FNativeClassHeaderGenerator::ExportNativeGeneratedInitCode(FClass* Class)
 	}
 }
 
+void FNativeClassHeaderGenerator::ExportFunction(UFunction* Function, FScope* Scope, bool bIsNoExport)
+{
+	UFunction* SuperFunction = Function->GetSuperFunction();
+
+	auto* CompilerInfo = FFunctionData::FindForFunction(Function);
+
+	bool bIsDelegate = Function->HasAnyFunctionFlags(FUNC_Delegate);
+
+	const FFuncInfo& FunctionData = CompilerInfo->GetFunctionData();
+
+	const FString SingletonName(GetSingletonName(Function));
+
+	FString Extern = FString::Printf(TEXT("    %sclass UFunction* %s;\r\n"), *GetAPIString(), *SingletonName);
+	SingletonNameToExternDecl.Add(SingletonName, Extern);
+	GeneratedFunctionDeclarations.Log(*Extern);
+
+	FStringOutputDevice CurrentFunctionText;
+
+	CurrentFunctionText.Logf(TEXT("    UFunction* %s\r\n"), *SingletonName);
+	CurrentFunctionText.Logf(TEXT("    {\r\n"));
+
+	if (bIsNoExport || !(Function->FunctionFlags&FUNC_Event))  // non-events do not export a params struct, so lets do that locally for offset determination
+	{
+		TGuardValue<bool> Guard(bIsExportingForOffsetDeterminationOnly, true);
+		TArray<UScriptStruct*> Structs;
+		FindNoExportStructs(Structs, Function);
+		if (Structs.Num())
+		{
+			ExportMirrorsForNoexportStructs(Structs, 8, CurrentFunctionText);
+		}
+		TArray<UFunction*> CallbackFunctions;
+		CallbackFunctions.Add(Function);
+		ExportEventParms(*Scope, CallbackFunctions, CurrentFunctionText, 8, false);
+	}
+
+	if (UObject* Outer = Function->GetOuter())
+	{
+		CurrentFunctionText.Logf(TEXT("        UObject* Outer=%s;\r\n"), Outer->IsA<UPackage>() ? *GetPackageSingletonName((UPackage*)Outer) : *GetSingletonName(Function->GetOwnerClass()));
+	}
+	else
+	{
+		CurrentFunctionText.Logf(TEXT("        UObject* Outer=nullptr;\r\n"));
+	}
+
+	CurrentFunctionText.Logf(TEXT("        static UFunction* ReturnFunction = NULL;\r\n"));
+	CurrentFunctionText.Logf(TEXT("        if (!ReturnFunction)\r\n"));
+	CurrentFunctionText.Logf(TEXT("        {\r\n"));
+	FString SuperFunctionString(TEXT("NULL"));
+	if (SuperFunction)
+	{
+		SuperFunctionString = GetSingletonName(SuperFunction);
+	}
+	TArray<UProperty*> Props;
+	for (TFieldIterator<UProperty> ItInner(Function, EFieldIteratorFlags::ExcludeSuper); ItInner; ++ItInner)
+	{
+		Props.Add(*ItInner);
+	}
+	FString StructureSize;
+	if (Props.Num())
+	{
+		UFunction* TempFunction = Function;
+		while (TempFunction->GetSuperFunction())
+		{
+			TempFunction = TempFunction->GetSuperFunction();
+		}
+		FString FunctionName = TempFunction->GetName();
+		if (TempFunction->HasAnyFunctionFlags(FUNC_Delegate))
+		{
+			FunctionName = FunctionName.LeftChop(FString(HEADER_GENERATED_DELEGATE_SIGNATURE_SUFFIX).Len());
+		}
+
+		StructureSize = FString::Printf(TEXT(", sizeof(%s)"), *GetEventStructParamsName(TempFunction->GetOuter(), *FunctionName));
+	}
+
+	FString UFunctionType = bIsDelegate ? TEXT("UDelegateFunction") : TEXT("UFunction");
+
+	CurrentFunctionText.Logf(TEXT("            ReturnFunction = new(Outer, TEXT(\"%s\"), RF_Public|RF_Transient|RF_Native) %s(FObjectInitializer(), %s, 0x%08X, %d%s);\r\n"),
+		*Function->GetName(),
+		*UFunctionType,
+		*SuperFunctionString,
+		Function->FunctionFlags,
+		(uint32)Function->RepOffset,
+		*StructureSize
+		);
+	TheFlagAudit.Add(Function, TEXT("FunctionFlags"), Function->FunctionFlags);
+
+
+	FString OuterString = FString(TEXT("ReturnFunction"));
+	FString Meta = GetMetaDataCodeForObject(Function, *OuterString, TEXT("            "));
+
+	for (int32 Index = Props.Num() - 1; Index >= 0; Index--)
+	{
+		OutputProperty(Meta, CurrentFunctionText, OuterString, Props[Index], TEXT("            "));
+	}
+
+	if (FunctionData.FunctionFlags & (FUNC_NetRequest | FUNC_NetResponse))
+	{
+		CurrentFunctionText.Logf(TEXT("            ReturnFunction->RPCId=%d;\r\n"), FunctionData.RPCId);
+		CurrentFunctionText.Logf(TEXT("            ReturnFunction->RPCResponseId=%d;\r\n"), FunctionData.RPCResponseId);
+	}
+
+	CurrentFunctionText.Logf(TEXT("            ReturnFunction->Bind();\r\n"));
+	CurrentFunctionText.Logf(TEXT("            ReturnFunction->StaticLink();\r\n"));
+
+	if (Meta.Len())
+	{
+		CurrentFunctionText.Logf(TEXT("#if WITH_METADATA\r\n"));
+		CurrentFunctionText.Logf(TEXT("            UMetaData* MetaData = ReturnFunction->GetOutermost()->GetMetaData();\r\n"));
+		CurrentFunctionText.Log(*Meta);
+		CurrentFunctionText.Logf(TEXT("#endif\r\n"));
+	}
+
+	CurrentFunctionText.Logf(TEXT("        }\r\n"));
+	CurrentFunctionText.Logf(TEXT("        return ReturnFunction;\r\n"));
+	CurrentFunctionText.Logf(TEXT("    }\r\n"));
+
+	uint32 FunctionCrc = GenerateTextCRC(*CurrentFunctionText);
+	GGeneratedCodeCRCs.Add(Function, FunctionCrc);
+
+	GetGeneratedFunctionTextDevice() += CurrentFunctionText;
+}
+
 void FNativeClassHeaderGenerator::ExportNatives(FClass* Class)
 {
-	// Skip temporary UHT classes.
-	if (FClassUtils::IsNoExportOrTemporaryClass(Class))
+	// Skip no export classes.
+	if (Class->HasAnyClassFlags(CLASS_NoExport))
+	{
 		return;
+	}
 
 	GeneratedPackageCPP.Logf(TEXT("    void %s::StaticRegisterNatives%s()\r\n"),NameLookupCPP.GetNameCPP(Class),NameLookupCPP.GetNameCPP(Class));
 	GeneratedPackageCPP.Logf(TEXT("    {\r\n"));
@@ -1472,10 +1342,8 @@ void FNativeClassHeaderGenerator::ExportNatives(FClass* Class)
 	GeneratedPackageCPP.Logf(TEXT("    }\r\n"));
 }
 
-void FNativeClassHeaderGenerator::ExportInterfaceCallFunctions( const TArray<UFunction*>& InCallbackFunctions, FStringOutputDevice& HeaderOutput )
+void FNativeClassHeaderGenerator::ExportInterfaceCallFunctions(const TArray<UFunction*>& InCallbackFunctions, UClass* Class, FClassMetaData* ClassData, FStringOutputDevice& HeaderOutput)
 {
-	FClassMetaData* ClassData = GScriptHelper.FindClassData(CurrentClass);
-
 	TArray<UFunction*> CallbackFunctions = InCallbackFunctions;
 	CallbackFunctions.Sort();
 
@@ -1483,10 +1351,9 @@ void FNativeClassHeaderGenerator::ExportInterfaceCallFunctions( const TArray<UFu
 	{
 		UFunction* Function     = *FuncIt;
 		FString    FunctionName = Function->GetName();
-		FString    ClassName    = CurrentClass->GetName();
+		FString    ClassName    = Class->GetName();
 
-		FFunctionData* CompilerInfo = ClassData->FindFunctionData(Function);
-		check(CompilerInfo);
+		auto* CompilerInfo = FFunctionData::FindForFunction(Function);
 
 		const FFuncInfo& FunctionData = CompilerInfo->GetFunctionData();
 		const TCHAR* ConstQualifier = FunctionData.FunctionReference->HasAllFunctionFlags(FUNC_Const) ? TEXT("const ") : TEXT("");
@@ -1508,7 +1375,7 @@ void FNativeClassHeaderGenerator::ExportInterfaceCallFunctions( const TArray<UFu
 		bool bHasParms = Parameters.HasParms();
 		if (bHasParms)
 		{
-			FString EventParmStructName = GetEventStructParamsName(*Function->GetOwnerClass()->GetName(), *FunctionName);
+			FString EventParmStructName = GetEventStructParamsName(Function->GetOuter(), *FunctionName);
 			GeneratedPackageCPP.Logf(TEXT("%s%s Parms;%s"), FCString::Spc(8), *EventParmStructName, LINE_TERMINATOR);
 		}
 
@@ -1579,70 +1446,7 @@ void FNativeClassHeaderGenerator::ExportInterfaceCallFunctions( const TArray<UFu
 		}
 
 		GeneratedPackageCPP.Logf(TEXT("%s}%s"), FCString::Spc(4), LINE_TERMINATOR);
-
-		//
-		// also generate an empty implementation for the base event function, since we did the Execute_ version above
-		//
-		ExportNativeFunctionHeader(FunctionData, GeneratedPackageCPP, EExportFunctionType::Event, EExportFunctionHeaderStyle::Definition, NULL);
-		GeneratedPackageCPP.Logf(TEXT("%s%s{%s"), LINE_TERMINATOR, FCString::Spc(4), LINE_TERMINATOR);
-
-		// assert if this is ever called directly
-		GeneratedPackageCPP.Logf(TEXT("%scheck(0 && \"Do not directly call Event functions in Interfaces. Call Execute_%s instead.\");%s"), FCString::Spc(8), *FunctionName, LINE_TERMINATOR);
-
-		// satisfy compiler if it's expecting a return value
-		if (Parameters.Return)
-		{
-			FString EventParmStructName = GetEventStructParamsName(*Function->GetOwnerClass()->GetName(), *FunctionName);
-			GeneratedPackageCPP.Logf(TEXT("%s%s Parms;%s"), FCString::Spc(8), *EventParmStructName, LINE_TERMINATOR);
-			GeneratedPackageCPP.Logf(TEXT("%sreturn Parms.ReturnValue;%s"), FCString::Spc(8), LINE_TERMINATOR);
-		}
-		GeneratedPackageCPP.Logf(TEXT("%s}%s"), FCString::Spc(4), LINE_TERMINATOR);
-	}
-}
-
-void FNativeClassHeaderGenerator::ExportNames()
-{
-	// Skip temporary UHT classes.
-	if (FClassUtils::IsNoExportOrTemporaryClass(CurrentClass))
-	{
-		return;
-	}
-	GeneratedHeaderTextBeforeForwardDeclarations.Logf(
-		TEXT("#ifdef %s_%s_generated_h")													LINE_TERMINATOR
-		TEXT("#error \"%s.generated.h already included, missing '#pragma once' in %s.h\"")	LINE_TERMINATOR
-		TEXT("#endif")																		LINE_TERMINATOR
-		TEXT("#define %s_%s_generated_h")													LINE_TERMINATOR
-		LINE_TERMINATOR,
-		*API, *CurrentClass->GetName(), *CurrentClass->GetName(), *CurrentClass->GetName(), *API, *CurrentClass->GetName()
-		);
-
-	TSet<FName> ClassReferencedNames;
-	for( TFieldIterator<UFunction> Function(CurrentClass,EFieldIteratorFlags::ExcludeSuper); Function; ++Function )
-	{
-		if ( Function->HasAnyFunctionFlags(FUNC_Event) && !Function->GetSuperFunction() )
-		{
-			FName FunctionName = Function->GetFName();
-			ClassReferencedNames.Add(FunctionName);
-			ReferencedNames.Add(FunctionName);
-		}
-	}
-	if (!ClassReferencedNames.Num())
-	{
-		return;
-	}
-
-	TArray<FString> Names;
-	for (const auto& Name : ClassReferencedNames)
-	{
-		Names.Add(Name.ToString());
-	}	
-	Names.Sort();
-
-	for( int32 NameIndex=0; NameIndex<Names.Num(); NameIndex++ )
-	{
-		GeneratedHeaderTextBeforeForwardDeclarations.Logf(TEXT("extern %s_API FName %s_%s;") LINE_TERMINATOR, *API, *API, *Names[NameIndex]);
 	}
-
 }
 
 /**
@@ -1690,326 +1494,264 @@ FString GetPreservedAccessSpecifierString(FClass* Class)
 }
 
 /**
- * Exports the C++ class declarations for a native interface class.
+ * Exports interface body macros.
  */
-void FNativeClassHeaderGenerator::ExportInterfaceClassDeclaration( FClass* Class )
+void ExportInterfaceBodyMacros(FStringOutputDevice& Out, FUnrealSourceFile& SourceFile, int32 ClassGeneratedBodyLine, FClass* Class, const FString& StdCtorCall, const FString& EnhCtorCall, const FString& UInterfaceBoilerplate)
 {
-	FClassMetaData* ClassData = GScriptHelper.FindClassData(CurrentClass);
-
-	TArray<UEnum*>			Enums;
-	TArray<UScriptStruct*>	Structs;
-	TArray<UFunction*>		CallbackFunctions;
-	TArray<UFunction*>		NativeFunctions;
-	TArray<UFunction*>		DelegateFunctions;
+	Out.Log(TEXT("#undef GENERATED_UINTERFACE_BODY_COMMON\r\n"));
+	Out.Log(Macroize(TEXT("GENERATED_UINTERFACE_BODY_COMMON()"), *UInterfaceBoilerplate));
+
+	auto DeprecationWarning = GetGeneratedMacroDeprecationWarning(TEXT("GENERATED_UINTERFACE_BODY"));
+
+	auto DeprecationPushString = TEXT("PRAGMA_DISABLE_DEPRECATION_WARNINGS") LINE_TERMINATOR;
+	auto DeprecationPopString = TEXT("PRAGMA_POP") LINE_TERMINATOR;
+	auto Offset = FCString::Spc(4);
+
+	Out.Log(Macroize(*SourceFile.GetGeneratedBodyMacroName(ClassGeneratedBodyLine, true), *(FString() +
+		Offset + DeprecationWarning +
+		Offset + DeprecationPushString +
+		Offset + TEXT("GENERATED_UINTERFACE_BODY_COMMON()") LINE_TERMINATOR +
+		StdCtorCall +
+		Offset + DeprecationPopString)));
+
+	Out.Log(Macroize(*SourceFile.GetGeneratedBodyMacroName(ClassGeneratedBodyLine), *(FString() +
+		Offset + DeprecationPushString +
+		Offset + TEXT("GENERATED_UINTERFACE_BODY_COMMON()") LINE_TERMINATOR +
+		EnhCtorCall +
+		GetPreservedAccessSpecifierString(Class) +
+		Offset + DeprecationPopString)));
+}
 
-	// get the lists of fields that belong to this class that should be exported
-	RetrieveRelevantFields(Class, Enums, Structs, CallbackFunctions, NativeFunctions, DelegateFunctions);
+void FNativeClassHeaderGenerator::ExportSourceFileHeader(FClasses& AllClasses, FUnrealSourceFile* SourceFile)
+{
+	TSet<const FUnrealSourceFile*>	VisitedSet;
+	ExportSourceFileHeaderRecursive(AllClasses, SourceFile, VisitedSet, false);
+}
 
-	ExportNames();
+void WriteMacro(FStringOutputDevice &Output, const FString& MacroName, const FString& MacroContent)
+{
+	Output.Log(*Macroize(*MacroName, *MacroContent));
+}
 
-	// export enum declarations
-	ExportEnums(Enums);
+void FNativeClassHeaderGenerator::ExportClassesFromSourceFileInner(FUnrealSourceFile& SourceFile)
+{
+	TArray<UEnum*>				Enums;
+	TArray<UScriptStruct*>		Structs;
+	TArray<UDelegateFunction*>	DelegateFunctions;
+
+	GeneratedHeaderText.Logf(
+		TEXT("#ifdef %s_%s_generated_h")													LINE_TERMINATOR
+		TEXT("#error \"%s.generated.h already included, missing '#pragma once' in %s.h\"")	LINE_TERMINATOR
+		TEXT("#endif")																		LINE_TERMINATOR
+		TEXT("#define %s_%s_generated_h")													LINE_TERMINATOR
+		LINE_TERMINATOR,
+		*API, *SourceFile.GetStrippedFilename(), *SourceFile.GetStrippedFilename(), *SourceFile.GetStrippedFilename(), *API, *SourceFile.GetStrippedFilename());
+
+	// get the lists of fields that belong to this class that should be exported
+	SourceFile.GetScope()->SplitTypesIntoArrays(Enums, Structs, DelegateFunctions);
+
+	// export enum declarations
+	ExportEnums(Enums);
 
 	// export delegate definitions
-	ExportDelegateDefinitions(DelegateFunctions, false);
+	ExportDelegateDefinitions(SourceFile, DelegateFunctions, false);
 
 	// Export enums declared in non-UClass headers.
 	ExportGeneratedEnumsInitCode(Enums);
 
-	// export boilerplate macros for struct
-	ExportGeneratedStructBodyMacros(Structs);
-
-	// export parameters structs for all events and delegates
-	ExportEventParms(CallbackFunctions, TEXT(""));
+	// export boilerplate macros for structs
+	ExportGeneratedStructBodyMacros(SourceFile, Structs);
 
 	// export delegate wrapper function implementations
-	ExportDelegateDefinitions(DelegateFunctions, true);
-
-	FClass* SuperClass = Class->GetSuperClass();
-
-
-	// the name for the C++ version of the UClass
-	const TCHAR* ClassCPPName		= NameLookupCPP.GetNameCPP( Class );
-	const TCHAR* SuperClassCPPName	= NameLookupCPP.GetNameCPP( SuperClass );
-
-	// =============================================
-	// Export the UClass declaration
-	// Class definition.
-	ExportNatives(Class);
-	ExportNativeGeneratedInitCode(Class);
+	ExportDelegateDefinitions(SourceFile, DelegateFunctions, true);
 
+	for (auto* RawUClass : SourceFile.GetDefinedClasses())
 	{
-		FStringOutputDevice UInterfaceBoilerplate;
+		PrologMacroCalls.Empty();
+		InClassMacroCalls.Empty();
+		InClassNoPureDeclsMacroCalls.Empty();
+		StandardUObjectConstructorsMacroCall.Empty();
+		EnhancedUObjectConstructorsMacroCall.Empty();
 
-		// Export the class's native function registration.
-		UInterfaceBoilerplate.Logf(TEXT("    private:\r\n"));
-		UInterfaceBoilerplate.Logf(TEXT("    static void StaticRegisterNatives%s();\r\n"),NameLookupCPP.GetNameCPP(Class));
-		UInterfaceBoilerplate.Log(*FriendText);
-		FriendText.Empty();
-		UInterfaceBoilerplate.Logf( TEXT("public:\r\n") );
+		auto* Class = (FClass*)RawUClass;
+		auto* ClassData = GScriptHelper.FindClassData(Class);
 
-		// Build the DECLARE_CLASS line
-		FString APIArg(API);
-		if( !Class->HasAnyClassFlags( CLASS_MinimalAPI ) )
+		if (Class->ClassFlags & CLASS_Intrinsic)
 		{
-			APIArg = TEXT("NO");
+			continue;
 		}
 
-		const bool bCastedClass = Class->HasAnyCastFlag(CASTCLASS_AllFlags) && SuperClass && Class->ClassCastFlags != SuperClass->ClassCastFlags;
-		UInterfaceBoilerplate.Logf( TEXT("    DECLARE_CLASS(%s, %s, COMPILED_IN_FLAGS(CLASS_Abstract%s), %s, %s, %s_API)\r\n"), 
-			ClassCPPName, 
-			SuperClassCPPName, 
-			*GetClassFlagExportText(Class), 
-			bCastedClass ? *FString::Printf(TEXT("CASTCLASS_%s"), ClassCPPName) : TEXT("0"),
-			*FPackageName::GetShortName(Class->GetOuter()->GetName()), 
-			*APIArg );
-		
-		UInterfaceBoilerplate.Logf(TEXT("    DECLARE_SERIALIZER(%s)\r\n"), ClassCPPName);
-		UInterfaceBoilerplate.Log(TEXT("    enum {IsIntrinsic=COMPILED_IN_INTRINSIC};\r\n"));
+		// C++ -> VM stubs (native function execs)
+		ExportNativeFunctions(SourceFile, Class, ClassData);
 
-		if (Class->ClassWithin != Class->GetSuperClass()->ClassWithin)
-		{
-			UInterfaceBoilerplate.Logf(TEXT("    DECLARE_WITHIN(%s)\r\n"), NameLookupCPP.GetNameCPP( Class->GetClassWithin() ) );
-		}
+		// VM -> C++ proxies (events and delegates).
+		TArray<UFunction*> CallbackFunctions = ExportCallbackFunctions(SourceFile, Class, ClassData);
 
-		ExportConstructorsMacros(Class);
+		// Class definition.
+		ExportNatives(Class);
 
-		GeneratedHeaderText.Log(TEXT("#undef GENERATED_UINTERFACE_BODY_COMMON\r\n"));
-		GeneratedHeaderText.Log(Macroize(TEXT("GENERATED_UINTERFACE_BODY_COMMON()"), *UInterfaceBoilerplate));
+		FStringOutputDevice FriendText;
+		ExportNativeGeneratedInitCode(Class, FriendText);
 
-		auto DeprecationWarning = GetGeneratedMacroDeprecationWarning(TEXT("GENERATED_UINTERFACE_BODY"));
-		
-		auto DeprecationPushString = TEXT("PRAGMA_DISABLE_DEPRECATION_WARNINGS") LINE_TERMINATOR;
-		auto DeprecationPopString = TEXT("PRAGMA_POP") LINE_TERMINATOR;
-		auto Offset = FCString::Spc(4);
+		FClass* SuperClass = Class->GetSuperClass();
 
-		GeneratedHeaderText.Log(TEXT("#undef GENERATED_UINTERFACE_BODY\r\n"));
-		GeneratedHeaderText.Log(Macroize(TEXT("GENERATED_UINTERFACE_BODY()"), *(FString() +
-			Offset + DeprecationWarning +
-			Offset + DeprecationPushString +
-			Offset + TEXT("GENERATED_UINTERFACE_BODY_COMMON()") LINE_TERMINATOR +
-			StandardUObjectConstructorsMacroCall +
-			Offset + DeprecationPopString)));
+		// the name for the C++ version of the UClass
+		const TCHAR* ClassCPPName = NameLookupCPP.GetNameCPP(Class);
+		const TCHAR* SuperClassCPPName = (SuperClass != nullptr) ? NameLookupCPP.GetNameCPP(SuperClass) : nullptr;
 
-		GeneratedHeaderText.Log(TEXT("#undef GENERATED_BODY\r\n"));
-		GeneratedHeaderText.Log(Macroize(TEXT("GENERATED_BODY()"), *(FString() +
-			Offset + DeprecationPushString +
-			Offset + TEXT("GENERATED_UINTERFACE_BODY_COMMON()") LINE_TERMINATOR +
-			EnhancedUObjectConstructorsMacroCall +
-			GetPreservedAccessSpecifierString(Class) +
-			Offset + DeprecationPopString)));
-	}
+		if (Class->HasAnyClassFlags(CLASS_Interface))
+		{
+			{
+				FStringOutputDevice UInterfaceBoilerplate;
 
-	// =============================================
-	// Export the pure interface version of the class
+				// Export the class's native function registration.
+				UInterfaceBoilerplate.Logf(TEXT("    private:\r\n"));
+				UInterfaceBoilerplate.Logf(TEXT("    static void StaticRegisterNatives%s();\r\n"), NameLookupCPP.GetNameCPP(Class));
+				UInterfaceBoilerplate.Log(*FriendText);
+				FriendText.Empty();
+				UInterfaceBoilerplate.Logf(TEXT("public:\r\n"));
 
-	// the name of the pure interface class
-	FString InterfaceCPPName		= FString::Printf(TEXT("I%s"), *Class->GetName());
-	FString SuperInterfaceCPPName;
-	if ( SuperClass != NULL )
-	{
-		SuperInterfaceCPPName = FString::Printf(TEXT("I%s"), *SuperClass->GetName());
-	}
+				// Build the DECLARE_CLASS line
+				FString APIArg(API);
+				if (!Class->HasAnyClassFlags(CLASS_MinimalAPI))
+				{
+					APIArg = TEXT("NO");
+				}
 
-	// C++ -> VM stubs (native function execs)
-	ExportNativeFunctions(NativeFunctions);
+				const bool bCastedClass = Class->HasAnyCastFlag(CASTCLASS_AllFlags) && SuperClass && Class->ClassCastFlags != SuperClass->ClassCastFlags;
+				UInterfaceBoilerplate.Logf(TEXT("    DECLARE_CLASS(%s, %s, COMPILED_IN_FLAGS(CLASS_Abstract%s), %s, %s, %s_API)\r\n"),
+					ClassCPPName,
+					SuperClassCPPName,
+					*GetClassFlagExportText(Class),
+					bCastedClass ? *FString::Printf(TEXT("CASTCLASS_%s"), ClassCPPName) : TEXT("0"),
+					*FPackageName::GetShortName(Class->GetOuter()->GetName()),
+					*APIArg);
 
-	// VM -> C++ proxies (events).
-	ExportCallbackFunctions(CallbackFunctions);
+				UInterfaceBoilerplate.Logf(TEXT("    DECLARE_SERIALIZER(%s)\r\n"), ClassCPPName);
+				UInterfaceBoilerplate.Log(TEXT("    enum {IsIntrinsic=COMPILED_IN_INTRINSIC};\r\n"));
 
-	// Thunk functions
-	{
-		FStringOutputDevice InterfaceBoilerplate;
+				if (Class->ClassWithin != Class->GetSuperClass()->ClassWithin)
+				{
+					UInterfaceBoilerplate.Logf(TEXT("    DECLARE_WITHIN(%s)\r\n"), NameLookupCPP.GetNameCPP(Class->GetClassWithin()));
+				}
 
-		InterfaceBoilerplate.Logf( TEXT("protected:\r\n\tvirtual ~%s() {}\r\npublic:\r\n"), *InterfaceCPPName );
-		InterfaceBoilerplate.Logf( TEXT("\ttypedef %s UClassType;\r\n"), ClassCPPName );
+				ExportConstructorsMacros(SourceFile.GetGeneratedMacroName(ClassData), Class);
 
-		ExportInterfaceCallFunctions(CallbackFunctions, InterfaceBoilerplate);
+				ExportInterfaceBodyMacros(GeneratedHeaderText, SourceFile, ClassData->GetGeneratedBodyLine(), Class, StandardUObjectConstructorsMacroCall, EnhancedUObjectConstructorsMacroCall, UInterfaceBoilerplate);
+			}
 
-		// we'll need a way to get to the UObject portion of a native interface, so that we can safely pass native interfaces
-		// to script VM functions
-		if (SuperClass->IsChildOf(UInterface::StaticClass()))
-		{
-			InterfaceBoilerplate.Logf(TEXT("\tvirtual UObject* _getUObject() const = 0;\r\n"));
-		}
+			// =============================================
+			// Export the pure interface version of the class
 
-		// Replication, add in the declaration for GetLifetimeReplicatedProps() automatically if there are any net flagged properties
-		bool bNeedsRep = false;
-		FStringOutputDevice StaticChecks;
-		auto ClassName = FString(NameLookupCPP.GetNameCPP(Class));
-		for( TFieldIterator<UProperty> It(Class,EFieldIteratorFlags::ExcludeSuper); It; ++It )
-		{
-			if( (It->PropertyFlags & CPF_Net) != 0 )
+			// the name of the pure interface class
+			FString InterfaceCPPName = FString::Printf(TEXT("I%s"), *Class->GetName());
+			FString SuperInterfaceCPPName;
+			if (SuperClass != NULL)
 			{
-				bNeedsRep = true;
-				StaticChecks.Logf(TEXT("\tstatic inline void StaticChecks()\r\n"));
-				StaticChecks.Logf(TEXT("\t{\r\n"));
-				StaticChecks.Logf(TEXT("\t\tstatic_assert(THasMemberFunction_GetLifetimeReplicatedProps<%s>::Value, \"Class %s has Net flagged properties and should declare member function: void GetLifetimeReplicatedProps( TArray< FLifetimeProperty > & OutLifetimeProps ) const override;\");\r\n"), ClassCPPName, ClassCPPName, ClassCPPName);
-				StaticChecks.Logf(TEXT("\t}\r\n"));
-				break;
+				SuperInterfaceCPPName = FString::Printf(TEXT("I%s"), *SuperClass->GetName());
 			}
-		}
-		
-		auto ClassRange = ClassDefinitionRange();
-		if (ClassDefinitionRanges.Contains(CurrentClass))
-		{
-			ClassRange = ClassDefinitionRanges[CurrentClass];
-		}
 
-		auto ClassStart = ClassRange.Start;
-		auto ClassEnd = ClassRange.End;
-		auto ClassDefinition = FString(ClassEnd - ClassStart, ClassStart);
 
-		bool bHasGetLifetimeReplicatedProps = ClassDefinition.Contains(TEXT("GetLifetimeReplicatedProps"), ESearchCase::CaseSensitive);
+			// Thunk functions
+			FStringOutputDevice InterfaceBoilerplate;
 
-		if (bNeedsRep && !bHasGetLifetimeReplicatedProps)
-		{
-			InterfaceBoilerplate.Logf(TEXT("\void GetLifetimeReplicatedProps(TArray<FLifetimeProperty>& OutLifetimeProps) const;\r\n"));
-		}
-		auto NoPureDeclsSuffix = TEXT("_INCLASS_IINTERFACE_NO_PURE_DECLS");
-		auto NoPureDeclsMacroName = ClassName + NoPureDeclsSuffix;
-		auto Macroized = Macroize(*NoPureDeclsMacroName, *(InterfaceBoilerplate));
-		GeneratedHeaderText.Log(*Macroized);
+			InterfaceBoilerplate.Logf(TEXT("protected:\r\n\tvirtual ~%s() {}\r\npublic:\r\n"), *InterfaceCPPName);
+			InterfaceBoilerplate.Logf(TEXT("\ttypedef %s UClassType;\r\n"), ClassCPPName);
 
-		const TCHAR* Suffix = TEXT("_INCLASS_IINTERFACE");
-		FString MacroName = ClassName + Suffix;
-		Macroized = Macroize(*MacroName, *(InterfaceBoilerplate + StaticChecks));
-		GeneratedHeaderText.Log(*Macroized);
-		InClassMacroCalls.Logf(TEXT("%s%s\r\n"), FCString::Spc(4), *MacroName);
-		InClassNoPureDeclsMacroCalls.Logf(TEXT("%s%s\r\n"), FCString::Spc(4), *NoPureDeclsMacroName);
-	}
-}
+			ExportInterfaceCallFunctions(CallbackFunctions, Class, ClassData, InterfaceBoilerplate);
 
-/**
- * Helper function to find an element in a TArray of pointers by GetFName().
- * @param Array		TArray of pointers.
- * @param Name		Name to search for.
- * @return			Index of the found element, or INDEX_NONE if not found.
- */
-template <typename T>
-int32 FindByIndirectName( const TArray<T>& Array, const FName& Name )
-{
-	for ( int32 Index=0; Index < Array.Num(); ++Index )
-	{
-		if ( Array[Index]->GetFName() == Name )
-		{
-			return Index;
-		}
-	}
-	return INDEX_NONE;
-}
-
-/**
- * Appends the header definition for an inheritance hierarchy of classes to the header.
- * Wrapper for ExportClassHeaderRecursive().
- *
- * @param	Class				The class to be exported.
- */
-void FNativeClassHeaderGenerator::ExportClassHeader( FClasses& AllClasses, FClass* Class )
-{
-	TArray<FClass*> DependencyChain;
-	TSet<const UClass*>	VisitedSet;
-	ExportClassHeaderRecursive(AllClasses, Class, DependencyChain, VisitedSet, false);
-}
-
-
-void FNativeClassHeaderGenerator::ExportClassHeaderInner(FClass* Class, bool bValidNonTemporaryClass)
-{
-	if ( !Class->HasAnyClassFlags(CLASS_Interface) )
-	{
-		TArray<UEnum*>			Enums;
-		TArray<UScriptStruct*>	Structs;
-		TArray<UFunction*>		CallbackFunctions;
-		TArray<UFunction*>		NativeFunctions;
-		TArray<UFunction*>		DelegateFunctions;
-
-		// get the lists of fields that belong to this class that should be exported
-		RetrieveRelevantFields(Class, Enums, Structs, CallbackFunctions, NativeFunctions, DelegateFunctions);
-
-		ExportNames();
-
-		// export enum declarations
-		ExportEnums(Enums);
-
-		// export delegate definitions
-		ExportDelegateDefinitions(DelegateFunctions, false);
+			// we'll need a way to get to the UObject portion of a native interface, so that we can safely pass native interfaces
+			// to script VM functions
+			if (SuperClass->IsChildOf(UInterface::StaticClass()))
+			{
+				InterfaceBoilerplate.Logf(TEXT("\tvirtual UObject* _getUObject() const = 0;\r\n"));
+			}
 
-		// Export enums declared in non-UClass headers.
-		ExportGeneratedEnumsInitCode(Enums);
+			// Replication, add in the declaration for GetLifetimeReplicatedProps() automatically if there are any net flagged properties
+			bool bNeedsRep = false;
+			FStringOutputDevice StaticChecks;
+			auto ClassName = FString(NameLookupCPP.GetNameCPP(Class));
+			for (TFieldIterator<UProperty> It(Class, EFieldIteratorFlags::ExcludeSuper); It; ++It)
+			{
+				if ((It->PropertyFlags & CPF_Net) != 0)
+				{
+					bNeedsRep = true;
+					StaticChecks.Logf(TEXT("\tstatic inline void StaticChecks()\r\n"));
+					StaticChecks.Logf(TEXT("\t{\r\n"));
+					StaticChecks.Logf(TEXT("\t\tstatic_assert(THasMemberFunction_GetLifetimeReplicatedProps<%s>::Value, \"Class %s has Net flagged properties and should declare member function: void GetLifetimeReplicatedProps( TArray< FLifetimeProperty > & OutLifetimeProps ) const override;\");\r\n"), ClassCPPName, ClassCPPName, ClassCPPName);
+					StaticChecks.Logf(TEXT("\t}\r\n"));
+					break;
+				}
+			}
 
-		// export boilerplate macros for structs
-		ExportGeneratedStructBodyMacros(Structs);
+			auto ClassRange = ClassDefinitionRange();
+			if (ClassDefinitionRanges.Contains(Class))
+			{
+				ClassRange = ClassDefinitionRanges[Class];
+			}
 
-		if (bValidNonTemporaryClass)
-		{
-			// export parameters structs for all events and delegates
-			ExportEventParms(CallbackFunctions, TEXT(""));
+			auto ClassStart = ClassRange.Start;
+			auto ClassEnd = ClassRange.End;
+			auto ClassDefinition = FString(ClassEnd - ClassStart, ClassStart);
 
-			// export .proto files for any net service functions
-			ExportProtoMessage(CallbackFunctions);
+			bool bHasGetLifetimeReplicatedProps = ClassDefinition.Contains(TEXT("GetLifetimeReplicatedProps"), ESearchCase::CaseSensitive);
 
-			// export .java files for any net service functions
-			ExportMCPMessage(CallbackFunctions);
+			if (bNeedsRep && !bHasGetLifetimeReplicatedProps)
+			{
+				InterfaceBoilerplate.Logf(TEXT("\void GetLifetimeReplicatedProps(TArray<FLifetimeProperty>& OutLifetimeProps) const;\r\n"));
+			}
 
-			// export delegate wrapper function implementations
-			ExportDelegateDefinitions(DelegateFunctions, true);
+			auto NoPureDeclsMacroName = SourceFile.GetGeneratedMacroName(ClassData, TEXT("_INCLASS_IINTERFACE_NO_PURE_DECLS"));
+			WriteMacro(GeneratedHeaderText, NoPureDeclsMacroName, InterfaceBoilerplate);
+			InClassNoPureDeclsMacroCalls.Logf(TEXT("%s%s\r\n"), FCString::Spc(4), *NoPureDeclsMacroName);
 
-			// Class definition.
-			ExportNatives(Class);
+			auto MacroName = SourceFile.GetGeneratedMacroName(ClassData, TEXT("_INCLASS_IINTERFACE"));
+			WriteMacro(GeneratedHeaderText, MacroName, InterfaceBoilerplate + StaticChecks);
+			InClassMacroCalls.Logf(TEXT("%s%s\r\n"), FCString::Spc(4), *MacroName);
 		}
-
-		ExportNativeGeneratedInitCode(Class);
-
-		if (bValidNonTemporaryClass)
+		else
 		{
-			// C++ -> VM stubs (native function execs)
-			ExportNativeFunctions(NativeFunctions);
-
-			// VM -> C++ proxies (events and delegates).
-			ExportCallbackFunctions(CallbackFunctions);
-
 			FClass* SuperClass = Class->GetSuperClass();
 
-			const TCHAR* ClassCPPName = NameLookupCPP.GetNameCPP( Class );
-			const TCHAR* SuperClassCPPName = (SuperClass != NULL) ? NameLookupCPP.GetNameCPP( SuperClass ) : NULL;
+			const TCHAR* ClassCPPName = NameLookupCPP.GetNameCPP(Class);
+			const TCHAR* SuperClassCPPName = (SuperClass != NULL) ? NameLookupCPP.GetNameCPP(SuperClass) : NULL;
 
 			{
 				FStringOutputDevice ClassBoilerplate;
 
 				// Export the class's native function registration.
 				ClassBoilerplate.Logf(TEXT("    private:\r\n"));
-				ClassBoilerplate.Logf(TEXT("    static void StaticRegisterNatives%s();\r\n"),NameLookupCPP.GetNameCPP(Class));
-				ClassBoilerplate.Log(FriendText);
-				FriendText.Empty();
+				ClassBoilerplate.Logf(TEXT("    static void StaticRegisterNatives%s();\r\n"), NameLookupCPP.GetNameCPP(Class));
+				ClassBoilerplate.Log(*FriendText);
 				ClassBoilerplate.Logf(TEXT("    public:\r\n"));
 
 				// Build the DECLARE_CLASS line
 				FString APIArg(API);
-				if( !Class->HasAnyClassFlags( CLASS_MinimalAPI ) )
+				if (!Class->HasAnyClassFlags(CLASS_MinimalAPI))
 				{
 					APIArg = TEXT("NO");
 				}
 				const bool bCastedClass = Class->HasAnyCastFlag(CASTCLASS_AllFlags) && SuperClass && Class->ClassCastFlags != SuperClass->ClassCastFlags;
-				ClassBoilerplate.Logf( TEXT("    DECLARE_CLASS(%s, %s, COMPILED_IN_FLAGS(%s%s), %s, %s, %s_API)\r\n"), 
-					ClassCPPName, 
-					SuperClassCPPName ? SuperClassCPPName : TEXT("None"), 
+				ClassBoilerplate.Logf(TEXT("    DECLARE_CLASS(%s, %s, COMPILED_IN_FLAGS(%s%s), %s, %s, %s_API)\r\n"),
+					ClassCPPName,
+					SuperClassCPPName ? SuperClassCPPName : TEXT("None"),
 					Class->HasAnyClassFlags(CLASS_Abstract) ? TEXT("CLASS_Abstract") : TEXT("0"),
-					*GetClassFlagExportText(Class), 
+					*GetClassFlagExportText(Class),
 					bCastedClass ? *FString::Printf(TEXT("CASTCLASS_%s"), ClassCPPName) : TEXT("0"),
-					*FPackageName::GetShortName(*Class->GetOuter()->GetName()), 
-					*APIArg );
+					*FPackageName::GetShortName(*Class->GetOuter()->GetName()),
+					*APIArg);
 
 				ClassBoilerplate.Logf(TEXT("    DECLARE_SERIALIZER(%s)\r\n"), ClassCPPName);
 				ClassBoilerplate.Log(TEXT("    /** Indicates whether the class is compiled into the engine */"));
 				ClassBoilerplate.Log(TEXT("    enum {IsIntrinsic=COMPILED_IN_INTRINSIC};\r\n"));
 
-				if(SuperClass && Class->ClassWithin != SuperClass->ClassWithin)
+				if (SuperClass && Class->ClassWithin != SuperClass->ClassWithin)
 				{
-					ClassBoilerplate.Logf(TEXT("    DECLARE_WITHIN(%s)\r\n"), NameLookupCPP.GetNameCPP( Class->GetClassWithin() ) );
+					ClassBoilerplate.Logf(TEXT("    DECLARE_WITHIN(%s)\r\n"), NameLookupCPP.GetNameCPP(Class->GetClassWithin()));
 				}
 
 				// export the class's config name
-				if ( SuperClass && Class->ClassConfigName != NAME_None && Class->ClassConfigName != SuperClass->ClassConfigName )
+				if (SuperClass && Class->ClassConfigName != NAME_None && Class->ClassConfigName != SuperClass->ClassConfigName)
 				{
 					ClassBoilerplate.Logf(TEXT("    static const TCHAR* StaticConfigName() {return TEXT(\"%s\");}\r\n\r\n"), *Class->ClassConfigName.ToString());
 				}
@@ -2019,9 +1761,9 @@ void FNativeClassHeaderGenerator::ExportClassHeaderInner(FClass* Class, bool bVa
 				// Replication, add in the declaration for GetLifetimeReplicatedProps() automatically if there are any net flagged properties
 				bool bHasNetFlaggedProperties = false;
 				FStringOutputDevice StaticChecks;
-				for (auto Field : TFieldRange<UProperty>(Class, EFieldIteratorFlags::ExcludeSuper))
+				for (TFieldIterator<UProperty> It(Class, EFieldIteratorFlags::ExcludeSuper); It; ++It)
 				{
-					if ((Field->PropertyFlags & CPF_Net) != 0)
+					if ((It->PropertyFlags & CPF_Net) != 0)
 					{
 						bHasNetFlaggedProperties = true;
 						StaticChecks.Logf(TEXT("\tstatic inline void StaticChecks()\r\n"));
@@ -2033,9 +1775,9 @@ void FNativeClassHeaderGenerator::ExportClassHeaderInner(FClass* Class, bool bVa
 				}
 
 				auto ClassRange = ClassDefinitionRange();
-				if (ClassDefinitionRanges.Contains(CurrentClass))
+				if (ClassDefinitionRanges.Contains(Class))
 				{
-					ClassRange = ClassDefinitionRanges[CurrentClass];
+					ClassRange = ClassDefinitionRanges[Class];
 				}
 
 				auto ClassStart = ClassRange.Start;
@@ -2049,28 +1791,57 @@ void FNativeClassHeaderGenerator::ExportClassHeaderInner(FClass* Class, bool bVa
 					ClassBoilerplate.Logf(TEXT("\tvoid GetLifetimeReplicatedProps(TArray<FLifetimeProperty>& OutLifetimeProps) const override;\r\n"));
 				}
 
-				auto ClassNameString = FString(ClassCPPName);
-				auto NoPureDeclsSuffix = TEXT("_INCLASS_NO_PURE_DECLS");
-				auto NoPureDeclsMacroName = ClassNameString + NoPureDeclsSuffix;
-				auto Macroized = Macroize(*NoPureDeclsMacroName, *(ClassBoilerplate + StaticChecks));
-				GeneratedHeaderText.Log(*Macroized);
-
-				const TCHAR* Suffix = TEXT("_INCLASS");
-				FString MacroName = ClassNameString + Suffix;
-				Macroized = Macroize(*MacroName, *ClassBoilerplate);
-				GeneratedHeaderText.Log(*Macroized);
+				auto NoPureDeclsMacroName = SourceFile.GetGeneratedMacroName(ClassData, TEXT("_INCLASS_NO_PURE_DECLS"));
+				WriteMacro(GeneratedHeaderText, NoPureDeclsMacroName, ClassBoilerplate + StaticChecks);
+				InClassNoPureDeclsMacroCalls.Logf(TEXT("%s%s\r\n"), FCString::Spc(4), *NoPureDeclsMacroName);
 
+				FString MacroName = SourceFile.GetGeneratedMacroName(ClassData, TEXT("_INCLASS"));
+				WriteMacro(GeneratedHeaderText, MacroName, ClassBoilerplate);
 				InClassMacroCalls.Logf(TEXT("%s%s\r\n"), FCString::Spc(4), *MacroName);
-				InClassNoPureDeclsMacroCalls.Logf(TEXT("%s%s\r\n"), FCString::Spc(4), *NoPureDeclsMacroName);
 
-				ExportConstructorsMacros(Class);
+				ExportConstructorsMacros(SourceFile.GetGeneratedMacroName(ClassData), Class);
 			}
 		}
+
+		{
+			FString MacroName = SourceFile.GetGeneratedMacroName(ClassData->GetPrologLine(), TEXT("_PROLOG"));
+			WriteMacro(GeneratedHeaderText, MacroName, PrologMacroCalls);
+		}
+
+		{
+			bool bIsIInterface = Class->HasAnyClassFlags(CLASS_Interface);
+
+			FString MacroName = FString::Printf(TEXT("GENERATED_%s_BODY()"), bIsIInterface ? TEXT("IINTERFACE") : TEXT("UCLASS"));
+
+			auto DeprecationWarning = bIsIInterface ? FString(TEXT("")) : GetGeneratedMacroDeprecationWarning(*MacroName);
+
+			auto DeprecationPushString = TEXT("PRAGMA_DISABLE_DEPRECATION_WARNINGS") LINE_TERMINATOR;
+			auto DeprecationPopString = TEXT("PRAGMA_POP") LINE_TERMINATOR;
+
+			auto Public = TEXT("public:" LINE_TERMINATOR);
+
+			FString BodyMacros;
+
+			if (bIsIInterface)
+			{
+				BodyMacros = Macroize(*SourceFile.GetGeneratedBodyMacroName(ClassData->GetInterfaceGeneratedBodyLine()), *(DeprecationPushString + FString(Public) + InClassMacroCalls + Public + DeprecationPopString));
+			}
+			else
+			{
+				auto StandardWrappedInClassMacroCalls = Public + InClassMacroCalls + StandardUObjectConstructorsMacroCall + Public;
+				auto EnhancedWrappedInClassMacroCalls = Public + InClassNoPureDeclsMacroCalls + EnhancedUObjectConstructorsMacroCall + GetPreservedAccessSpecifierString(Class);
+
+				BodyMacros = Macroize(*SourceFile.GetGeneratedBodyMacroName(ClassData->GetGeneratedBodyLine(), true), *(DeprecationWarning + DeprecationPushString + StandardWrappedInClassMacroCalls + DeprecationPopString)) +
+					Macroize(*SourceFile.GetGeneratedBodyMacroName(ClassData->GetGeneratedBodyLine(), false), *(FString(DeprecationPushString) + EnhancedWrappedInClassMacroCalls + DeprecationPopString));
+			}
+
+			GeneratedHeaderText.Log(*BodyMacros);
+		}
 	}
-	else
+
+	if (GPublicSourceFileSet.Contains(&SourceFile) && !ListOfPublicHeaderGroupIncludes.Contains(&SourceFile))
 	{
-		// this is an interface class
-		ExportInterfaceClassDeclaration(Class);
+		ListOfPublicHeaderGroupIncludes.Add(&SourceFile);
 	}
 }
 
@@ -2166,7 +1937,30 @@ void ExportEnhancedConstructorsMacro(FStringOutputDevice& Out, FClass* Class, co
 	ExportDefaultConstructorCallDefinition(Out, Class);
 }
 
-void FNativeClassHeaderGenerator::ExportConstructorsMacros(FClass* Class)
+FString FNativeClassHeaderGenerator::GetListOfPublicHeaderGroupIncludesString(UPackage* Package)
+{
+	FStringOutputDevice Out;
+
+	// Fill with the rest source files from this package.
+	for (auto* SourceFile : GPublicSourceFileSet)
+	{
+		if (SourceFile->GetPackage() == Package && !ListOfPublicHeaderGroupIncludes.Contains(SourceFile))
+		{
+			ListOfPublicHeaderGroupIncludes.Add(SourceFile);
+		}
+	}
+
+	for (auto* SourceFile : ListOfPublicHeaderGroupIncludes)
+	{
+		Out.Logf(TEXT("#include \"%s\"") LINE_TERMINATOR, *SourceFile->GetFilename());
+	}
+
+	Out.Log(LINE_TERMINATOR);
+
+	return Out;
+}
+
+void FNativeClassHeaderGenerator::ExportConstructorsMacros(const FString& ConstructorsMacroPrefix, FClass* Class)
 {
 	FString APIArg(API);
 	if (!Class->HasAnyClassFlags(CLASS_MinimalAPI))
@@ -2175,8 +1969,8 @@ void FNativeClassHeaderGenerator::ExportConstructorsMacros(FClass* Class)
 	}
 
 	FStringOutputDevice StdMacro, EnhMacro;
-	FString StdMacroName = FString(NameLookupCPP.GetNameCPP(Class)) + TEXT("_STANDARD_CONSTRUCTORS");
-	FString EnhMacroName = FString(NameLookupCPP.GetNameCPP(Class)) + TEXT("_ENHANCED_CONSTRUCTORS");
+	FString StdMacroName = ConstructorsMacroPrefix + TEXT("_STANDARD_CONSTRUCTORS");
+	FString EnhMacroName = ConstructorsMacroPrefix + TEXT("_ENHANCED_CONSTRUCTORS");
 
 	ExportStandardConstructorsMacro(StdMacro, Class, APIArg);
 	ExportEnhancedConstructorsMacro(EnhMacro, Class, APIArg);
@@ -2188,13 +1982,10 @@ void FNativeClassHeaderGenerator::ExportConstructorsMacros(FClass* Class)
 	EnhancedUObjectConstructorsMacroCall.Logf(TEXT("%s%s\r\n"), FCString::Spc(4), *EnhMacroName);
 }
 
-void FNativeClassHeaderGenerator::ExportClassHeaderWrapper( FClass* Class, bool bIsExportClass )
+void FNativeClassHeaderGenerator::ExportClassesFromSourceFileWrapper(FUnrealSourceFile& SourceFile)
 {
-	const bool bTemporaryJunkClass = FClassUtils::IsTemporaryClass(Class);
-	const bool bValidClass = !bTemporaryJunkClass;
-
 	check(!GeneratedHeaderText.Len());
-	ExportClassHeaderInner(Class, bValidClass);
+	ExportClassesFromSourceFileInner(SourceFile);
 
 	if (!GeneratedHeaderText.Len() && !EnumForeachText.Len())
 	{
@@ -2202,71 +1993,11 @@ void FNativeClassHeaderGenerator::ExportClassHeaderWrapper( FClass* Class, bool
 		return;
 	}
 
-	GeneratedHeaderText.Logf(TEXT("#undef UCLASS_CURRENT_FILE_NAME\r\n"));
-	GeneratedHeaderText.Logf(TEXT("#define UCLASS_CURRENT_FILE_NAME %s\r\n\r\n\r\n"), NameLookupCPP.GetNameCPP(CurrentClass));
-
-	if (bValidClass)
-	{
-		{
-			GeneratedHeaderText.Logf(TEXT("#undef UCLASS\r\n"));
-			GeneratedHeaderText.Logf(TEXT("#undef UINTERFACE\r\n"));
-			GeneratedHeaderText.Logf(TEXT("#if UE_BUILD_DOCS\r\n"));
-			FString MacroName = FString::Printf(TEXT("%s(...)"), Class->HasAnyClassFlags(CLASS_Interface) ? TEXT("UINTERFACE") : TEXT("UCLASS"));
-			GeneratedHeaderText.Logf(TEXT("#define %s\r\n"), *MacroName);
-			GeneratedHeaderText.Logf(TEXT("#else\r\n"));
-			FString Macroized = Macroize(*MacroName, *PrologMacroCalls).Replace(TEXT("\r\n\r\n\r\n"), TEXT("\r\n"));
-			GeneratedHeaderText.Log(*Macroized);
-			GeneratedHeaderText.Logf(TEXT("#endif\r\n\r\n\r\n"));
-			PrologMacroCalls.Empty();
-		}
-
-		{
-			bool bIsIInterface = Class->HasAnyClassFlags(CLASS_Interface);
-
-			GeneratedHeaderText.Logf(TEXT("#undef GENERATED_UCLASS_BODY\r\n"));
-			
-			if (!bIsIInterface)
-			{
-				GeneratedHeaderText.Logf(TEXT("#undef GENERATED_BODY\r\n"));
-			}
-
-			GeneratedHeaderText.Logf(TEXT("#undef GENERATED_IINTERFACE_BODY\r\n"));
-
-			FString MacroName = FString::Printf(TEXT("GENERATED_%s_BODY()"), bIsIInterface ? TEXT("IINTERFACE") : TEXT("UCLASS"));
-
-			auto DeprecationWarning = bIsIInterface ? FString(TEXT("")) : GetGeneratedMacroDeprecationWarning(*MacroName);
-
-			auto DeprecationPushString = TEXT("PRAGMA_DISABLE_DEPRECATION_WARNINGS") LINE_TERMINATOR;
-			auto DeprecationPopString = TEXT("PRAGMA_POP") LINE_TERMINATOR;
-
-			auto Public = TEXT("public:" LINE_TERMINATOR);
-
-			FString BodyMacros;
-			
-			if (bIsIInterface)
-			{
-				BodyMacros = Macroize(*MacroName, *(FString(Public) + InClassMacroCalls + Public));
-			}
-			else
-			{
-				auto StandardWrappedInClassMacroCalls = Public + InClassMacroCalls + StandardUObjectConstructorsMacroCall + Public;
-				auto EnhancedWrappedInClassMacroCalls = Public + InClassNoPureDeclsMacroCalls + EnhancedUObjectConstructorsMacroCall + GetPreservedAccessSpecifierString(Class);
+	GeneratedHeaderText.Logf(TEXT("#undef CURRENT_FILE_ID\r\n"));
+	GeneratedHeaderText.Logf(TEXT("#define CURRENT_FILE_ID %s\r\n\r\n\r\n"), *SourceFile.GetFileId());
 
-				BodyMacros = Macroize(*MacroName, *(DeprecationWarning + DeprecationPushString + StandardWrappedInClassMacroCalls + DeprecationPopString)) +
-					Macroize(TEXT("GENERATED_BODY()"), *(DeprecationPushString + FString(EnhancedWrappedInClassMacroCalls) + DeprecationPopString));
-			}
-
-			GeneratedHeaderText.Log(*BodyMacros);
-
-			InClassMacroCalls.Empty();
-			InClassNoPureDeclsMacroCalls.Empty();
-			StandardUObjectConstructorsMacroCall.Empty();
-			EnhancedUObjectConstructorsMacroCall.Empty();
-		}
-	}
-
-	const FString SourceFilename = GClassSourceFileMap[Class];
-	const FString PkgName        = FPackageName::GetShortName(Package);
+	const FString SourceFilename	= SourceFile.GetFilename();
+	const FString PkgName			= FPackageName::GetShortName(Package);
 
 	FString NewFileName = SourceFilename.EndsWith(TEXT(".h")) ? SourceFilename : (SourceFilename + TEXT(".h"));
 
@@ -2277,73 +2008,61 @@ void FNativeClassHeaderGenerator::ExportClassHeaderWrapper( FClass* Class, bool
 		UE_LOG(LogCompile, Error, TEXT("Failed to find path for package %s"), *PkgName);
 	}
 
-	if (bIsExportClass)
-	{
-		const FString ClassHeaderPath(GeneratedIncludeDirectory / GClassHeaderNameWithNoPathMap[Class] + TEXT(".generated.h") );
+	const FString ClassHeaderPath(GeneratedIncludeDirectory / FPaths::GetBaseFilename(SourceFile.GetFilename()) + TEXT(".generated.h"));
 
-		GeneratedHeaderText += EnumForeachText;
+	GeneratedHeaderText += EnumForeachText;
 
-		FStringOutputDevice GeneratedHeaderTextWithCopyright;
+	FStringOutputDevice GeneratedHeaderTextWithCopyright;
 
-		GeneratedHeaderTextWithCopyright.Log(
-			TEXT("// Copyright 1998-2015 Epic Games, Inc. All Rights Reserved.\r\n")
-			TEXT("/*===========================================================================\r\n")
-			TEXT("    C++ class header boilerplate exported from UnrealHeaderTool.\r\n")
-			TEXT("    This is automatically generated by the tools.\r\n")
-			TEXT("    DO NOT modify this manually! Edit the corresponding .h files instead!\r\n")
-			TEXT("===========================================================================*/\r\n"));
-		GeneratedHeaderTextWithCopyright.Log(
-			LINE_TERMINATOR
-			TEXT("#include \"ObjectBase.h\"") LINE_TERMINATOR
-			LINE_TERMINATOR);
-		GeneratedHeaderTextWithCopyright.Log(*GeneratedHeaderTextBeforeForwardDeclarations);
+	GeneratedHeaderTextWithCopyright.Log(
+		TEXT("// Copyright 1998-2015 Epic Games, Inc. All Rights Reserved.\r\n")
+		TEXT("/*===========================================================================\r\n")
+		TEXT("    C++ class header boilerplate exported from UnrealHeaderTool.\r\n")
+		TEXT("    This is automatically generated by the tools.\r\n")
+		TEXT("    DO NOT modify this manually! Edit the corresponding .h files instead!\r\n")
+		TEXT("===========================================================================*/\r\n"));
+	GeneratedHeaderTextWithCopyright.Log(
+		LINE_TERMINATOR
+		TEXT("#include \"ObjectBase.h\"") LINE_TERMINATOR
+		LINE_TERMINATOR);
+	GeneratedHeaderTextWithCopyright.Log(*GeneratedHeaderTextBeforeForwardDeclarations);
 
-		TSet<FString> ForwardDeclarationStrings;
-		for (auto Property : ForwardDeclarations)
+	TSet<FString> ForwardDeclarationStrings;
+	for (auto Property : ForwardDeclarations)
+	{
+		auto FWDecl = Property->GetCPPTypeForwardDeclaration();
+		if (FWDecl.Len() > 0)
 		{
-			auto FWDecl = Property->GetCPPTypeForwardDeclaration();
-			if (FWDecl.Len() > 0)
-			{
-				ForwardDeclarationStrings.Add(FWDecl);
-			}
+			ForwardDeclarationStrings.Add(FWDecl);
 		}
+	}
 
-		for (auto FWDeclString : ForwardDeclarationStrings)
-		{
-			GeneratedForwardDeclarations.Logf(TEXT("%s\r\n"), *FWDeclString);
-		}
+	for (auto FWDeclString : ForwardDeclarationStrings)
+	{
+		GeneratedForwardDeclarations.Logf(TEXT("%s\r\n"), *FWDeclString);
+	}
 
-		GeneratedHeaderTextWithCopyright.Log(*GeneratedForwardDeclarations);
-		GeneratedHeaderTextWithCopyright.Log(*GeneratedHeaderText);
+	GeneratedHeaderTextWithCopyright.Log(*GeneratedForwardDeclarations);
+	GeneratedHeaderTextWithCopyright.Log(*GeneratedHeaderText);
 
-		auto ClassHeaderInfo = GClassGeneratedFileMap.Find(Class);
-		check(ClassHeaderInfo != NULL);
-		ClassHeaderInfo->GeneratedFilename = ClassHeaderPath;
-		ClassHeaderInfo->bHasChanged = SaveHeaderIfChanged(*ClassHeaderPath, *GeneratedHeaderTextWithCopyright);
+	SourceFile.SetGeneratedFilename(ClassHeaderPath);
+	SourceFile.SetHasChanged(SaveHeaderIfChanged(*ClassHeaderPath, *GeneratedHeaderTextWithCopyright));
 
-		check(SourceFilename.EndsWith(TEXT(".h")));
+	check(SourceFilename.EndsWith(TEXT(".h")));
 
-		ConvertToBuildIncludePath(Package, NewFileName);
+	ConvertToBuildIncludePath(Package, NewFileName);
 
-		auto IncludeStr = FString::Printf(
-			TEXT("#ifndef %s_%s_generated_h")													LINE_TERMINATOR
-			TEXT("    #include \"%s\"")															LINE_TERMINATOR
-			TEXT("#endif")																		LINE_TERMINATOR,
-			*API, *CurrentClass->GetName(), *NewFileName );
+	auto IncludeStr = FString::Printf(
+		TEXT("#ifndef %s_%s_generated_h")													LINE_TERMINATOR
+		TEXT("    #include \"%s\"")															LINE_TERMINATOR
+		TEXT("#endif")																		LINE_TERMINATOR,
+		*API, *SourceFile.GetStrippedFilename(), *NewFileName);
 
-		// Keep track of all of the UObject headers for this module, in the same order that we digest them in
-		// @todo uht: We're wrapping these includes in checks for header guards, ONLY because most existing UObject headers are missing '#pragma once'
-		ListOfAllUObjectHeaderIncludes                .Logf(TEXT("%s"), *IncludeStr);
-		ListOfPublicClassesUObjectHeaderModuleIncludes.Logf(TEXT("%s"), *IncludeStr);
+	// Keep track of all of the UObject headers for this module, in the same order that we digest them in
+	// @todo uht: We're wrapping these includes in checks for header guards, ONLY because most existing UObject headers are missing '#pragma once'
+	ListOfAllUObjectHeaderIncludes                .Logf(TEXT("%s"), *IncludeStr);
+	ListOfPublicClassesUObjectHeaderModuleIncludes.Logf(TEXT("%s"), *IncludeStr);
 
-		// For the 'Classes' mega-header, only include legacy UObjects classes from the 'Classes' folder.  Private and Public classes are
-		// not included here -- you should include those yourself if you need them!
-		const bool bIsPublicClassesHeader = GPublicClassSet.Contains( Class );
-		if( bIsPublicClassesHeader )
-		{
-			ListOfPublicClassesUObjectHeaderGroupIncludes.Logf(TEXT("#include \"%s\"\r\n"), *NewFileName);
-		}
-	}
 	ForwardDeclarations.Empty();
 	GeneratedHeaderTextBeforeForwardDeclarations.Empty();
 	GeneratedForwardDeclarations.Empty();
@@ -2351,101 +2070,80 @@ void FNativeClassHeaderGenerator::ExportClassHeaderWrapper( FClass* Class, bool
 	EnumForeachText.Empty();
 }
 
-void FNativeClassHeaderGenerator::ExportClassHeaderRecursive( FClasses& AllClasses, FClass* Class, TArray<FClass*>& DependencyChain, TSet<const UClass*>& VisitedSet, bool bCheckDependenciesOnly )
+void FNativeClassHeaderGenerator::ExportSourceFileHeaderRecursive(FClasses& AllClasses, FUnrealSourceFile* SourceFile, TSet<const FUnrealSourceFile*>& VisitedSet, bool bCheckDependenciesOnly)
 {
-	bool bIsExportClass   = GClassStrippedHeaderTextMap.Contains(Class) && (Class->HasAnyClassFlags(CLASS_Native) || FClassUtils::IsTemporaryClass(Class)) && !Class->HasAnyClassFlags(CLASS_NoExport);
-	bool bIsCorrectHeader = Class->GetOuter() == Package;
+	bool bIsCorrectHeader = SourceFile->GetPackage() == Package;
 
 	// Check for circular header dependencies between export classes.
-	if ( bIsExportClass )
+	if (!bIsCorrectHeader)
 	{
-		if ( !bIsCorrectHeader )
+		if (VisitedSet.Num() == 0)
 		{
-			if ( DependencyChain.Num() == 0 )
-			{
-				// The first export class we found doesn't belong in this header: No need to keep exporting along this dependency path.
-				return;
-			}
-
-			// From now on, we're not going to export anything. Instead, we're going to check that no deeper dependency tries to export to this header file.
-			bCheckDependenciesOnly = true;
+			// The first export class we found doesn't belong in this header: No need to keep exporting along this dependency path.
+			return;
 		}
+
+		// From now on, we're not going to export anything. Instead, we're going to check that no deeper dependency tries to export to this header file.
+		bCheckDependenciesOnly = true;
 	}
 
 	// Check if the Class has already been exported, after we've checked for circular header dependencies.
-	if (GExportedClasses.Contains(Class))
+	if (GExportedSourceFiles.Contains(SourceFile))
 	{
 		return;
 	}
 
 	// Check for circular dependencies.
-	if (VisitedSet.Contains(Class))
+	if (VisitedSet.Contains(SourceFile))
 	{
-		UE_LOG(LogCompile, Error, TEXT("Circular dependency detected for class %s!"), *Class->GetName() );
+		UE_LOG(LogCompile, Error, TEXT("Circular dependency detected for filename %s!"), *SourceFile->GetFilename());
 		return;
 	}
 
 	// Temporarily mark the Class as VISITED. Make sure to clear this flag before returning!
-	VisitedSet.Add(Class);
+	VisitedSet.Add(SourceFile);
 
-	if ( bIsExportClass )
-	{
-		DependencyChain.Add( Class );
-	}
+	TArray<FUnrealSourceFile*> DependOnList;
 
 	// Export the super class first.
-	if (FClass* SuperClass = Class->GetSuperClass())
+	for (auto* Class : SourceFile->GetDefinedClasses())
 	{
-		ExportClassHeaderRecursive(AllClasses, SuperClass, DependencyChain, VisitedSet, bCheckDependenciesOnly);
+		if (FClass* SuperClass = ((FClass*)Class)->GetSuperClass())
+		{
+			if (SuperClass->ClassFlags & (CLASS_NoExport | CLASS_Intrinsic))
+			{
+				continue;
+			}
+
+			DependOnList.Add(&GTypeDefinitionInfoMap[SuperClass]->GetUnrealSourceFile());
+		}
 	}
 
-	// Export all classes we depend on.
-	for (const FName& DependencyClassName : *GClassDependentOnMap.FindOrAdd(Class))
+	for (auto& Include : SourceFile->GetIncludes())
 	{
-		// By this point, all classes should have been validated so we can strip the given class name when searching.		
-		FString DependencyClassNameStripped = DependencyClassName.ToString();
-
-		// Handle #include directives as 'dependson' (allowing them to fail).
-		const bool bClassNameFromHeader = FHeaderParser::DependentClassNameFromHeader(*DependencyClassNameStripped, DependencyClassNameStripped);
-		DependencyClassNameStripped = GetClassNameWithPrefixRemoved( DependencyClassNameStripped );
-
-		// Only export the class if it's in Classes array (it may be in another package).
-		FClass* FoundClass = AllClasses.FindAnyClass(*DependencyClassNameStripped);
-		if (!FoundClass)
+		auto* SourceFile = Include.Resolve();
+		if (SourceFile && !DependOnList.Contains(SourceFile))
 		{
-			// Try again with temporary class name. This may be struct only header.
-			DependencyClassNameStripped = FHeaderParser::GenerateTemporaryClassName(*GetClassNameWithoutPrefix(DependencyClassName.ToString()));
-			FoundClass                  = AllClasses.FindAnyClass(*DependencyClassNameStripped);
+			DependOnList.Add(SourceFile);
 		}
+	}
 
-		int32 ClassIndex = FindByIndirectName( Classes, *DependencyClassNameStripped );
-		if (FClass* DependsOnClass = (ClassIndex != INDEX_NONE) ? FoundClass : NULL)
-		{
-			ExportClassHeaderRecursive(AllClasses, DependsOnClass, DependencyChain, VisitedSet, bCheckDependenciesOnly);
-		}
-		else if ( !FoundClass && !bClassNameFromHeader )
-		{
-			UE_LOG(LogCompile, Error, TEXT("Unknown class %s used in dependson declaration in class %s!"), *DependencyClassName.ToString(), *Class->GetName() );
-		}
+	for (auto* DependOn : DependOnList)
+	{
+		ExportSourceFileHeaderRecursive(AllClasses, DependOn, VisitedSet, bCheckDependenciesOnly);
 	}
 
 	// Export class header.
-	if ( bIsCorrectHeader && !bCheckDependenciesOnly && !Class->HasAnyClassFlags(CLASS_Intrinsic) && (Class->HasAnyClassFlags(CLASS_Native) || FClassUtils::IsTemporaryClass(Class)) )
+	if (bIsCorrectHeader && !bCheckDependenciesOnly)
 	{
-		CurrentClass = Class;
-
 		// Mark class as exported.
-		GExportedClasses.Add(Class);
+		GExportedSourceFiles.Add(SourceFile);
 
-		ExportClassHeaderWrapper(Class, bIsExportClass);
+		ExportClassesFromSourceFileWrapper(*SourceFile);
 	}
 
 	// We're done visiting this Class.
-	VisitedSet.Remove( Class );
-	if ( bIsExportClass )
-	{
-		DependencyChain.Pop();
-	}
+	VisitedSet.Remove(SourceFile);
 }
 
 /**
@@ -2481,54 +2179,6 @@ FString FNativeClassHeaderGenerator::GetClassFlagExportText( UClass* Class )
 	return StaticClassFlagText;
 }
 
-void FNativeClassHeaderGenerator::RetrieveRelevantFields(UClass* Class, TArray<UEnum*>& Enums, TArray<UScriptStruct*>& Structs, TArray<UFunction*>& CallbackFunctions, TArray<UFunction*>& NativeFunctions, TArray<UFunction*>& DelegateFunctions)
-{
-	for ( TFieldIterator<UField> It(Class,EFieldIteratorFlags::ExcludeSuper); It; ++It )
-	{
-		UField* CurrentField = *It;
-		UClass* FieldClass = CurrentField->GetClass();
-		if ( FieldClass == UEnum::StaticClass() )
-		{
-			UEnum* Enum = (UEnum*)CurrentField;
-			Enums.Add(Enum);
-		}
-
-		else if ( FieldClass == UScriptStruct::StaticClass() )
-		{
-			UScriptStruct* Struct = (UScriptStruct*)CurrentField;
-			Structs.Add(Struct);
-		}
-
-		else if ( FieldClass == UFunction::StaticClass() )
-		{
-			bool bAdded = false;
-			UFunction* Function = (UFunction*)CurrentField;
-
-			if ( (Function->FunctionFlags&(FUNC_Event)) != 0 &&
-				Function->GetSuperFunction() == NULL )
-			{
-				CallbackFunctions.Add(Function);
-				bAdded = true;
-			}
-
-			else if ( (Function->FunctionFlags&(FUNC_Delegate)) != 0 &&
-				Function->GetSuperFunction() == NULL )
-			{
-				DelegateFunctions.Add(Function);
-				bAdded = true;
-			}
-
-			if ( (Function->FunctionFlags&FUNC_Native) != 0 )
-			{
-				NativeFunctions.Add(Function);
-				bAdded = true;
-			}
-
-			check(bAdded)
-		}
-	}
-}
-
 /**
 * Exports the header text for the list of enums specified
 *
@@ -2553,7 +2203,7 @@ void FNativeClassHeaderGenerator::ExportEnums( const TArray<UEnum*>& Enums )
 }
 
 // Exports the header text for the list of structs specified (GENERATED_USTRUCT_BODY impls)
-void FNativeClassHeaderGenerator::ExportGeneratedStructBodyMacros(const TArray<UScriptStruct*>& NativeStructs)
+void FNativeClassHeaderGenerator::ExportGeneratedStructBodyMacros(FUnrealSourceFile& SourceFile, const TArray<UScriptStruct*>& NativeStructs)
 {
 	// reverse the order.
 	for (int32 i = NativeStructs.Num() - 1; i >= 0; --i)
@@ -2565,7 +2215,7 @@ void FNativeClassHeaderGenerator::ExportGeneratedStructBodyMacros(const TArray<U
 		{
 			check(Struct->StructMacroDeclaredLineNumber != INDEX_NONE);
 
-			const FString FriendApiString = FString::Printf(TEXT("%s_API "), *API);
+			const FString FriendApiString = GetAPIString();
 			const FString StaticConstructionString = GetSingletonName(Struct);
 
 			FString RequiredAPI;
@@ -2578,8 +2228,7 @@ void FNativeClassHeaderGenerator::ExportGeneratedStructBodyMacros(const TArray<U
 			const FString StaticClassLine = FString::Printf(TEXT("    %sstatic class UScriptStruct* StaticStruct();\r\n"), *RequiredAPI);
 
 			const FString CombinedLine = FriendLine + StaticClassLine;
-			const FString MacroName = 
-				FString::Printf(TEXT("%s_USTRUCT_BODY_LINE_%d"), NameLookupCPP.GetNameCPP(CurrentClass), Struct->StructMacroDeclaredLineNumber);
+			const FString MacroName = SourceFile.GetGeneratedBodyMacroName(Struct->StructMacroDeclaredLineNumber);
 
 			const FString Macroized = Macroize(*MacroName, *CombinedLine);
 			GeneratedHeaderText.Log(*Macroized);
@@ -2595,25 +2244,25 @@ void FNativeClassHeaderGenerator::ExportGeneratedStructBodyMacros(const TArray<U
 			GeneratedPackageCPP.Logf(TEXT("    {\r\n"));
 			GeneratedPackageCPP.Logf(TEXT("        extern %sclass UScriptStruct* %s;\r\n"), *FriendApiString, *StaticConstructionString);
 			GeneratedPackageCPP.Logf(TEXT("        extern %suint32 %s();\r\n"), *FriendApiString, *GetCRCName);
-			
+
 			// UStructs can have UClass or UPackage outer (if declared in non-UClass headers).
 			if (Struct->GetOuter()->IsA(UStruct::StaticClass()))
 			{
-				GeneratedPackageCPP.Logf(TEXT("        Singleton = GetStaticStruct(%s, %s::StaticClass(), TEXT(\"%s\"), sizeof(%s), %s());\r\n"), 
+				GeneratedPackageCPP.Logf(TEXT("        Singleton = GetStaticStruct(%s, %s::StaticClass(), TEXT(\"%s\"), sizeof(%s), %s());\r\n"),
 					*SingletonName, NameLookupCPP.GetNameCPP(CastChecked<UStruct>(Struct->GetOuter())), *Struct->GetName(), StructNameCPP, *GetCRCName);
 			}
 			else
 			{
 				FString PackageSingletonName = GetPackageSingletonName(CastChecked<UPackage>(Struct->GetOuter()));
 				GeneratedPackageCPP.Logf(TEXT("        extern %sclass UPackage* %s;\r\n"), *FriendApiString, *PackageSingletonName);
-				GeneratedPackageCPP.Logf(TEXT("        Singleton = GetStaticStruct(%s, %s, TEXT(\"%s\"), sizeof(%s), %s());\r\n"), 
+				GeneratedPackageCPP.Logf(TEXT("        Singleton = GetStaticStruct(%s, %s, TEXT(\"%s\"), sizeof(%s), %s());\r\n"),
 					*SingletonName, *PackageSingletonName, *Struct->GetName(), StructNameCPP, *GetCRCName);
 			}
-			
+
 			GeneratedPackageCPP.Logf(TEXT("    }\r\n"));
 			GeneratedPackageCPP.Logf(TEXT("    return Singleton;\r\n"));
 			GeneratedPackageCPP.Logf(TEXT("}\r\n"));
-			GeneratedPackageCPP.Logf(TEXT("static FCompiledInDeferStruct Z_CompiledInDeferStruct_UScriptStruct_%s(%s::StaticStruct, TEXT(\"%s\"));\r\n"), 
+			GeneratedPackageCPP.Logf(TEXT("static FCompiledInDeferStruct Z_CompiledInDeferStruct_UScriptStruct_%s(%s::StaticStruct, TEXT(\"%s\"));\r\n"),
 				StructNameCPP, StructNameCPP, *Struct->GetOutermost()->GetName());
 
 			// Generate StaticRegisterNatives equivalent for structs without classes.
@@ -2625,12 +2274,109 @@ void FNativeClassHeaderGenerator::ExportGeneratedStructBodyMacros(const TArray<U
 				GeneratedPackageCPP.Logf(TEXT("\tFScriptStruct_%s_StaticRegisterNatives%s()\r\n"), *ShortPackageName, StructNameCPP);
 				GeneratedPackageCPP.Logf(TEXT("\t{\r\n"));
 
-				GeneratedPackageCPP.Logf( TEXT("\t\tUScriptStruct::DeferCppStructOps(FName(TEXT(\"%s\")),new UScriptStruct::TCppStructOps<%s%s>);\r\n"), *Struct->GetName(), Struct->GetPrefixCPP(), *Struct->GetName() );
+				GeneratedPackageCPP.Logf(TEXT("\t\tUScriptStruct::DeferCppStructOps(FName(TEXT(\"%s\")),new UScriptStruct::TCppStructOps<%s%s>);\r\n"), *Struct->GetName(), Struct->GetPrefixCPP(), *Struct->GetName());
 
 				GeneratedPackageCPP.Logf(TEXT("\t}\r\n"));
 				GeneratedPackageCPP.Logf(TEXT("} ScriptStruct_%s_StaticRegisterNatives%s;\r\n"), *ShortPackageName, StructNameCPP);
 			}
 		}
+
+		{
+			UScriptStruct* ScriptStruct = Struct;
+			UStruct* BaseStruct = ScriptStruct->GetSuperStruct();
+
+			const FString SingletonName(GetSingletonName(ScriptStruct));
+
+			FString ApiString = GetAPIString();
+			FString Extern = FString::Printf(TEXT("    %sclass UScriptStruct* %s;\r\n"), *ApiString, *SingletonName);
+			SingletonNameToExternDecl.Add(SingletonName, Extern);
+			GeneratedFunctionDeclarations.Log(*Extern);
+
+			FStringOutputDevice GeneratedStructRegisterFunctionText;
+
+			GeneratedStructRegisterFunctionText.Logf(TEXT("    UScriptStruct* %s\r\n"), *SingletonName);
+			GeneratedStructRegisterFunctionText.Logf(TEXT("    {\r\n"));
+
+			// if this is a no export struct, we will put a local struct here for offset determination
+			TArray<UScriptStruct*> Structs;
+			FindNoExportStructs(Structs, ScriptStruct);
+			if (Structs.Num())
+			{
+				TGuardValue<bool> Guard(bIsExportingForOffsetDeterminationOnly, true);
+				ExportMirrorsForNoexportStructs(Structs, 8, GeneratedStructRegisterFunctionText);
+			}
+
+			// Structs can either have a UClass or UPackage as outer (if delcared in non-UClass header).
+			if (ScriptStruct->GetOuter()->IsA(UStruct::StaticClass()))
+			{
+				GeneratedStructRegisterFunctionText.Logf(TEXT("        UStruct* Outer=%s;\r\n"), *GetSingletonName(CastChecked<UStruct>(ScriptStruct->GetOuter())));
+			}
+			else
+			{
+				GeneratedStructRegisterFunctionText.Logf(TEXT("        UPackage* Outer=%s;\r\n"), *GetPackageSingletonName(CastChecked<UPackage>(ScriptStruct->GetOuter())));
+			}
+			GeneratedStructRegisterFunctionText.Logf(TEXT("        static UScriptStruct* ReturnStruct = NULL;\r\n"));
+			GeneratedStructRegisterFunctionText.Logf(TEXT("        if (!ReturnStruct)\r\n"));
+			GeneratedStructRegisterFunctionText.Logf(TEXT("        {\r\n"));
+			FString BaseStructString(TEXT("NULL"));
+			if (BaseStruct)
+			{
+				CastChecked<UScriptStruct>(BaseStruct); // this better actually be a script struct
+				BaseStructString = GetSingletonName(BaseStruct);
+			}
+			FString CppStructOpsString(TEXT("NULL"));
+			FString ExplicitSizeString;
+			FString ExplicitAlignmentString;
+			if ((ScriptStruct->StructFlags&STRUCT_Native) != 0)
+			{
+				//@todo .u we don't need the auto register versions of these (except for hotreload, which should be fixed)
+				CppStructOpsString = FString::Printf(TEXT("new UScriptStruct::TCppStructOps<%s>"), NameLookupCPP.GetNameCPP(ScriptStruct));
+			}
+			else
+			{
+				ExplicitSizeString = FString::Printf(TEXT(", sizeof(%s), ALIGNOF(%s)"), NameLookupCPP.GetNameCPP(ScriptStruct), NameLookupCPP.GetNameCPP(ScriptStruct));
+			}
+
+			GeneratedStructRegisterFunctionText.Logf(TEXT("            ReturnStruct = new(Outer, TEXT(\"%s\"), RF_Public|RF_Transient|RF_Native) UScriptStruct(FObjectInitializer(), %s, %s, EStructFlags(0x%08X)%s);\r\n"),
+				*ScriptStruct->GetName(),
+				*BaseStructString,
+				*CppStructOpsString,
+				(uint32)(ScriptStruct->StructFlags & ~STRUCT_ComputedFlags),
+				*ExplicitSizeString
+				);
+			TheFlagAudit.Add(ScriptStruct, TEXT("StructFlags"), (uint64)(ScriptStruct->StructFlags & ~STRUCT_ComputedFlags));
+
+			TArray<UProperty*> Props;
+			for (TFieldIterator<UProperty> ItInner(ScriptStruct, EFieldIteratorFlags::ExcludeSuper); ItInner; ++ItInner)
+			{
+				Props.Add(*ItInner);
+			}
+			FString OuterString = FString(TEXT("ReturnStruct"));
+			FString Meta = GetMetaDataCodeForObject(ScriptStruct, *OuterString, TEXT("            "));
+			OutputProperties(Meta, GeneratedStructRegisterFunctionText, OuterString, Props, TEXT("            "));
+			GeneratedStructRegisterFunctionText.Logf(TEXT("            ReturnStruct->StaticLink();\r\n"));
+
+			if (Meta.Len())
+			{
+				GeneratedStructRegisterFunctionText.Logf(TEXT("#if WITH_METADATA\r\n"));
+				GeneratedStructRegisterFunctionText.Logf(TEXT("            UMetaData* MetaData = ReturnStruct->GetOutermost()->GetMetaData();\r\n"));
+				GeneratedStructRegisterFunctionText.Log(*Meta);
+				GeneratedStructRegisterFunctionText.Logf(TEXT("#endif\r\n"));
+			}
+
+			GeneratedStructRegisterFunctionText.Logf(TEXT("        }\r\n"));
+			GeneratedStructRegisterFunctionText.Logf(TEXT("        return ReturnStruct;\r\n"));
+			GeneratedStructRegisterFunctionText.Logf(TEXT("    }\r\n"));
+
+			uint32 StructCrc = GenerateTextCRC(*GeneratedStructRegisterFunctionText);
+			GGeneratedCodeCRCs.Add(ScriptStruct, StructCrc);
+
+			auto& GeneratedFunctionText = GetGeneratedFunctionTextDevice();
+			GeneratedFunctionText += GeneratedStructRegisterFunctionText;
+			GeneratedFunctionText.Logf(TEXT("    uint32 Get_%s_CRC() { return %uU; }\r\n"), *SingletonName.Replace(TEXT("()"), TEXT("")), StructCrc);
+
+			//CallSingletons.Logf(TEXT("                OuterClass->LinkChild(%s); // %u\r\n"), *SingletonName, StructCrc);
+		}
 	}
 }
 
@@ -2642,28 +2388,91 @@ void FNativeClassHeaderGenerator::ExportGeneratedEnumsInitCode(const TArray<UEnu
 		UEnum* Enum = Enums[i];
 
 		// Export Enum.
-		if (Enum->GetOuter()->IsA(UPackage::StaticClass()))
+		if (!Enum->GetOuter()->IsA(UPackage::StaticClass()))
 		{
-			const FString FriendApiString = FString::Printf(TEXT("%s_API "), *API);
-			const FString StaticConstructionString = GetSingletonName(Enum);
+			continue;
+		}
 
-			FString SingletonName = StaticConstructionString.Replace(TEXT("()"), TEXT("")); // function address
+		const FString FriendApiString = GetAPIString();
+		const FString StaticConstructionString = GetSingletonName(Enum);
 
-			GeneratedPackageCPP.Logf(TEXT("static class UEnum* %s_StaticEnum()\r\n"), *Enum->GetName());
-			GeneratedPackageCPP.Logf(TEXT("{\r\n"));
-			GeneratedPackageCPP.Logf(TEXT("    static class UEnum* Singleton = NULL;\r\n"));
-			GeneratedPackageCPP.Logf(TEXT("    if (!Singleton)\r\n"));
-			GeneratedPackageCPP.Logf(TEXT("    {\r\n"));
-			GeneratedPackageCPP.Logf(TEXT("        extern %sclass UEnum* %s;\r\n"), *FriendApiString, *StaticConstructionString);
+		FString SingletonName = StaticConstructionString.Replace(TEXT("()"), TEXT("")); // function address
+
+		GeneratedPackageCPP.Logf(TEXT("static class UEnum* %s_StaticEnum()\r\n"), *Enum->GetName());
+		GeneratedPackageCPP.Logf(TEXT("{\r\n"));
+		GeneratedPackageCPP.Logf(TEXT("    static class UEnum* Singleton = NULL;\r\n"));
+		GeneratedPackageCPP.Logf(TEXT("    if (!Singleton)\r\n"));
+		GeneratedPackageCPP.Logf(TEXT("    {\r\n"));
+		GeneratedPackageCPP.Logf(TEXT("        extern %sclass UEnum* %s;\r\n"), *FriendApiString, *StaticConstructionString);
 
-			FString PackageSingletonName = GetPackageSingletonName(CastChecked<UPackage>(Enum->GetOuter()));
-			GeneratedPackageCPP.Logf(TEXT("        extern %sclass UPackage* %s;\r\n"), *FriendApiString, *PackageSingletonName);
-			GeneratedPackageCPP.Logf(TEXT("        Singleton = GetStaticEnum(%s, %s, TEXT(\"%s\"));\r\n"), *SingletonName, *PackageSingletonName, *Enum->GetName());
+		FString PackageSingletonName = GetPackageSingletonName(CastChecked<UPackage>(Enum->GetOuter()));
+		GeneratedPackageCPP.Logf(TEXT("        extern %sclass UPackage* %s;\r\n"), *FriendApiString, *PackageSingletonName);
+		GeneratedPackageCPP.Logf(TEXT("        Singleton = GetStaticEnum(%s, %s, TEXT(\"%s\"));\r\n"), *SingletonName, *PackageSingletonName, *Enum->GetName());
 			
-			GeneratedPackageCPP.Logf(TEXT("    }\r\n"));
-			GeneratedPackageCPP.Logf(TEXT("    return Singleton;\r\n"));
-			GeneratedPackageCPP.Logf(TEXT("}\r\n"));
-			GeneratedPackageCPP.Logf(TEXT("static FCompiledInDeferEnum Z_CompiledInDeferEnum_UEnum_%s(%s_StaticEnum, TEXT(\"%s\"));\r\n"), *Enum->GetName(), *Enum->GetName(), *Enum->GetOutermost()->GetName());
+		GeneratedPackageCPP.Logf(TEXT("    }\r\n"));
+		GeneratedPackageCPP.Logf(TEXT("    return Singleton;\r\n"));
+		GeneratedPackageCPP.Logf(TEXT("}\r\n"));
+		GeneratedPackageCPP.Logf(TEXT("static FCompiledInDeferEnum Z_CompiledInDeferEnum_UEnum_%s(%s_StaticEnum, TEXT(\"%s\"));\r\n"), *Enum->GetName(), *Enum->GetName(), *Enum->GetOutermost()->GetName());
+
+		{
+			const FString SingletonName(GetSingletonName(Enum));
+
+			FString Extern = FString::Printf(TEXT("    %s_API class UEnum* %s;\r\n"), *API, *SingletonName);
+			SingletonNameToExternDecl.Add(SingletonName, Extern);
+			GeneratedFunctionDeclarations.Log(*Extern);
+
+			FStringOutputDevice GeneratedEnumRegisterFunctionText;
+
+			GeneratedEnumRegisterFunctionText.Logf(TEXT("    UEnum* %s\r\n"), *SingletonName);
+			GeneratedEnumRegisterFunctionText.Logf(TEXT("    {\r\n"));
+			// Enums can either have a UClass or UPackage as outer (if declared in non-UClass header).
+			if (Enum->GetOuter()->IsA(UStruct::StaticClass()))
+			{
+				GeneratedEnumRegisterFunctionText.Logf(TEXT("        UClass* Outer=%s;\r\n"), *GetSingletonName(CastChecked<UStruct>(Enum->GetOuter())));
+			}
+			else
+			{
+				GeneratedEnumRegisterFunctionText.Logf(TEXT("        UPackage* Outer=%s;\r\n"), *GetPackageSingletonName(CastChecked<UPackage>(Enum->GetOuter())));
+			}
+			GeneratedEnumRegisterFunctionText.Logf(TEXT("        static UEnum* ReturnEnum = NULL;\r\n"));
+			GeneratedEnumRegisterFunctionText.Logf(TEXT("        if (!ReturnEnum)\r\n"));
+			GeneratedEnumRegisterFunctionText.Logf(TEXT("        {\r\n"));
+
+			GeneratedEnumRegisterFunctionText.Logf(TEXT("            ReturnEnum = new(Outer, TEXT(\"%s\"), RF_Public|RF_Transient|RF_Native) UEnum(FObjectInitializer());\r\n"), *Enum->GetName());
+			GeneratedEnumRegisterFunctionText.Logf(TEXT("            TArray<FName> EnumNames;\r\n"));
+			for (int32 Index = 0; Index < Enum->NumEnums(); Index++)
+			{
+				GeneratedEnumRegisterFunctionText.Logf(TEXT("            EnumNames.Add(FName(TEXT(\"%s\")));\r\n"), *Enum->GetEnum(Index).ToString());
+			}
+
+			FString EnumTypeStr;
+			switch (Enum->GetCppForm())
+			{
+			case UEnum::ECppForm::Regular:    EnumTypeStr = TEXT("UEnum::ECppForm::Regular");    break;
+			case UEnum::ECppForm::Namespaced: EnumTypeStr = TEXT("UEnum::ECppForm::Namespaced"); break;
+			case UEnum::ECppForm::EnumClass:  EnumTypeStr = TEXT("UEnum::ECppForm::EnumClass");  break;
+			}
+			GeneratedEnumRegisterFunctionText.Logf(TEXT("            ReturnEnum->SetEnums(EnumNames, %s);\r\n"), *EnumTypeStr);
+
+			FString Meta = GetMetaDataCodeForObject(Enum, TEXT("ReturnEnum"), TEXT("            "));
+			if (Meta.Len())
+			{
+				GeneratedEnumRegisterFunctionText.Logf(TEXT("#if WITH_METADATA\r\n"));
+				GeneratedEnumRegisterFunctionText.Logf(TEXT("            UMetaData* MetaData = ReturnEnum->GetOutermost()->GetMetaData();\r\n"));
+				GeneratedEnumRegisterFunctionText.Log(*Meta);
+				GeneratedEnumRegisterFunctionText.Logf(TEXT("#endif\r\n"));
+			}
+
+			GeneratedEnumRegisterFunctionText.Logf(TEXT("        }\r\n"));
+			GeneratedEnumRegisterFunctionText.Logf(TEXT("        return ReturnEnum;\r\n"));
+			GeneratedEnumRegisterFunctionText.Logf(TEXT("    }\r\n"));
+
+			auto& GeneratedFunctionText = GetGeneratedFunctionTextDevice();
+			GeneratedFunctionText += GeneratedEnumRegisterFunctionText;
+
+			uint32 EnumCrc = GenerateTextCRC(*GeneratedEnumRegisterFunctionText);
+			GGeneratedCodeCRCs.Add(Enum, EnumCrc);
+			// CallSingletons.Logf(TEXT("                OuterClass->LinkChild(%s); // %u\r\n"), *SingletonName, EnumCrc);
 		}
 	}
 }
@@ -2692,13 +2501,6 @@ void FNativeClassHeaderGenerator::ExportMirrorsForNoexportStructs(const TArray<U
 	}
 }
 
-
-/**
-	* Exports the parameter struct declarations for the list of functions specified
-	* 
-	* @param	Function	the function that (may) have parameters which need to be exported
-	* @return	true		if the structure generated is not completely empty
-	*/
 bool FNativeClassHeaderGenerator::WillExportEventParms( UFunction* Function )
 {
 	for( TFieldIterator<UProperty> It(Function); It && (It->PropertyFlags&CPF_Parm); ++It )
@@ -2708,7 +2510,7 @@ bool FNativeClassHeaderGenerator::WillExportEventParms( UFunction* Function )
 	return false;
 }
 
-void WriteEventFunctionPrologue(FStringOutputDevice& Output, int32 Indent, const FParmsAndReturnProperties& Parameters, const TCHAR* ClassName, const TCHAR* FunctionName)
+void WriteEventFunctionPrologue(FStringOutputDevice& Output, int32 Indent, const FParmsAndReturnProperties& Parameters, UObject* FunctionOuter, const TCHAR* FunctionName)
 {
 	// now the body - first we need to declare a struct which will hold the parameters for the event/delegate call
 	Output.Logf( TEXT("\r\n%s{\r\n"), FCString::Spc(Indent) );
@@ -2717,7 +2519,7 @@ void WriteEventFunctionPrologue(FStringOutputDevice& Output, int32 Indent, const
 	if (!Parameters.HasParms())
 		return;
 
-	FString EventStructName = GetEventStructParamsName(ClassName, FunctionName);
+	FString EventStructName = GetEventStructParamsName(FunctionOuter, FunctionName);
 
 	Output.Logf( TEXT("%s%s Parms;\r\n"), FCString::Spc(Indent + 4), *EventStructName );
 
@@ -2744,7 +2546,7 @@ void WriteEventFunctionPrologue(FStringOutputDevice& Output, int32 Indent, const
 	}
 }
 
-void WriteEventFunctionEpilogue(FStringOutputDevice& Output, int32 Indent, const FParmsAndReturnProperties& Parameters, const TCHAR* ClassName, const TCHAR* FunctionName)
+void WriteEventFunctionEpilogue(FStringOutputDevice& Output, int32 Indent, const FParmsAndReturnProperties& Parameters, const TCHAR* FunctionName)
 {
 	// Out parm copying.
 	for (auto It = Parameters.Parms.CreateConstIterator(); It; ++It)
@@ -2775,38 +2577,28 @@ void WriteEventFunctionEpilogue(FStringOutputDevice& Output, int32 Indent, const
 	Output.Logf( TEXT("%s}\r\n"), FCString::Spc(Indent) );
 }
 
-/**
- * Exports C++ type definitions for delegates
- *
- * @param	DelegateFunctions	the functions that have parameters which need to be exported
- * @param	bWrapperImplementationsOnly	 True if only function implementations for delegate wrappers should be exported
- */
-void FNativeClassHeaderGenerator::ExportDelegateDefinitions( const TArray<UFunction*>& DelegateFunctions, const bool bWrapperImplementationsOnly )
+void FNativeClassHeaderGenerator::ExportDelegateDefinitions(FUnrealSourceFile& SourceFile, const TArray<UDelegateFunction*>& DelegateFunctions, const bool bWrapperImplementationsOnly)
 {
 	FStringOutputDevice HeaderOutput;
-	if( bWrapperImplementationsOnly )
-	{
-		// Export parameters structs for all delegates.  We'll need these to declare our delegate execution function.
-		ExportEventParms( DelegateFunctions, TEXT(""), 0, true, &HeaderOutput);
-	}
-
-	// find the class info for this class
-	FClassMetaData* ClassData = GScriptHelper.FindClassData(CurrentClass);
-	UClass*         Class     = CurrentClass;
-	const FString   ClassName = Class->GetName();
 
 	for ( int32 i = DelegateFunctions.Num() - 1; i >= 0 ; i-- )
 	{
 		UFunction* Function = DelegateFunctions[i];
+		FStringOutputDevice DelegateOutput;
 		check( Function->HasAnyFunctionFlags( FUNC_Delegate ) );
 
+		if (bWrapperImplementationsOnly)
+		{
+			// Export parameters structs for all delegates.  We'll need these to declare our delegate execution function.
+			ExportEventParm(Function, DelegateOutput);
+		}
+
 		const bool bIsMulticastDelegate = Function->HasAnyFunctionFlags( FUNC_MulticastDelegate );
 
 		// Unmangle the function name
 		const FString DelegateName = Function->GetName().LeftChop( FString( HEADER_GENERATED_DELEGATE_SIGNATURE_SUFFIX ).Len() );
 
-		FFunctionData* CompilerInfo = ClassData->FindFunctionData(Function);
-		check(CompilerInfo);
+		auto* CompilerInfo = FFunctionData::FindForFunction(Function);
 
 		FFuncInfo FunctionData = CompilerInfo->GetFunctionData();
 
@@ -2827,45 +2619,47 @@ void FNativeClassHeaderGenerator::ExportDelegateDefinitions( const TArray<UFunct
 			bIsMulticastDelegate ? TEXT( "FMulticastScriptDelegate" ) : TEXT( "FScriptDelegate" ),
 			*DelegateName ) );
 
+		DelegateOutput.Log(TEXT("static "));
+
 		// export the line that looks like: int32 Main(const FString& Parms)
-		ExportNativeFunctionHeader(FunctionData,HeaderOutput,EExportFunctionType::Event,EExportFunctionHeaderStyle::Declaration,*ExtraParam);
+		ExportNativeFunctionHeader(FunctionData, DelegateOutput, EExportFunctionType::Event, EExportFunctionHeaderStyle::Declaration, *ExtraParam);
 
 		if( !bWrapperImplementationsOnly )
 		{
 			// Only exporting function prototype
-			HeaderOutput.Logf( TEXT( ";\r\n" ) );
+			DelegateOutput.Logf(TEXT(";\r\n"));
+
+			ExportFunction(Function, &SourceFile.GetScope().Get(), false);
 		}
 		else
 		{
 			auto Parameters = GetFunctionParmsAndReturn(FunctionData.FunctionReference);
 
-			WriteEventFunctionPrologue(HeaderOutput, 0, Parameters, *ClassName, *DelegateName);
+			WriteEventFunctionPrologue(DelegateOutput, 0, Parameters, Function->GetOuter(), *DelegateName);
 			{
 				const TCHAR* DelegateType = bIsMulticastDelegate ? TEXT( "ProcessMulticastDelegate" ) : TEXT( "ProcessDelegate" );
 				const TCHAR* DelegateArg  = Parameters.HasParms() ? TEXT("&Parms") : TEXT("NULL");
-				HeaderOutput.Logf( TEXT("%s%s.%s<UObject>(%s);\r\n"), FCString::Spc(4), *DelegateName, DelegateType, DelegateArg );
+				DelegateOutput.Logf(TEXT("%s%s.%s<UObject>(%s);\r\n"), FCString::Spc(4), *DelegateName, DelegateType, DelegateArg);
 			}
-			WriteEventFunctionEpilogue(HeaderOutput, 0, Parameters, *ClassName, *DelegateName);
+			WriteEventFunctionEpilogue(DelegateOutput, 0, Parameters, *DelegateName);
 		}
-	}
 
-	if (HeaderOutput.Len())
-	{
-		if( !bWrapperImplementationsOnly )
+		if (!bWrapperImplementationsOnly)
 		{
-			// these are at the top, so we can output directly
-			GeneratedHeaderText.Log(*HeaderOutput);
-			GeneratedHeaderText.Log(TEXT("\r\n\r\n"));
+			HeaderOutput.Log(DelegateOutput);
 		}
 		else
 		{
-			const TCHAR* Suffix = TEXT("_DELEGATES");
-			FString MacroName = FString(NameLookupCPP.GetNameCPP(CurrentClass)) + Suffix;
-			FString Macroized = Macroize(*MacroName, *HeaderOutput);
-			GeneratedHeaderText.Log(*Macroized);
-			PrologMacroCalls.Logf( TEXT("%s%s\r\n"), FCString::Spc(0),  *MacroName);
+			FString MacroName = SourceFile.GetGeneratedMacroName(FunctionData.DelegateMacroLine, TEXT("_DELEGATE"));
+			WriteMacro(HeaderOutput, MacroName, DelegateOutput);
 		}
 	}
+
+	if (HeaderOutput.Len())
+	{
+		GeneratedHeaderText.Log(*HeaderOutput);
+		GeneratedHeaderText.Log(TEXT("\r\n\r\n"));
+	}
 }
 
 /**
@@ -2990,11 +2784,8 @@ void ExportProtoDeclaration(FOutputDevice& Out, const FString& MessageName, TFie
  * @param Indent starting indentation level
  * @param Output optional output redirect
  */
-void FNativeClassHeaderGenerator::ExportProtoMessage(const TArray<UFunction*>& InCallbackFunctions, int32 Indent, class FStringOutputDevice* Output)
+void FNativeClassHeaderGenerator::ExportProtoMessage(const TArray<UFunction*>& InCallbackFunctions, FClassMetaData* ClassData, int32 Indent, class FStringOutputDevice* Output)
 {
-	// find the class info for this class
-	FClassMetaData* ClassData = GScriptHelper.FindClassData(CurrentClass);
-
 	// Parms struct definitions.
 	FStringOutputDevice HeaderOutput;
 
@@ -3004,8 +2795,8 @@ void FNativeClassHeaderGenerator::ExportProtoMessage(const TArray<UFunction*>& I
 	for (int32 Index = 0; Index < CallbackFunctions.Num(); Index++)
 	{
 		UFunction* Function = CallbackFunctions[Index];
-		FFunctionData* CompilerInfo = ClassData->FindFunctionData(Function);
-		check(CompilerInfo);
+		auto* CompilerInfo = FFunctionData::FindForFunction(Function);
+
 		const FFuncInfo& FunctionData = CompilerInfo->GetFunctionData();
 		if (FunctionData.FunctionExportFlags & FUNCEXPORT_NeedsProto)
 		{
@@ -3222,11 +3013,8 @@ void ExportMCPDeclaration(FOutputDevice& Out, const FString& MessageName, TField
  * @param Indent starting indentation level
  * @param Output optional output redirect
  */
-void FNativeClassHeaderGenerator::ExportMCPMessage(const TArray<UFunction*>& InCallbackFunctions, int32 Indent, class FStringOutputDevice* Output)
+void FNativeClassHeaderGenerator::ExportMCPMessage(const TArray<UFunction*>& InCallbackFunctions, FClassMetaData* ClassData, int32 Indent, class FStringOutputDevice* Output)
 {
-	// find the class info for this class
-	FClassMetaData* ClassData = GScriptHelper.FindClassData(CurrentClass);
-
 	// Parms struct definitions.
 	FStringOutputDevice HeaderOutput;
 
@@ -3236,8 +3024,8 @@ void FNativeClassHeaderGenerator::ExportMCPMessage(const TArray<UFunction*>& InC
 	for (int32 Index = 0; Index < CallbackFunctions.Num(); Index++)
 	{
 		UFunction* Function = CallbackFunctions[Index];
-		FFunctionData* CompilerInfo = ClassData->FindFunctionData(Function);
-		check(CompilerInfo);
+		auto* CompilerInfo = FFunctionData::FindForFunction(Function);
+
 		const FFuncInfo& FunctionData = CompilerInfo->GetFunctionData();
 		if (FunctionData.FunctionExportFlags & FUNCEXPORT_NeedsMCP)
 		{
@@ -3281,120 +3069,110 @@ void FNativeClassHeaderGenerator::ExportMCPMessage(const TArray<UFunction*>& InC
 	}
 }
 
-/**
-* Exports the parameter struct declarations for the list of functions specified
-*
-* @param	CallbackFunctions	the functions that have parameters which need to be exported
-*/
-void FNativeClassHeaderGenerator::ExportEventParms( const TArray<UFunction*>& InCallbackFunctions, const TCHAR* MacroSuffix, int32 Indent, bool bOutputConstructor, class FStringOutputDevice* Output )
+void FNativeClassHeaderGenerator::ExportEventParm(UFunction* Function, FStringOutputDevice &HeaderOutput, int32 Indent, bool bOutputConstructor)
 {
-	// Parms struct definitions.
-	FStringOutputDevice HeaderOutput;
+	if (!WillExportEventParms(Function))
+	{
+		return;
+	}
 
-	TArray<UFunction*> CallbackFunctions = InCallbackFunctions;
-	CallbackFunctions.Sort();
+	FString FunctionName = Function->GetName();
+	if (Function->HasAnyFunctionFlags(FUNC_Delegate))
+	{
+		FunctionName = FunctionName.LeftChop(FString(HEADER_GENERATED_DELEGATE_SIGNATURE_SUFFIX).Len());
+	}
 
-	for ( int32 Index = 0; Index < CallbackFunctions.Num(); Index++ )
+	FString EventParmStructName = GetEventStructParamsName(Function->GetOuter(), *FunctionName);
+	HeaderOutput.Logf(TEXT("%sstruct %s\r\n"), FCString::Spc(Indent), *EventParmStructName);
+	HeaderOutput.Logf(TEXT("%s{\r\n"), FCString::Spc(Indent));
+
+	for (TFieldIterator<UProperty> It(Function); It && (It->PropertyFlags&CPF_Parm); ++It)
 	{
-		UFunction* Function = CallbackFunctions[Index];
-		if (!WillExportEventParms(Function))
-		{
-			continue;
-		}
+		UProperty* Prop = *It;
+		ForwardDeclarations.Add(Prop);
 
-		FString FunctionName = Function->GetName();
-		if( Function->HasAnyFunctionFlags( FUNC_Delegate ) )
-		{
-			FunctionName = FunctionName.LeftChop( FString( HEADER_GENERATED_DELEGATE_SIGNATURE_SUFFIX ).Len() );
-		}
+		FStringOutputDevice PropertyText;
+		PropertyText.Log(FCString::Spc(4 + Indent));
 
-		FString EventParmStructName = GetEventStructParamsName(*Function->GetOwnerClass()->GetName(), *FunctionName);
-		HeaderOutput.Logf( TEXT("%sstruct %s\r\n"), FCString::Spc(Indent), *EventParmStructName);
-		HeaderOutput.Logf( TEXT("%s{\r\n"), FCString::Spc(Indent) );
+		bool bEmitConst = Prop->HasAnyPropertyFlags(CPF_ConstParm) && Prop->IsA<UObjectProperty>();
 
-		for( TFieldIterator<UProperty> It(Function); It && (It->PropertyFlags&CPF_Parm); ++It )
+		//@TODO: UCREMOVAL: This is awful code duplication to avoid a double-const
 		{
-			UProperty* Prop = *It;
-			ForwardDeclarations.Add(Prop);
-			FStringOutputDevice PropertyText;
-			PropertyText.Log( FCString::Spc(4 + Indent) );
-
-			bool bEmitConst = Prop->HasAnyPropertyFlags(CPF_ConstParm) && Prop->IsA<UObjectProperty>();
+			// export 'const' for parameters
+			const bool bIsConstParam = (Prop->IsA(UInterfaceProperty::StaticClass()) && !Prop->HasAllPropertyFlags(CPF_OutParm)); //@TODO: This should be const once that flag exists
+			const bool bIsOnConstClass = (Prop->IsA(UObjectProperty::StaticClass()) && ((UObjectProperty*)Prop)->PropertyClass != NULL && ((UObjectProperty*)Prop)->PropertyClass->HasAnyClassFlags(CLASS_Const));
 
-			//@TODO: UCREMOVAL: This is awful code duplication to avoid a double-const
+			if (bIsConstParam || bIsOnConstClass)
 			{
-				// export 'const' for parameters
-				const bool bIsConstParam = (Prop->IsA(UInterfaceProperty::StaticClass()) && !Prop->HasAllPropertyFlags(CPF_OutParm)); //@TODO: This should be const once that flag exists
-				const bool bIsOnConstClass = (Prop->IsA(UObjectProperty::StaticClass()) && ((UObjectProperty*)Prop)->PropertyClass != NULL && ((UObjectProperty*)Prop)->PropertyClass->HasAnyClassFlags(CLASS_Const));
-
-				if (bIsConstParam || bIsOnConstClass)
-				{
-					bEmitConst = false; // ExportCppDeclaration will do it for us
-				}
+				bEmitConst = false; // ExportCppDeclaration will do it for us
 			}
+		}
 
-			if (bEmitConst)
-			{
-				PropertyText.Logf(TEXT("const "));
-			}
+		if (bEmitConst)
+		{
+			PropertyText.Logf(TEXT("const "));
+		}
 
-			const FString* Dim = GArrayDimensions.Find(Prop);
-			Prop->ExportCppDeclaration(PropertyText, EExportedDeclaration::Local, Dim ? **Dim : NULL);
-			ApplyAlternatePropertyExportText(Prop, PropertyText);
-			
-			PropertyText.Log( TEXT(";\r\n") );
-			HeaderOutput += *PropertyText;
+		const FString* Dim = GArrayDimensions.Find(Prop);
+		Prop->ExportCppDeclaration(PropertyText, EExportedDeclaration::Local, Dim ? **Dim : NULL);
+		ApplyAlternatePropertyExportText(Prop, PropertyText);
+
+		PropertyText.Log(TEXT(";\r\n"));
+		HeaderOutput += *PropertyText;
+
+	}
+	// constructor must initialize the return property if it needs it
+	UProperty* Prop = Function->GetReturnProperty();
+	if (Prop && bOutputConstructor)
+	{
+		FStringOutputDevice InitializationAr;
 
+		UStructProperty* InnerStruct = Cast<UStructProperty>(Prop);
+		bool bNeedsOutput = true;
+		if (InnerStruct)
+		{
+			bNeedsOutput = InnerStruct->HasNoOpConstructor();
 		}
-		// constructor must initialize the return property if it needs it
-		UProperty* Prop = Function->GetReturnProperty();
-		if (Prop && bOutputConstructor)
+		else if (
+			Cast<UNameProperty>(Prop) ||
+			Cast<UDelegateProperty>(Prop) ||
+			Cast<UMulticastDelegateProperty>(Prop) ||
+			Cast<UStrProperty>(Prop) ||
+			Cast<UTextProperty>(Prop) ||
+			Cast<UArrayProperty>(Prop) ||
+			Cast<UInterfaceProperty>(Prop)
+			)
 		{
-			FStringOutputDevice InitializationAr;
-
-			UStructProperty* InnerStruct = Cast<UStructProperty>(Prop);
-			bool bNeedsOutput = true;
-			if (InnerStruct)
-			{
-				bNeedsOutput = InnerStruct->HasNoOpConstructor();
-			}
-			else if (
-				Cast<UNameProperty>(Prop) ||
-				Cast<UDelegateProperty>(Prop) ||
-				Cast<UMulticastDelegateProperty>(Prop) ||
-				Cast<UStrProperty>(Prop) ||
-				Cast<UTextProperty>(Prop) ||
-				Cast<UArrayProperty>(Prop) ||
-				Cast<UInterfaceProperty>(Prop)
-				)
-			{
-				bNeedsOutput = false;
-			}
-			if (bNeedsOutput)
-			{
-				check(Prop->ArrayDim == 1); // can't return arrays
-				HeaderOutput.Logf(TEXT("\r\n%s/** Constructor, intializes return property only **/\r\n"), FCString::Spc(4 + Indent));
-				HeaderOutput.Logf(TEXT("%s%s()\r\n"), FCString::Spc(4 + Indent), *EventParmStructName);
-				HeaderOutput.Logf(TEXT("%s%s %s(%s)\r\n"), FCString::Spc(8 + Indent),TEXT(":") , *Prop->GetName(), *GetNullParameterValue(Prop,false,true));
-				HeaderOutput.Logf(TEXT("%s{\r\n"), FCString::Spc(4 + Indent));
-				HeaderOutput.Logf(TEXT("%s}\r\n"), FCString::Spc(4 + Indent));
-			}
+			bNeedsOutput = false;
+		}
+		if (bNeedsOutput)
+		{
+			check(Prop->ArrayDim == 1); // can't return arrays
+			HeaderOutput.Logf(TEXT("\r\n%s/** Constructor, intializes return property only **/\r\n"), FCString::Spc(4 + Indent));
+			HeaderOutput.Logf(TEXT("%s%s()\r\n"), FCString::Spc(4 + Indent), *EventParmStructName);
+			HeaderOutput.Logf(TEXT("%s%s %s(%s)\r\n"), FCString::Spc(8 + Indent), TEXT(":"), *Prop->GetName(), *GetNullParameterValue(Prop, false, true));
+			HeaderOutput.Logf(TEXT("%s{\r\n"), FCString::Spc(4 + Indent));
+			HeaderOutput.Logf(TEXT("%s}\r\n"), FCString::Spc(4 + Indent));
 		}
-		HeaderOutput.Logf( TEXT("%s};\r\n"), FCString::Spc(Indent) );
 	}
+	HeaderOutput.Logf(TEXT("%s};\r\n"), FCString::Spc(Indent));
+}
 
-	if (!Output)
-	{
-		const TCHAR* Suffix = TEXT("_EVENTPARMS");
-		FString MacroName = FString(NameLookupCPP.GetNameCPP(CurrentClass)) + Suffix + MacroSuffix;
-		FString Macroized = Macroize(*MacroName, *HeaderOutput);
-		GeneratedHeaderText.Log(*Macroized);
-		PrologMacroCalls.Logf( TEXT("%s%s\r\n"), FCString::Spc(Indent),  *MacroName);
-	}
-	else
+void FNativeClassHeaderGenerator::ExportEventParms(FScope& Scope, const TArray<UFunction*>& InCallbackFunctions, FStringOutputDevice& Output, int32 Indent, bool bOutputConstructor)
+{
+	// Parms struct definitions.
+	FStringOutputDevice HeaderOutput;
+
+	TArray<UFunction*> CallbackFunctions = InCallbackFunctions;
+	CallbackFunctions.Sort();
+
+	for ( int32 Index = 0; Index < CallbackFunctions.Num(); Index++ )
 	{
-		Output->Log(HeaderOutput);
+		UFunction* Function = CallbackFunctions[Index];
+		ExportEventParm(Function, HeaderOutput, Indent, bOutputConstructor);
 	}
+
+	Output.Log(HeaderOutput);
 }
 
 /**
@@ -3578,10 +3356,11 @@ void FNativeClassHeaderGenerator::ExportNativeFunctionHeader( const FFuncInfo& F
 
 		// If the function was marked as 'RequiredAPI', then add the *_API macro prefix.  Note that if the class itself
 		// was marked 'RequiredAPI', this is not needed as C++ will exports all methods automatically.
-		if( !Function->GetOwnerClass()->HasAnyClassFlags( CLASS_RequiredAPI ) &&
-			( FunctionData.FunctionExportFlags & FUNCEXPORT_RequiredAPI ) )
+		if (FunctionType != EExportFunctionType::Event &&
+			!Function->GetOwnerClass()->HasAnyClassFlags(CLASS_RequiredAPI) &&
+			(FunctionData.FunctionExportFlags & FUNCEXPORT_RequiredAPI))
 		{
-			HeaderOutput.Log( FString::Printf( TEXT( "%s_API " ), *API ) );
+			HeaderOutput.Log(GetAPIString());
 		}
 
 		if(FunctionType == EExportFunctionType::Interface)
@@ -3656,6 +3435,7 @@ void FNativeClassHeaderGenerator::ExportNativeFunctionHeader( const FFuncInfo& F
 	{
 		UProperty* Property = *It;
 		ForwardDeclarations.Add(Property);
+
 		if( ParmCount++ )
 		{
 			HeaderOutput.Log(TEXT(", "));
@@ -3694,7 +3474,7 @@ void FNativeClassHeaderGenerator::ExportNativeFunctionHeader( const FFuncInfo& F
  * @param Parameters list of parameters in the function
  * @param Return return parameter for the function
  */
-void FNativeClassHeaderGenerator::ExportFunctionThunk(FStringOutputDevice& RPCWrappers, const FFuncInfo& FunctionData, const TArray<UProperty*>& Parameters, UProperty* Return)
+void FNativeClassHeaderGenerator::ExportFunctionThunk(FStringOutputDevice& RPCWrappers, UFunction* Function, const FFuncInfo& FunctionData, const TArray<UProperty*>& Parameters, UProperty* Return)
 {
 	// export the GET macro for this parameter
 	FString ParameterList;
@@ -3702,6 +3482,7 @@ void FNativeClassHeaderGenerator::ExportFunctionThunk(FStringOutputDevice& RPCWr
 	{
 		UProperty* Param = Parameters[ParameterIndex];
 		ForwardDeclarations.Add(Param);
+
 		FString EvalBaseText = TEXT("P_GET_");	// e.g. P_GET_STR
 		FString EvalModifierText;				// e.g. _REF
 		FString EvalParameterText;				// e.g. (UObject*,NULL)
@@ -3801,17 +3582,16 @@ void FNativeClassHeaderGenerator::ExportFunctionThunk(FStringOutputDevice& RPCWr
 	const auto DeprecationPopString = TEXT("PRAGMA_POP") LINE_TERMINATOR;
 
 	auto ClassRange = ClassDefinitionRange();
-	if (ClassDefinitionRanges.Contains(CurrentClass))
+	if (ClassDefinitionRanges.Contains(Function->GetOwnerClass()))
 	{
-		ClassRange = ClassDefinitionRanges[CurrentClass];
+		ClassRange = ClassDefinitionRanges[Function->GetOwnerClass()];
 	}
 
 	auto ClassStart = ClassRange.Start;
 	auto ClassEnd = ClassRange.End;
 	auto ClassDefinition = FString(ClassEnd - ClassStart, ClassStart);
-	auto ClassName = CurrentClass->GetName();
+	auto ClassName = Function->GetOwnerClass()->GetName();
 
-	auto Function = FunctionData.FunctionReference;
 	auto FunctionName = Function->GetName();
 
 	bool bHasImplementation = ClassDefinition.Contains(FunctionName + TEXT("_Implementation"), ESearchCase::CaseSensitive);
@@ -3834,7 +3614,6 @@ void FNativeClassHeaderGenerator::ExportFunctionThunk(FStringOutputDevice& RPCWr
 	ExportNativeFunctionHeader(FunctionData, FunctionDeclaration, EExportFunctionType::Function, EExportFunctionHeaderStyle::Declaration);
 	FunctionDeclaration.Trim();
 
-
 	// Call the validate function if there is one
 	if (!(FunctionData.FunctionExportFlags & FUNCEXPORT_CppStatic) && (FunctionData.FunctionFlags & FUNC_NetValidate))
 	{
@@ -3871,7 +3650,7 @@ void FNativeClassHeaderGenerator::ExportFunctionThunk(FStringOutputDevice& RPCWr
 	// export the call to the C++ version
 	if (FunctionData.FunctionExportFlags & FUNCEXPORT_CppStatic)
 	{
-		RPCWrappers.Logf(TEXT("%s::%s(%s);%s"), NameLookupCPP.GetNameCPP(CurrentClass), *FunctionData.CppImplName, *ParameterList, LINE_TERMINATOR);
+		RPCWrappers.Logf(TEXT("%s::%s(%s);%s"), NameLookupCPP.GetNameCPP(Function->GetOwnerClass()), *FunctionData.CppImplName, *ParameterList, LINE_TERMINATOR);
 	}
 	else
 	{
@@ -3909,7 +3688,7 @@ FString FNativeClassHeaderGenerator::GetFunctionParameterString(UFunction* Funct
 	return ParameterList;
 }
 
-void FNativeClassHeaderGenerator::ExportNativeFunctions(const TArray<UFunction*>& NativeFunctions)
+void FNativeClassHeaderGenerator::ExportNativeFunctions(FUnrealSourceFile& SourceFile, UClass* Class, FClassMetaData* ClassData)
 {
 	FStringOutputDevice RPCWrappers;
 	FStringOutputDevice	AutogeneratedBlueprintFunctionDeclarations;
@@ -3918,31 +3697,35 @@ void FNativeClassHeaderGenerator::ExportNativeFunctions(const TArray<UFunction*>
 	FStringOutputDevice StaticChecks;
 	StaticChecks.Logf(TEXT("\tstatic inline void StaticChecks_Implementation_Validate()\r\n"));
 	StaticChecks.Logf(TEXT("\t{\r\n"));
-	
-	// find the class info for this class
-	FClassMetaData* ClassData = GScriptHelper.FindClassData(CurrentClass);
 
-	auto ClassName = FString(NameLookupCPP.GetNameCPP(CurrentClass));
 	auto ClassRange = ClassDefinitionRange();
-	if (ClassDefinitionRanges.Contains(CurrentClass))
+	if (ClassDefinitionRanges.Contains(Class))
 	{
-		ClassRange = ClassDefinitionRanges[CurrentClass];
+		ClassRange = ClassDefinitionRanges[Class];
 	}
 
 	// export the C++ stubs
-	for (auto Function : NativeFunctions)
+	
+	for (TFieldIterator<UFunction> FunctionIt(Class, EFieldIteratorFlags::ExcludeSuper); FunctionIt; ++FunctionIt)
 	{
-		FFunctionData* CompilerInfo = ClassData->FindFunctionData(Function);
-		check(CompilerInfo);
+		UFunction*     Function = *FunctionIt;
+
+		if (!(Function->FunctionFlags & FUNC_Native))
+		{
+			continue;
+		}
+
+		auto* CompilerInfo = FFunctionData::FindForFunction(Function);
 
 		const FFuncInfo& FunctionData = CompilerInfo->GetFunctionData();
 
 		// Custom thunks don't get any C++ stub function generated
 		if (FunctionData.FunctionExportFlags & FUNCEXPORT_CustomThunk)
+		{
 			continue;
+		}
 
-		// All functions that are declared in class headers are by definition native
-		check(Function->FunctionFlags & FUNC_Native);
+		FStringOutputDevice& DestinationForDecl = RPCWrappers;
 
 		// Should we emit these to RPC wrappers or just ignore them?
 		const bool bWillBeProgrammerTyped = FunctionData.CppImplName == Function->GetName();
@@ -3953,7 +3736,6 @@ void FNativeClassHeaderGenerator::ExportNativeFunctions(const TArray<UFunction*>
 			auto ClassEnd = ClassRange.End;
 			auto ClassDefinition = FString(ClassEnd - ClassStart, ClassStart);
 
-			auto Function = FunctionData.FunctionReference;
 			auto FunctionName = Function->GetName();
 
 			bool bHasImplementation = ClassDefinition.Contains(FunctionName + TEXT("_Implementation"), ESearchCase::CaseSensitive);
@@ -3969,7 +3751,6 @@ void FNativeClassHeaderGenerator::ExportNativeFunctions(const TArray<UFunction*>
 			{
 				FString ParameterList = GetFunctionParameterString(Function);
 
-				// this is not an event, the function is not a static function and the function is not marked final
 				auto Virtual = (!FunctionData.FunctionReference->HasAnyFunctionFlags(FUNC_Static) && !(FunctionData.FunctionExportFlags & FUNCEXPORT_Final)) ? TEXT("virtual") : TEXT("");
 				AutogeneratedBlueprintFunctionDeclarations.Logf(TEXT("\t%s bool %s(%s);\r\n"), Virtual, *FunctionData.CppValidationImplName, *ParameterList);
 				if (!bHasValidate)
@@ -3984,14 +3765,13 @@ void FNativeClassHeaderGenerator::ExportNativeFunctions(const TArray<UFunction*>
 				AutogeneratedBlueprintFunctionDeclarationsOnlyNotDeclared.Log(*FunctionDeclaration);
 			}
 
-			ExportFunctionChecks(FunctionData, MemberFunctionCheckClasses, StaticChecks, ClassName);
+			ExportFunctionChecks(FunctionData, MemberFunctionCheckClasses, StaticChecks, NameLookupCPP.GetNameCPP(Class));
 		}
 
 		if (bMultiLineUFUNCTION)
 		{
 			RPCWrappers.Log(TEXT("\r\n"));
 		}
-		RPCWrappers.Log(TEXT("\r\n"));
 
 		// if this function was originally declared in a base class, and it isn't a static function,
 		// only the C++ function header will be exported
@@ -4005,26 +3785,21 @@ void FNativeClassHeaderGenerator::ExportNativeFunctions(const TArray<UFunction*>
 		RPCWrappers.Logf(TEXT("%s%s{%s"), LINE_TERMINATOR, FCString::Spc(4), LINE_TERMINATOR);
 
 		auto Parameters = GetFunctionParmsAndReturn(FunctionData.FunctionReference);
-		ExportFunctionThunk(RPCWrappers, FunctionData, Parameters.Parms, Parameters.Return);
+		ExportFunctionThunk(RPCWrappers, Function, FunctionData, Parameters.Parms, Parameters.Return);
 
 		RPCWrappers.Logf(TEXT("%s}%s"), FCString::Spc(4), LINE_TERMINATOR);
 	}
 
 	StaticChecks.Logf(TEXT("\t}\r\n"));
 
-	auto Suffix = FString(TEXT("_RPC_WRAPPERS"));
-	FString MacroName = ClassName + Suffix;
-	FString Macroized = Macroize(*MacroName, *(AutogeneratedBlueprintFunctionDeclarations + RPCWrappers));
-	GeneratedHeaderText.Log(*Macroized);
+	FString MacroName = SourceFile.GetGeneratedMacroName(ClassData, TEXT("_RPC_WRAPPERS"));
+	WriteMacro(GeneratedHeaderText, MacroName, AutogeneratedBlueprintFunctionDeclarations + RPCWrappers);
 	InClassMacroCalls.Logf(TEXT("%s%s\r\n"), FCString::Spc(4), *MacroName);
 
-	auto NoPureDeclsSuffix = FString(TEXT("_RPC_WRAPPERS_NO_PURE_DECLS"));
-	MacroName = ClassName + NoPureDeclsSuffix;
-
 	// Put static checks before RPCWrappers to get proper messages from static asserts before compiler errors.
-	Macroized = Macroize(*MacroName, *(AutogeneratedBlueprintFunctionDeclarationsOnlyNotDeclared + MemberFunctionCheckClasses + StaticChecks + RPCWrappers));
-	GeneratedHeaderText.Log(*Macroized);
-	InClassNoPureDeclsMacroCalls.Logf(TEXT("%s%s\r\n"), FCString::Spc(4), *MacroName);
+	FString NoPureDeclsMacroName = SourceFile.GetGeneratedMacroName(ClassData, TEXT("_RPC_WRAPPERS_NO_PURE_DECLS"));
+	WriteMacro(GeneratedHeaderText, NoPureDeclsMacroName, AutogeneratedBlueprintFunctionDeclarationsOnlyNotDeclared + MemberFunctionCheckClasses + StaticChecks + RPCWrappers);
+	InClassNoPureDeclsMacroCalls.Logf(TEXT("%s%s\r\n"), FCString::Spc(4), *NoPureDeclsMacroName);
 }
 
 /**
@@ -4032,37 +3807,61 @@ void FNativeClassHeaderGenerator::ExportNativeFunctions(const TArray<UFunction*>
  *
  * @param	CallbackFunctions	the functions to export
  */
-void FNativeClassHeaderGenerator::ExportCallbackFunctions( const TArray<UFunction*>& InCallbackFunctions )
+TArray<UFunction*> FNativeClassHeaderGenerator::ExportCallbackFunctions(FUnrealSourceFile& SourceFile, UClass* Class, FClassMetaData* ClassData)
 {
+	TArray<UFunction*> CallbackFunctions;
+
+	for (TFieldIterator<UFunction> FunctionIt(Class, EFieldIteratorFlags::ExcludeSuper); FunctionIt; ++FunctionIt)
+	{
+		UFunction*     Function = *FunctionIt;
+
+		if ((Function->FunctionFlags & FUNC_Event) && Function->GetSuperFunction() == nullptr)
+		{
+			CallbackFunctions.Add(Function);
+		}
+	}
+
 	FStringOutputDevice RPCWrappers;
 	FStringOutputDevice RPCWrappersCPP;
-	// find the class info for this class
-	FClassMetaData* ClassData = GScriptHelper.FindClassData(CurrentClass);
 
-	// Skip interfaces; they don't have any callbacks
-	if (CurrentClass->HasAnyClassFlags(CLASS_Interface))
+	if (CallbackFunctions.Num() == 0)
 	{
-		return;
+		// Early out.
+		return CallbackFunctions;
 	}
 
-	TArray<UFunction*> CallbackFunctions = InCallbackFunctions;
 	CallbackFunctions.Sort();
 
-	for ( int32 CallbackIndex = 0; CallbackIndex < CallbackFunctions.Num(); CallbackIndex++ )
+	FStringOutputDevice UClassMacroContent;
+
+	// export parameters structs for all events and delegates
+	ExportEventParms(SourceFile.GetScope().Get(), CallbackFunctions, UClassMacroContent, 4, true);
+
+	FString MacroName = SourceFile.GetGeneratedMacroName(ClassData, TEXT("_EVENT_PARMS"));
+	WriteMacro(GeneratedHeaderText, MacroName, UClassMacroContent);
+	PrologMacroCalls.Logf(TEXT("%s%s\r\n"), FCString::Spc(4), *MacroName);
+
+	// export .proto files for any net service functions
+	ExportProtoMessage(CallbackFunctions, ClassData);
+
+	// export .java files for any net service functions
+	ExportMCPMessage(CallbackFunctions, ClassData);
+
+	for (int32 CallbackIndex = 0; CallbackIndex < CallbackFunctions.Num(); CallbackIndex++)
 	{
-		UFunction* Function = CallbackFunctions[ CallbackIndex ];
+		UFunction* Function = CallbackFunctions[CallbackIndex];
 
 		// Never expecting to export delegate functions this way
-		check( !Function->HasAnyFunctionFlags( FUNC_Delegate ) );
-
-		FClass* Class = CurrentClass;
+		check(!Function->HasAnyFunctionFlags(FUNC_Delegate));
 
-		FFunctionData* CompilerInfo = ClassData->FindFunctionData(Function);
-		check(CompilerInfo);
+		auto* CompilerInfo = FFunctionData::FindForFunction(Function);
 
 		// cache the TCHAR* for a few strings we'll use a lot here
 		FString FunctionName = Function->GetName();
 
+		GeneratedHeaderText.Logf(TEXT("extern %s FName %s_%s;") LINE_TERMINATOR, *GetAPIString(), *API, *FunctionName);
+		ReferencedNames.Add(Function->GetFName());
+
 		const FFuncInfo& FunctionData = CompilerInfo->GetFunctionData();
 
 		if (FunctionData.FunctionFlags & FUNC_NetResponse)
@@ -4071,8 +3870,6 @@ void FNativeClassHeaderGenerator::ExportCallbackFunctions( const TArray<UFunctio
 			continue;
 		}
 
-		FString ClassName = Class->GetName();
-
 		const bool bWillBeProgrammerTyped = Function->GetName() == FunctionData.MarshallAndCallName;
 
 		// Emit the declaration if the programmer isn't responsible for declaring this wrapper
@@ -4094,33 +3891,54 @@ void FNativeClassHeaderGenerator::ExportCallbackFunctions( const TArray<UFunctio
 
 		auto Parameters = GetFunctionParmsAndReturn(FunctionData.FunctionReference);
 
-		WriteEventFunctionPrologue(RPCWrappersCPP, 4, Parameters, *ClassName, *FunctionName);
+		if (!(Class->ClassFlags & CLASS_Interface))
 		{
-			// Cast away const just in case, because ProcessEvent isn't const
-			RPCWrappersCPP.Logf(
-				TEXT("%s%sProcessEvent(FindFunctionChecked(%s_%s),%s);\r\n"),
-				FCString::Spc(8),
-				(Function->HasAllFunctionFlags(FUNC_Const)) ? *FString::Printf(TEXT("const_cast<%s*>(this)->"), NameLookupCPP.GetNameCPP(Class)) : TEXT(""),
-				*API,
-				*FunctionName,
-				Parameters.HasParms() ? TEXT("&Parms") : TEXT("NULL")
-			);
+			WriteEventFunctionPrologue(RPCWrappersCPP, 4, Parameters, Function->GetOuter(), *FunctionName);
+			{
+				// Cast away const just in case, because ProcessEvent isn't const
+				RPCWrappersCPP.Logf(
+					TEXT("%s%sProcessEvent(FindFunctionChecked(%s_%s),%s);\r\n"),
+					FCString::Spc(8),
+					(Function->HasAllFunctionFlags(FUNC_Const)) ? *FString::Printf(TEXT("const_cast<%s*>(this)->"), NameLookupCPP.GetNameCPP(Cast<UClass>(Function->GetOuter()))) : TEXT(""),
+					*API,
+					*FunctionName,
+					Parameters.HasParms() ? TEXT("&Parms") : TEXT("NULL")
+					);
+			}
+			WriteEventFunctionEpilogue(RPCWrappersCPP, 4, Parameters, *FunctionName);
+		}
+		else
+		{
+			RPCWrappersCPP.Logf(LINE_TERMINATOR);
+			RPCWrappersCPP.Logf(TEXT("%s{%s"), FCString::Spc(4), LINE_TERMINATOR);
+
+			// assert if this is ever called directly
+			RPCWrappersCPP.Logf(TEXT("%scheck(0 && \"Do not directly call Event functions in Interfaces. Call Execute_%s instead.\");%s"), FCString::Spc(8), *FunctionName, LINE_TERMINATOR);
+
+			// satisfy compiler if it's expecting a return value
+			if (Parameters.Return)
+			{
+				FString EventParmStructName = GetEventStructParamsName(Function->GetOuter(), *FunctionName);
+				RPCWrappersCPP.Logf(TEXT("%s%s Parms;%s"), FCString::Spc(8), *EventParmStructName, LINE_TERMINATOR);
+				RPCWrappersCPP.Logf(TEXT("%sreturn Parms.ReturnValue;%s"), FCString::Spc(8), LINE_TERMINATOR);
+			}
+			RPCWrappersCPP.Logf(TEXT("%s}%s"), FCString::Spc(4), LINE_TERMINATOR);
 		}
-		WriteEventFunctionEpilogue(RPCWrappersCPP, 4, Parameters, *ClassName, *FunctionName);
 	}
 
-	if (!FClassUtils::IsNoExportOrTemporaryClass(CurrentClass))
+	if (!Class->HasAnyClassFlags(CLASS_NoExport))
 	{
 		GeneratedPackageCPP.Log(*RPCWrappersCPP);
 	}
 	// else drop the implementation on the floor
 
-	const TCHAR* Suffix = TEXT("_CALLBACK_WRAPPERS");
-	FString MacroName = FString(NameLookupCPP.GetNameCPP(CurrentClass)) + Suffix;
-	FString Macroized = Macroize(*MacroName, *RPCWrappers);
-	GeneratedHeaderText.Log(*Macroized);
-	InClassMacroCalls.Logf( TEXT("%s%s\r\n"), FCString::Spc(4),  *MacroName);
-	InClassNoPureDeclsMacroCalls.Logf(TEXT("%s%s\r\n"), FCString::Spc(4), *MacroName);
+	FString CallbackWrappersMacroName = SourceFile.GetGeneratedMacroName(ClassData, TEXT("_CALLBACK_WRAPPERS"));
+	WriteMacro(GeneratedHeaderText, CallbackWrappersMacroName, RPCWrappers);
+
+	InClassMacroCalls.Logf(TEXT("%s%s\r\n"), FCString::Spc(4), *CallbackWrappersMacroName);
+	InClassNoPureDeclsMacroCalls.Logf(TEXT("%s%s\r\n"), FCString::Spc(4), *CallbackWrappersMacroName);
+
+	return CallbackFunctions;
 }
 
 
@@ -4134,80 +3952,26 @@ void FNativeClassHeaderGenerator::ExportCallbackFunctions( const TArray<UFunctio
  */
 void FNativeClassHeaderGenerator::ApplyAlternatePropertyExportText( UProperty* Prop, FStringOutputDevice& PropertyText )
 {
-	if (bIsExportingForOffsetDeterminationOnly)
-	{
-		UDelegateProperty* DelegateProperty = Cast<UDelegateProperty>(Prop);
-		UMulticastDelegateProperty* MulticastDelegateProperty = Cast<UMulticastDelegateProperty>(Prop);
-		if (DelegateProperty || MulticastDelegateProperty)
-		{
-			FString Original = Prop->GetCPPType();
-			FString PlaceholderOfSameSizeAndAlignemnt;
-			if (DelegateProperty)
-			{
-				PlaceholderOfSameSizeAndAlignemnt = TEXT("FScriptDelegate");
-			}
-			else
-			{
-				PlaceholderOfSameSizeAndAlignemnt = TEXT("FMulticastScriptDelegate");
-			}
-			PropertyText.ReplaceInline(*Original, *PlaceholderOfSameSizeAndAlignemnt, ESearchCase::CaseSensitive);
-		}
-	}
-	// find the class info for this class
-	FClassMetaData* ClassData = GScriptHelper.FindClassData(CurrentClass);
-	// find the compiler token for this property
-	FTokenData* PropData = ClassData->FindTokenData(Prop);
-
-	UObjectProperty* ObjectProp = Cast<UObjectProperty>(Prop);
-	if ( ObjectProp == NULL )
+	if (!bIsExportingForOffsetDeterminationOnly)
 	{
-		UArrayProperty* ArrayProp = Cast<UArrayProperty>(Prop);
-		if ( ArrayProp != NULL )
-		{
-			ObjectProp = Cast<UObjectProperty>(ArrayProp->Inner);
-		}
-	}
-	if ( ObjectProp != NULL && PropData && PropData->Token.ExportInfo.Len() > 0)
-	{
-		PropertyText.ReplaceInline(*(FString("class ") + NameLookupCPP.GetNameCPP(ObjectProp->PropertyClass) + TEXT("*")), *PropData->Token.ExportInfo);
 		return;
 	}
 
-	UStructProperty* StructProp = Cast<UStructProperty>(Prop);
-	if ( StructProp == NULL )
+	UDelegateProperty* DelegateProperty = Cast<UDelegateProperty>(Prop);
+	UMulticastDelegateProperty* MulticastDelegateProperty = Cast<UMulticastDelegateProperty>(Prop);
+	if (DelegateProperty || MulticastDelegateProperty)
 	{
-		UArrayProperty* ArrayProp = Cast<UArrayProperty>(Prop);
-		if ( ArrayProp != NULL )
+		FString Original = Prop->GetCPPType();
+		FString PlaceholderOfSameSizeAndAlignemnt;
+		if (DelegateProperty)
 		{
-			StructProp = Cast<UStructProperty>(ArrayProp->Inner);
+			PlaceholderOfSameSizeAndAlignemnt = TEXT("FScriptDelegate");
 		}
-	}
-
-	if ( StructProp )
-	{
-		if ( PropData != NULL )
+		else
 		{
-			FString ExportText;
-			if ( PropData->Token.ExportInfo.Len() > 0 )
-			{
-				ExportText = PropData->Token.ExportInfo;
-			}
-			else
-			{
-				// we didn't have any export-text associated with the variable, so check if we have
-				// anything for the struct itself
-				FStructData* StructData = ClassData->FindStructData(StructProp->Struct);
-				if ( StructData != NULL && StructData->StructData.ExportInfo.Len() )
-				{
-					ExportText = StructData->StructData.ExportInfo;
-				}
-			}
-
-			if ( ExportText.Len() > 0 )
-			{
-				(FString&)PropertyText = PropertyText.Replace(NameLookupCPP.GetNameCPP(StructProp->Struct), *ExportText);
-			}
+			PlaceholderOfSameSizeAndAlignemnt = TEXT("FMulticastScriptDelegate");
 		}
+		PropertyText.ReplaceInline(*Original, *PlaceholderOfSameSizeAndAlignemnt, ESearchCase::CaseSensitive);
 	}
 }
 
@@ -4295,12 +4059,9 @@ bool FNativeClassHeaderGenerator::FindInterDependencyRecursive( TMap<const FStri
 	return false;
 }
 
-bool IsExportOrTemporaryClass(FClass* Class)
+bool IsExportClass(FClass* Class)
 {
-	bool bIsTemporaryClass = FClassUtils::IsTemporaryClass(Class);
-	bool bIsExportClass = Class->HasAnyClassFlags(CLASS_Native) && !(Class->HasAnyClassFlags(CLASS_NoExport | CLASS_Intrinsic) || bIsTemporaryClass);
-
-	return bIsExportClass || bIsTemporaryClass;
+	return Class->HasAnyClassFlags(CLASS_Native) && !Class->HasAnyClassFlags(CLASS_NoExport | CLASS_Intrinsic);
 }
 
 FStringOutputDevice& FNativeClassHeaderGenerator::GetGeneratedFunctionTextDevice()
@@ -4316,9 +4077,8 @@ FStringOutputDevice& FNativeClassHeaderGenerator::GetGeneratedFunctionTextDevice
 }
 
 // Constructor.
-FNativeClassHeaderGenerator::FNativeClassHeaderGenerator( UPackage* InPackage, FClasses& AllClasses, bool InAllowSaveExportedHeaders )
-	: CurrentClass                          (NULL)
-	, API                                   (FPackageName::GetShortName(InPackage).ToUpper())
+FNativeClassHeaderGenerator::FNativeClassHeaderGenerator(UPackage* InPackage, const TArray<FUnrealSourceFile*>& SourceFiles, FClasses& AllClasses, bool InAllowSaveExportedHeaders)
+	: API                                   (FPackageName::GetShortName(InPackage).ToUpper())
 	, Package                               (InPackage)
 	, bIsExportingForOffsetDeterminationOnly(false)
 	, bAllowSaveExportedHeaders             (InAllowSaveExportedHeaders)
@@ -4331,17 +4091,7 @@ FNativeClassHeaderGenerator::FNativeClassHeaderGenerator( UPackage* InPackage, F
 	GeneratedMCPFilenameBase = PackageName + TEXT(".generated");
 	bFailIfGeneratedCodeChanges = FParse::Param(FCommandLine::Get(), TEXT("FailIfGeneratedCodeChanges"));
 
-	// Tag native classes in this package for export.
-
-	Classes = AllClasses.GetClassesInPackage();
-
-	TMap<FName,UClass*> ClassNameMap;
-	for ( int32 ClassIndex = 0; ClassIndex < Classes.Num(); ClassIndex++ )
-	{
-		ClassNameMap.Add(Classes[ClassIndex]->GetFName(), Classes[ClassIndex]);
-	}
-
-	bool bPackageHasAnyExportOrTemporaryClasses = AllClasses.GetClassesInPackage(Package).ContainsByPredicate(IsExportOrTemporaryClass);
+	bool bPackageHasAnyExportClasses = AllClasses.GetClassesInPackage(Package).ContainsByPredicate(IsExportClass);
 	bool bHasNamesForExport = false;
 	TempHeaderPaths.Empty();
 	PackageHeaderPaths.Empty();
@@ -4356,7 +4106,7 @@ FNativeClassHeaderGenerator::FNativeClassHeaderGenerator( UPackage* InPackage, F
 	GeneratedFunctionBodyTextSplit                .Empty();
 	ListOfPublicClassesUObjectHeaderModuleIncludes.Empty();
 
-	if (bPackageHasAnyExportOrTemporaryClasses)
+	if (bPackageHasAnyExportClasses)
 	{
 		FString PkgDir;
 		FString GeneratedIncludeDirectory;
@@ -4368,24 +4118,23 @@ FNativeClassHeaderGenerator::FNativeClassHeaderGenerator( UPackage* InPackage, F
 		ClassesHeaderPath = GeneratedIncludeDirectory / ClassesHeaderName;
 
 		int32 ClassCount = 0;
-		for (FClass* Class : Classes)
+		for (auto* SourceFile : SourceFiles)
 		{
-			if( Class->GetOuter()==Package && (Class->HasAnyClassFlags(CLASS_Native) || FClassUtils::IsTemporaryClass(Class)))
+			for (auto* Class : SourceFile->GetDefinedClasses())
 			{
-				if (GClassStrippedHeaderTextMap.Contains(Class) && !Class->HasAnyClassFlags(CLASS_NoExport))
+				if (Class->HasAnyClassFlags(CLASS_Native))
 				{
-					ClassCount++;
-					// Skip temporary classes autogenerated in struct-only headers.
-					if (!FClassUtils::IsTemporaryClass(Class))
+					if (GTypeDefinitionInfoMap.Contains(Class) && !Class->HasAnyClassFlags(CLASS_NoExport))
 					{
+						ClassCount++;
 						Class->UnMark(OBJECTMARK_TagImp);
 						Class->Mark(OBJECTMARK_TagExp);
 					}
 				}
-			}
-			else
-			{
-				Class->UnMark(EObjectMark(OBJECTMARK_TagImp | OBJECTMARK_TagExp));
+				else
+				{
+					Class->UnMark(EObjectMark(OBJECTMARK_TagImp | OBJECTMARK_TagExp));
+				}
 			}
 		}
 
@@ -4393,8 +4142,7 @@ FNativeClassHeaderGenerator::FNativeClassHeaderGenerator( UPackage* InPackage, F
 		{
 			ClassesHeaders.Add(ClassesHeaderName);
 
-			CurrentClass = NULL;
-			ListOfPublicClassesUObjectHeaderGroupIncludes.Empty();
+			ListOfPublicHeaderGroupIncludes.Empty();
 			ListOfAllUObjectHeaderIncludes.Empty();
 			OriginalHeader.Empty();
 			PreHeaderText.Empty();
@@ -4424,21 +4172,16 @@ FNativeClassHeaderGenerator::FNativeClassHeaderGenerator( UPackage* InPackage, F
 			}
 
 			// Export an include line for each header
-			for (FClass* Class : Classes)
+			for (auto* SourceFile : SourceFiles)
 			{
-				if (Class->GetOuter() == Package)
-				{
-					ExportClassHeader(AllClasses, Class);
-				}
+				ExportSourceFileHeader(AllClasses, SourceFile);
 			}
 
-			ListOfPublicClassesUObjectHeaderGroupIncludes.Logf(TEXT("\r\n"));
-
 			// build the full header file out of its pieces
 			const FString FullClassesHeader = FString::Printf(
 				TEXT("%s\r\n%s"),
 				*PreHeaderText,
-				*ListOfPublicClassesUObjectHeaderGroupIncludes
+				*GetListOfPublicHeaderGroupIncludesString(InPackage)
 				);
 
 			// Save the classes header if it has changed.
@@ -4446,15 +4189,12 @@ FNativeClassHeaderGenerator::FNativeClassHeaderGenerator( UPackage* InPackage, F
 		}
 	}
 
-	if (!bPackageHasAnyExportOrTemporaryClasses)
+	if (!bPackageHasAnyExportClasses)
 	{
 		// If no headers are generated, check if there's any no export classes that need to have the inl file generated.
-		for (FClass* Class : Classes)
+		for (auto* SourceFile : SourceFiles)
 		{
-			if (Class->GetOuter() == Package && !(Class->HasAnyClassFlags(CLASS_Intrinsic) || FClassUtils::IsTemporaryClass(Class)))
-			{
-				ExportClassHeader(AllClasses, Class);
-			}
+			ExportSourceFileHeader(AllClasses, SourceFile);
 		}
 	}
 
@@ -5020,15 +4760,11 @@ ECompilationResult::Type UnrealHeaderTool_Main(const FString& ModuleInfoFilename
 
 			for (const FString& Filename : UObjectHeaders)
 			{
-				// Best faith effort at a useful line number for errors occurring during this initial pre-parsing
-				int32 ClassDeclLine = -1;
-
 			#if !PLATFORM_EXCEPTIONS_DISABLED
 				try
 			#endif
 				{
 					// Import class.
-					const FString ClassName      = FPaths::GetBaseFilename(Filename);
 					const FString FullModulePath = FPaths::ConvertRelativePathToFull(ModuleInfoPath, Filename);
 
 					FString HeaderFile;
@@ -5037,14 +4773,18 @@ ECompilationResult::Type UnrealHeaderTool_Main(const FString& ModuleInfoFilename
 						FError::Throwf(TEXT( "UnrealHeaderTool was unable to load source file '%s'"), *FullModulePath);
 					}
 
-					UClass* ResultClass = GenerateCodeForHeader(Package, *ClassName, RF_Public|RF_Standalone, *HeaderFile, ClassDeclLine);
-					GClassSourceFileMap.Add(ResultClass, Filename);
-					GClassDeclarationLineNumber.Add(ResultClass, ClassDeclLine);
-					GClassGeneratedFileMap.Add(ResultClass, FClassHeaderInfo(IFileManager::Get().ConvertToAbsolutePathForExternalAppForRead(*Filename)));
+					TSharedRef<FUnrealSourceFile> UnrealSourceFile = GenerateCodeForHeader(Package, Filename, RF_Public | RF_Standalone, *HeaderFile);
 
-					if( CurrentlyProcessing == PublicClassesHeaders )
+					GUnrealSourceFilesMap.Add(Filename, UnrealSourceFile);
+
+					if (CurrentlyProcessing == PublicClassesHeaders)
 					{
-						GPublicClassSet.Add(ResultClass);
+						for (auto* Class : UnrealSourceFile->GetDefinedClasses())
+						{
+							GPublicClassSet.Add(Class);
+						}
+
+						GPublicSourceFileSet.Add(&UnrealSourceFile.Get());
 					}
 
 					// Save metadata for the class path, both for it's include path and relative to the module base directory
@@ -5052,7 +4792,8 @@ ECompilationResult::Type UnrealHeaderTool_Main(const FString& ModuleInfoFilename
 					{
 						// Get the path relative to the module directory
 						const TCHAR *ModuleRelativePath = *FullModulePath + Module.BaseDirectory.Len();
-						GClassModuleRelativePathMap.Add(ResultClass, ModuleRelativePath);
+
+						UnrealSourceFile->SetModuleRelativePath(ModuleRelativePath);
 
 						// Calculate the include path
 						const TCHAR *IncludePath = ModuleRelativePath;
@@ -5083,7 +4824,7 @@ ECompilationResult::Type UnrealHeaderTool_Main(const FString& ModuleInfoFilename
 						// Add the include path
 						if(*IncludePath != 0)
 						{
-							GClassIncludePathMap.Add(ResultClass, IncludePath);
+							UnrealSourceFile->SetIncludePath(MoveTemp(IncludePath));
 						}
 					}
 				}
@@ -5093,11 +4834,9 @@ ECompilationResult::Type UnrealHeaderTool_Main(const FString& ModuleInfoFilename
 					TGuardValue<ELogTimes::Type> DisableLogTimes(GPrintLogTimes, ELogTimes::None);
 
 					FString Prefix;
-					if (ClassDeclLine != -1)
-					{
-						const FString AbsFilename = IFileManager::Get().ConvertToAbsolutePathForExternalAppForRead(*Filename);
-						Prefix = FString::Printf(TEXT("%s(%i): "), *AbsFilename, ClassDeclLine);
-					}
+
+					const FString AbsFilename = IFileManager::Get().ConvertToAbsolutePathForExternalAppForRead(*Filename);
+					Prefix = FString::Printf(TEXT("%s(1): "), *AbsFilename);
 
 					FString FormattedErrorMessage = FString::Printf(TEXT("%sError: %s\r\n"), *Prefix, ErrorMsg);
 					Result = GCompilationResult;
@@ -5125,23 +4864,23 @@ ECompilationResult::Type UnrealHeaderTool_Main(const FString& ModuleInfoFilename
 		{
 			const UClass* ScriptSuperClass = ScriptClass->GetSuperClass();
 
-			if (ScriptSuperClass && !ScriptSuperClass->HasAnyClassFlags(CLASS_Intrinsic) && GClassStrippedHeaderTextMap.Contains(ScriptClass) && !GClassStrippedHeaderTextMap.Contains(ScriptSuperClass))
+			if (ScriptSuperClass && !ScriptSuperClass->HasAnyClassFlags(CLASS_Intrinsic) && GTypeDefinitionInfoMap.Contains(ScriptClass) && !GTypeDefinitionInfoMap.Contains(ScriptSuperClass))
 			{
 				class FSuperClassContextSupplier : public FContextSupplier
 				{
 				public:
-					FSuperClassContextSupplier(const UClass* Class) :
-						ScriptClass(Class)
+					FSuperClassContextSupplier(const UClass* Class)
+						: DefinitionInfo(GTypeDefinitionInfoMap[Class])
 					{ }
 
 					virtual FString GetContext() override
 					{
-						FString Filename = IFileManager::Get().ConvertToAbsolutePathForExternalAppForRead(*GClassSourceFileMap[ScriptClass]);
-						int32 LineNumber = GClassDeclarationLineNumber[ScriptClass];
+						FString Filename = IFileManager::Get().ConvertToAbsolutePathForExternalAppForRead(*DefinitionInfo->GetUnrealSourceFile().GetFilename());
+						int32 LineNumber = DefinitionInfo->GetLineNumber();
 						return FString::Printf(TEXT("%s(%i)"), *Filename, LineNumber);
 					}
 				private:
-					const UClass* ScriptClass;
+					TSharedRef<FUnrealTypeDefinitionInfo> DefinitionInfo;
 				} ContextSupplier(ScriptClass);
 
 				auto OldContext = GWarn->GetContext();
@@ -5225,50 +4964,19 @@ ECompilationResult::Type UnrealHeaderTool_Main(const FString& ModuleInfoFilename
 	return Result;
 }
 
-UClass* GenerateCodeForHeader
-(
-	UObject*     InParent,
-	const TCHAR* Name,
-	EObjectFlags Flags,
-	const TCHAR* Buffer,
-	int32&       OutClassDeclLine
-)
+UClass* ProcessParsedClass(bool bClassIsAnInterface, TArray<FHeaderProvider> &DependentOn, const FString& ClassName, const FString& BaseClassName, UObject* InParent, EObjectFlags Flags)
 {
-	// Support for headers without UClasses.
-	bool bNonClassHeader = false;
-	const TCHAR* InBuffer = Buffer;
-
-	// Import the script text.
-	TArray<FName> DependentOn;
-
-	// is the parsed class name an interface?
-	bool bClassIsAnInterface = false;
-
-	// Parse the header to extract the information needed
-	FStringOutputDevice ClassHeaderTextStrippedOfCppText;
-	FString ClassName;
-	FString BaseClassName;
-	FHeaderParser::SimplifiedClassParse(Buffer, /*out*/ bClassIsAnInterface, /*out*/ DependentOn, /*out*/ ClassName, /*out*/ BaseClassName, /*out*/ OutClassDeclLine, ClassHeaderTextStrippedOfCppText);
-
-	// In case no UClass is defined, generate the default temporary UClass.
-	if (ClassName.IsEmpty())
-	{
-		ClassName = FString::Printf(TEXT("U%s"), Name);
-		BaseClassName = TEXT("UObject");
-		bNonClassHeader = true;
-	}
-
-	FString ClassNameStripped = GetClassNameWithPrefixRemoved( *ClassName );
+	FString ClassNameStripped = GetClassNameWithPrefixRemoved(*ClassName);
 
 	// Ensure the base class has any valid prefix and exists as a valid class. Checking for the 'correct' prefix will occur during compilation
 	FString BaseClassNameStripped;
 	if (!BaseClassName.IsEmpty())
 	{
 		BaseClassNameStripped = GetClassNameWithPrefixRemoved(BaseClassName);
-		if( !FHeaderParser::ClassNameHasValidPrefix(BaseClassName, BaseClassNameStripped) )
-			FError::Throwf(TEXT("No prefix or invalid identifier for base class %s.\nClass names must match Unreal prefix specifications (e.g., \"UObject\" or \"AActor\")"), *BaseClassName );
+		if (!FHeaderParser::ClassNameHasValidPrefix(BaseClassName, BaseClassNameStripped))
+			FError::Throwf(TEXT("No prefix or invalid identifier for base class %s.\nClass names must match Unreal prefix specifications (e.g., \"UObject\" or \"AActor\")"), *BaseClassName);
 
-		if (DependentOn.ContainsByPredicate([&](const FName& Dependency){ FString DependencyStr = Dependency.ToString(); return !DependencyStr.Contains(TEXT(".generated.h")) && FPaths::GetBaseFilename(DependencyStr) == ClassNameStripped; }))
+		if (DependentOn.ContainsByPredicate([&](const FHeaderProvider& Dependency){ FString DependencyStr = Dependency.GetId(); return !DependencyStr.Contains(TEXT(".generated.h")) && FPaths::GetBaseFilename(DependencyStr) == ClassNameStripped; }))
 			FError::Throwf(TEXT("Class '%s' contains a dependency (#include or DependsOn) to itself"), *ClassName);
 	}
 
@@ -5276,25 +4984,19 @@ UClass* GenerateCodeForHeader
 	// Handle failure and non-class headers.
 	if (BaseClassName.IsEmpty() && (ClassName != TEXT("UObject")))
 	{
-		FError::Throwf(TEXT("Class '%s' must inherit UObject or a UObject-derived class"), *ClassName );
+		FError::Throwf(TEXT("Class '%s' must inherit UObject or a UObject-derived class"), *ClassName);
 	}
 
 	if (ClassName == BaseClassName)
 	{
-		FError::Throwf(TEXT("Class '%s' cannot inherit from itself"), *ClassName );
+		FError::Throwf(TEXT("Class '%s' cannot inherit from itself"), *ClassName);
 	}
 
 	// In case the file system and the class disagree on the case of the
 	// class name replace the fname with the one from the script class file
 	// This is needed because not all source control systems respect the
 	// original filename's case
-	FName ClassNameReplace(*ClassName,FNAME_Replace_Not_Safe_For_Threading);
-
-	// All classes must start with a valid unreal prefix
-	if (!FHeaderParser::ClassNameHasValidPrefix(ClassName, Name) || !FHeaderParser::ClassNameHasValidPrefix(ClassNameReplace.ToString(), Name))
-	{
-		FError::Throwf(TEXT("Invalid class name '%s'. The class name must match the name of the class header file it is contained in ('%s'), with an appropriate prefix added (A for Actors, U for other classes)"), *ClassNameReplace.ToString(), Name);
-	}
+	FName ClassNameReplace(*ClassName, FNAME_Replace_Not_Safe_For_Threading);
 
 	// Use stripped class name for processing and replace as we did above
 	FName ClassNameStrippedReplace(*ClassNameStripped, FNAME_Replace_Not_Safe_For_Threading);
@@ -5305,47 +5007,20 @@ UClass* GenerateCodeForHeader
 
 	const static bool bVerboseOutput = FParse::Param(FCommandLine::Get(), TEXT("VERBOSE"));
 
-	// Gracefully update an existing hardcoded class.
-	// Note: Most 'native' classes will not have RF_Native set at this point, only classes in Core will go down this path
-	if ((ResultClass != NULL) && ResultClass->HasAnyFlags(RF_Native) )
-	{
-		UClass* SuperClass = ResultClass->GetSuperClass();
-		if ((SuperClass != NULL) && (SuperClass->GetName() != BaseClassNameStripped))
-		{
-			// the code that handles the DependsOn list in the script compiler doesn't work correctly if we manually add the Object class to a class's DependsOn list
-			// if Object is also the class's parent.  The only way this can happen (since specifying a parent class in a DependsOn statement is a compiler error) is
-			// in this block of code, so just handle that here rather than trying to make the script compiler handle this gracefully
-			if (BaseClassNameStripped != TEXT("Object"))
-			{
-				// we're changing the parent of a native class, which may result in the
-				// child class being parsed before the new parent class, so add the new
-				// parent class to this class's DependsOn() array to guarantee that it
-				// will be parsed before this class
-				DependentOn.AddUnique(*BaseClassNameStripped);
-			}
-
-			// if the new parent class is an existing native class, attempt to change the parent for this class to the new class 
-			UClass* NewSuperClass = FindObject<UClass>(ANY_PACKAGE, *BaseClassNameStripped);
-			if (NewSuperClass != NULL)
-			{
-				check(0); // we don't support changing base clases anymore
-			}
-		}
-	}
-	else
+	if (ResultClass == nullptr || !ResultClass->HasAnyFlags(RF_Native))
 	{
 		// detect if the same class name is used in multiple packages
-		if (ResultClass == NULL)
+		if (ResultClass == nullptr)
 		{
 			UClass* ConflictingClass = FindObject<UClass>(ANY_PACKAGE, *ClassNameStripped, true);
-			if (ConflictingClass != NULL)
+			if (ConflictingClass != nullptr)
 			{
 				UE_LOG(LogCompile, Warning, TEXT("Duplicate class name: %s also exists in file %s"), *ClassName, *ConflictingClass->GetOutermost()->GetName());
 			}
 		}
 
 		// Create new class.
-		ResultClass = new( InParent, *ClassNameStripped, Flags ) UClass( FObjectInitializer(),NULL );
+		ResultClass = new(InParent, *ClassNameStripped, Flags) UClass(FObjectInitializer(), nullptr);
 		GClassHeaderNameWithNoPathMap.Add(ResultClass, ClassNameStripped);
 
 		// add CLASS_Interface flag if the class is an interface
@@ -5356,37 +5031,38 @@ UClass* GenerateCodeForHeader
 			ResultClass->ClassFlags |= CLASS_Interface;
 		}
 
-		// In case there actually was no class, mark it as UHT temporary.
-		if (bNonClassHeader)
+		if (bVerboseOutput)
 		{
-			const FString DummyClassName = FHeaderParser::GenerateTemporaryClassName(*ResultClass->GetName());
-
-			ResultClass->ClassFlags |= CLASS_Native;
-			FClassUtils::MarkAsTemporaryClass(ResultClass);
-			ResultClass->Rename(*DummyClassName, NULL, REN_DontCreateRedirectors);
+			UE_LOG(LogCompile, Log, TEXT("Imported: %s"), *ResultClass->GetFullName());
 		}
+	}
 
-		// Find or forward-declare base class.
-		ResultClass->SetSuperStruct( FindObject<UClass>( InParent, *BaseClassNameStripped ) );
-		if (ResultClass->GetSuperStruct() == NULL)
+	if (ResultClass != nullptr)
+	{
+		if (!BaseClassNameStripped.IsEmpty() && !ResultClass->GetSuperClass())
 		{
-			ResultClass->SetSuperStruct( FindObject<UClass>( ANY_PACKAGE, *BaseClassNameStripped ) );
-		}
+			// Find or forward-declare base class.
+			ResultClass->SetSuperStruct(FindObject<UClass>(InParent, *BaseClassNameStripped));
+			if (ResultClass->GetSuperStruct() == nullptr)
+			{
+				ResultClass->SetSuperStruct(FindObject<UClass>(ANY_PACKAGE, *BaseClassNameStripped));
+			}
 
-		if (ResultClass->GetSuperStruct() == NULL)
-		{
-			// don't know its parent class yet
-			ResultClass->SetSuperStruct( new(InParent, *BaseClassNameStripped) UClass(FObjectInitializer(),NULL) );
-		}
+			if (ResultClass->GetSuperStruct() == nullptr)
+			{
+				// don't know its parent class yet
+				ResultClass->SetSuperStruct(new(InParent, *BaseClassNameStripped) UClass(FObjectInitializer(), nullptr));
+			}
 
-		if (ResultClass->GetSuperStruct() != NULL)
-		{
-			ResultClass->ClassCastFlags |= ResultClass->GetSuperClass()->ClassCastFlags;
+			if (ResultClass->GetSuperStruct() != nullptr)
+			{
+				ResultClass->ClassCastFlags |= ResultClass->GetSuperClass()->ClassCastFlags;
+			}
 		}
 
-		if ( bVerboseOutput )
+		if (BaseClassNameStripped != TEXT("Object") && ResultClass->GetSuperClass())
 		{
-			UE_LOG(LogCompile, Log, TEXT("Imported: %s"), *ResultClass->GetFullName() );
+			DependentOn.AddUnique(FHeaderProvider(EHeaderProviderSourceType::ClassName, *BaseClassNameStripped));
 		}
 	}
 
@@ -5398,9 +5074,47 @@ UClass* GenerateCodeForHeader
 		}
 	}
 
-	// Set class info.
-	GClassStrippedHeaderTextMap.Emplace(ResultClass, MoveTemp(ClassHeaderTextStrippedOfCppText));
-	GClassDependentOnMap       .Emplace(ResultClass, MoveTemp(DependentOn));
-
 	return ResultClass;
 }
+
+
+TSharedRef<FUnrealSourceFile> GenerateCodeForHeader
+(
+	UPackage*     InParent,
+	const FString& FileName,
+	EObjectFlags Flags,
+	const TCHAR* Buffer
+)
+{
+	const TCHAR* InBuffer = Buffer;
+
+	// is the parsed class name an interface?
+	bool bClassIsAnInterface = false;
+
+	TArray<FHeaderProvider> DependsOn;
+
+	// Parse the header to extract the information needed
+	FStringOutputDevice ClassHeaderTextStrippedOfCppText;
+	TArray<FSimplifiedParsingClassInfo> ParsedClassArray;
+	FHeaderParser::SimplifiedClassParse(Buffer, /*out*/ ParsedClassArray, /*out*/ DependsOn, ClassHeaderTextStrippedOfCppText);
+
+	TSharedRef<FUnrealSourceFile> UnrealSourceFile = MakeShareable(new FUnrealSourceFile(InParent, FileName, ClassHeaderTextStrippedOfCppText));
+
+	for (auto& ParsedClassInfo : ParsedClassArray)
+	{
+		UClass* ResultClass = ProcessParsedClass(ParsedClassInfo.IsInterface(), DependsOn, ParsedClassInfo.GetClassName(), ParsedClassInfo.GetBaseClassName(), InParent, Flags);
+
+		FScope::AddTypeScope(ResultClass, &UnrealSourceFile->GetScope().Get());
+
+		AddTypeDefinition(*UnrealSourceFile, ResultClass, ParsedClassInfo.GetClassDefLine());
+		UnrealSourceFile->AddDefinedClass(ResultClass);
+	}
+
+	for (const auto& DependsOnElement : DependsOn)
+	{
+		UnrealSourceFile->GetIncludes().Add(DependsOnElement);
+	}
+
+	return UnrealSourceFile;
+}
+
