UHT: Allows defining multiple UCLASSes in one header file.

diff --git a/Engine/Source/Programs/UnrealHeaderTool/Private/HeaderParser.cpp b/Engine/Source/Programs/UnrealHeaderTool/Private/HeaderParser.cpp
--- a/Engine/Source/Programs/UnrealHeaderTool/Private/HeaderParser.cpp
+++ b/Engine/Source/Programs/UnrealHeaderTool/Private/HeaderParser.cpp
@@ -17,13 +17,37 @@
 	Constants & declarations.
 -----------------------------------------------------------------------------*/
 
-// Annotation of classes that have had an error during parsing
-static FUObjectAnnotationSparseBool FailedClassesAnnotation;
+/**
+ * Data struct that annotates source files that failed during parsing.
+ */
+class FFailedFilesAnnotation
+{
+public:
+	/**
+	 * Gets annotation state for given source file.
+	 */
+	bool Get(FUnrealSourceFile* SourceFile) const
+	{
+		return AnnotatedSet.Contains(SourceFile);
+	}
+
+	/**
+	 * Sets annotation state to true for given source file.
+	 */
+	void Set(FUnrealSourceFile* SourceFile)
+	{
+		AnnotatedSet.Add(SourceFile);
+	}
+
+private:
+	// Annotation set.
+	TSet<FUnrealSourceFile*> AnnotatedSet;
+} static FailedFilesAnnotation;
 
 enum {MAX_ARRAY_SIZE=2048};
 
 static const FName NAME_ToolTip(TEXT("ToolTip"));
-TMap<FClass*, ClassDefinitionRange> ClassDefinitionRanges;
+TMap<UClass*, ClassDefinitionRange> ClassDefinitionRanges;
 /**
  * Dirty hack global variable to allow different result codes passed through
  * exceptions. Needs to be fixed in future versions of UHT.
@@ -369,13 +393,9 @@ FScriptLocation::FScriptLocation()
 
 FString FHeaderParser::GetContext()
 {
-	// Return something useful if we didn't even get as far as class parsing
-	if (!Class)
-	{
-		return Filename;
-	}
+	auto* FileScope = GetCurrentFileScope();
 
-	FString Filename = IFileManager::Get().ConvertToAbsolutePathForExternalAppForRead(*GClassSourceFileMap[Class]);
+	FString Filename = IFileManager::Get().ConvertToAbsolutePathForExternalAppForRead(*FileScope->GetSourceFile()->GetFilename());
 
 	return FString::Printf(TEXT("%s(%i)"), *Filename, InputLine);
 }
@@ -463,29 +483,6 @@ UField* FHeaderParser::FindField
 	return NULL;
 }
 
-// Check if a field obscures a field in an outer scope.
-void FHeaderParser::CheckObscures(UStruct* Scope, const FString& ScriptName, const FString& FieldName)
-{
-	UEnum* FoundEnum = NULL;
-	if (UEnum::LookupEnumName(*ScriptName, &FoundEnum) != INDEX_NONE)
-	{
-		FError::Throwf(TEXT("'%s' obscures a value in enumeration '%s'"), *FoundEnum->GetPathName());
-	}
-
-	UStruct* BaseScope = Scope->GetInheritanceSuper();
-	while (BaseScope != NULL)
-	{
-		int32 OuterContextCount = 0;
-		UField* Existing = FindField(BaseScope, *FieldName, false, UField::StaticClass(), NULL);
-		if (Existing != NULL)
-		{
-			FError::Throwf(TEXT("'%s' obscures '%s' defined in %s class '%s'."), *ScriptName, *FieldName, OuterContextCount ? TEXT("outer") : TEXT("base"), *Existing->GetOuter()->GetName());
-		}
-		BaseScope = BaseScope->GetInheritanceSuper();
-	}
-}
-
-
 /**
  * @return	true if Scope has UProperty objects in its list of fields
  */
@@ -522,11 +519,6 @@ UStruct* FHeaderParser::GetSuperScope( UStruct* CurrentScope, const FName& Searc
 			{
 				SuperScope = NextScope;
 			}
-			// if we reached the top and the scope is not a class, try the current class hierarchy next
-			else if (Cast<UClass>(SuperScope) == NULL)
-			{
-				SuperScope = Class;
-			}
 			else
 			{
 				// otherwise we've failed
@@ -538,6 +530,49 @@ UStruct* FHeaderParser::GetSuperScope( UStruct* CurrentScope, const FName& Searc
 	return SuperScope;
 }
 
+/**
+ * Adds source file's include path to given metadata.
+ *
+ * @param Type Type for which to add include path.
+ * @param MetaData Meta data to fill the information.
+ */
+void AddIncludePathToMetadata(UField* Type, TMap<FName, FString> &MetaData)
+{
+	// Add metadata for the include path.
+	auto* TypeDefinitionPtr = GTypeDefinitionInfoMap.Find(Type);
+	if (TypeDefinitionPtr != nullptr)
+	{
+		MetaData.Add(TEXT("IncludePath"), *(*TypeDefinitionPtr)->GetUnrealSourceFile().GetIncludePath());
+	}
+}
+
+/**
+ * Adds module's relative path from given file.
+ *
+ * @param SourceFile Given source file.
+ * @param MetaData Meta data to fill the information.
+ */
+void AddModuleRelativePathToMetadata(FUnrealSourceFile& SourceFile, TMap<FName, FString> &MetaData)
+{
+	MetaData.Add(TEXT("ModuleRelativePath"), *SourceFile.GetModuleRelativePath());
+}
+
+/**
+ * Adds module's relative path to given metadata.
+ *
+ * @param Type Type for which to add module's relative path.
+ * @param MetaData Meta data to fill the information.
+ */
+void AddModuleRelativePathToMetadata(UField* Type, TMap<FName, FString> &MetaData)
+{
+	// Add metadata for the module relative path.
+	auto* TypeDefinitionPtr = GTypeDefinitionInfoMap.Find(Type);
+	if (TypeDefinitionPtr != nullptr)
+	{
+		MetaData.Add(TEXT("ModuleRelativePath"), *(*TypeDefinitionPtr)->GetUnrealSourceFile().GetModuleRelativePath());
+	}
+}
+
 /*-----------------------------------------------------------------------------
 	Variables.
 -----------------------------------------------------------------------------*/
@@ -545,9 +580,9 @@ UStruct* FHeaderParser::GetSuperScope( UStruct* CurrentScope, const FName& Searc
 //
 // Compile an enumeration definition.
 //
-UEnum* FHeaderParser::CompileEnum(UClass* Scope)
+UEnum* FHeaderParser::CompileEnum(FUnrealSourceFile& SourceFile)
 {
-	check(Scope);
+	auto Scope = SourceFile.GetScope();
 
 	CheckAllow( TEXT("'Enum'"), ALLOW_TypeDecl );
 
@@ -601,19 +636,17 @@ UEnum* FHeaderParser::CompileEnum(UClass* Scope)
 	}
 
 	// Verify that the enumeration definition is unique within this scope.
-	UField* Existing = FindField(Scope, EnumToken.Identifier);
-	if (Existing && Existing->GetOuter() == Scope)
+	auto* Existing = Scope->FindTypeByName(EnumToken.Identifier);
+	if (Existing)
 	{
 		FError::Throwf(TEXT("enum: '%s' already defined here"), *EnumToken.TokenName.ToString());
 	}
 
-	CheckObscures(Scope, EnumToken.Identifier, EnumToken.Identifier);
 	ParseFieldMetaData(EnumToken.MetaData, EnumToken.Identifier);
 
 	// Create enum definition.
-	UEnum* Enum = new(FClassUtils::IsTemporaryClass(Scope) ? Scope->GetOuter() : Scope, EnumToken.Identifier, RF_Public) UEnum(FObjectInitializer());
-	Enum->Next = Scope->Children;
-	Scope->Children = Enum;
+	UEnum* Enum = new(SourceFile.GetPackage(), EnumToken.Identifier, RF_Public) UEnum(FObjectInitializer());
+	Scope->AddType(Enum);
 
 	// Validate the metadata for the enum
 	ValidateMetaDataFormat(Enum, EnumToken.MetaData);
@@ -669,13 +702,7 @@ UEnum* FHeaderParser::CompileEnum(UClass* Scope)
 	// List of all metadata generated for this enum
 	TMap<FName,FString> EnumValueMetaData = EnumToken.MetaData;
 
-	// Add metadata for the module relative path
-	const FString *ModuleRelativePath = GClassModuleRelativePathMap.Find(Scope);
-	if(ModuleRelativePath != NULL)
-	{
-		EnumValueMetaData.Add(TEXT("ModuleRelativePath"), *ModuleRelativePath);
-	}
-
+	AddModuleRelativePathToMetadata(Enum, EnumValueMetaData);
 	AddFormattedPrevCommentAsTooltipMetaData(EnumValueMetaData);
 
 	// Parse all enums tags.
@@ -1102,9 +1129,9 @@ EAccessSpecifier FHeaderParser::ParseAccessProtectionSpecifier(FToken& Token)
 /**
  * Compile a struct definition.
  */
-UScriptStruct* FHeaderParser::CompileStructDeclaration(FClasses& AllClasses, FClass* Scope)
+UScriptStruct* FHeaderParser::CompileStructDeclaration(FClasses& AllClasses, FUnrealSourceFile& SourceFile)
 {
-	check(Scope);
+	auto Scope = SourceFile.GetScope();
 
 	// Make sure structs can be declared here.
 	CheckAllow( TEXT("'struct'"), ALLOW_TypeDecl );//@TODO: UCREMOVAL: After the switch: Make this require global scope
@@ -1117,13 +1144,6 @@ UScriptStruct* FHeaderParser::CompileStructDeclaration(FClasses& AllClasses, FCl
 	uint32 StructFlags = STRUCT_Native;
 	TMap<FName, FString> MetaData;
 
-	// Add metadata for the module relative path
-	const FString *ModuleRelativePath = GClassModuleRelativePathMap.Find(Scope);
-	if(ModuleRelativePath != NULL)
-	{
-		MetaData.Add(TEXT("ModuleRelativePath"), *ModuleRelativePath);
-	}
-
 	// Get the struct specifier list
 	TArray<FPropertySpecifier> SpecifiersFound;
 	ReadSpecifierSetInsideMacro(SpecifiersFound, TEXT("Struct"), MetaData);
@@ -1151,21 +1171,10 @@ UScriptStruct* FHeaderParser::CompileStructDeclaration(FClasses& AllClasses, FCl
 		StructFlags |= STRUCT_RequiredAPI;
 	}
 
-	// Handle the forced naming inconsistency (field names for structs, etc... do *not* have their prefix on them)
-	bool bOverrideStructName = false;
-	FString CustomExportText;
-	if (StructsWithNoPrefix.Contains(StructNameInScript))
-	{
-		CustomExportText = StructNameInScript;
-	}
-	else
-	{
-		bOverrideStructName = true;
-		StructNameStripped = GetClassNameWithPrefixRemoved(StructNameInScript);
-	}
+	StructNameStripped = GetClassNameWithPrefixRemoved(StructNameInScript);
 
 	// Effective struct name
-	const FString EffectiveStructName = bOverrideStructName ? *StructNameStripped : *StructNameInScript;
+	const FString EffectiveStructName = *StructNameStripped;
 
 	// Process the list of specifiers
 	for (TArray<FPropertySpecifier>::TIterator SpecifierIt(SpecifiersFound); SpecifierIt; ++SpecifierIt)
@@ -1176,6 +1185,7 @@ UScriptStruct* FHeaderParser::CompileStructDeclaration(FClasses& AllClasses, FCl
 		{
 			//UE_LOG(LogCompile, Warning, TEXT("Struct named %s in %s is still marked noexport"), *EffectiveStructName, *(Class->GetName()));//@TODO: UCREMOVAL: Debug printing
 			StructFlags &= ~STRUCT_Native;
+			StructFlags |= STRUCT_NoExport;
 		}
 		else if (Specifier == TEXT("Atomic"))
 		{
@@ -1185,7 +1195,7 @@ UScriptStruct* FHeaderParser::CompileStructDeclaration(FClasses& AllClasses, FCl
 		{
 			StructFlags |= STRUCT_Immutable | STRUCT_Atomic;
 
-			if (Class != UObject::StaticClass())
+			if (!FPaths::IsSamePath(Filename, GTypeDefinitionInfoMap[UObject::StaticClass()]->GetUnrealSourceFile().GetFilename()))
 			{
 				FError::Throwf(TEXT("Immutable is being phased out in favor of SerializeNative, and is only legal on the mirror structs declared in UObject"));
 			}
@@ -1197,20 +1207,17 @@ UScriptStruct* FHeaderParser::CompileStructDeclaration(FClasses& AllClasses, FCl
 	}
 
 	// Verify uniqueness (if declared within UClass).
-	if (!FClassUtils::IsTemporaryClass(Scope))
 	{
-		UField* Existing = FindField(Scope, *EffectiveStructName);
-		if (Existing && Existing->GetOuter() == Scope)
+		auto* Existing = Scope->FindTypeByName(*EffectiveStructName);
+		if (Existing)
 		{
 			FError::Throwf(TEXT("struct: '%s' already defined here"), *EffectiveStructName);
 		}
 
-		if (FindObject<UClass>(ANY_PACKAGE, *EffectiveStructName) != NULL)
+		if (FindObject<UStruct>(ANY_PACKAGE, *EffectiveStructName) != NULL)
 		{
 			FError::Throwf(TEXT("struct: '%s' conflicts with class name"), *EffectiveStructName);
 		}
-
-		CheckObscures(Scope, StructNameInScript, EffectiveStructName);
 	}
 
 	// Get optional superstruct.
@@ -1228,7 +1235,7 @@ UScriptStruct* FHeaderParser::CompileStructDeclaration(FClasses& AllClasses, FCl
 		FToken ParentScope, ParentName;
 		if (GetIdentifier( ParentScope ))
 		{
-			UStruct* StructClass = Scope;
+			TSharedRef<FScope> StructScope = Scope;
 			FString ParentStructNameInScript = FString(ParentScope.Identifier);
 			if (MatchSymbol(TEXT(".")))
 			{
@@ -1236,7 +1243,7 @@ UScriptStruct* FHeaderParser::CompileStructDeclaration(FClasses& AllClasses, FCl
 				{
 					ParentStructNameInScript = FString(ParentName.Identifier);
 					FString ParentNameStripped = GetClassNameWithPrefixRemoved(ParentScope.Identifier);
-					StructClass = AllClasses.FindClass(*ParentNameStripped);
+					FClass* StructClass = AllClasses.FindClass(*ParentNameStripped);
 					if( !StructClass )
 					{
 						// If we find the literal class name, the user didn't use a prefix
@@ -1250,6 +1257,8 @@ UScriptStruct* FHeaderParser::CompileStructDeclaration(FClasses& AllClasses, FCl
 							FError::Throwf(TEXT("'struct': Can't find parent struct class '%s'"), ParentScope.Identifier );
 						}
 					}
+
+					StructScope = FScope::GetTypeScope(StructClass);
 				}
 				else
 				{
@@ -1258,7 +1267,7 @@ UScriptStruct* FHeaderParser::CompileStructDeclaration(FClasses& AllClasses, FCl
 			}
 			
 			FString ParentStructNameStripped;
-			UField* Field = NULL;
+			const UField* Type = nullptr;
 			bool bOverrideParentStructName = false;
 
 			if( !StructsWithNoPrefix.Contains(ParentStructNameInScript) )
@@ -1270,45 +1279,45 @@ UScriptStruct* FHeaderParser::CompileStructDeclaration(FClasses& AllClasses, FCl
 			// If we're expecting a prefix, first try finding the correct field with the stripped struct name
 			if (bOverrideParentStructName)
 			{
-				Field = FindField( StructClass, *ParentStructNameStripped, true, UScriptStruct::StaticClass(), TEXT("'extends'") );
+				Type = StructScope->FindTypeByName(*ParentStructNameStripped);
 			}
 
 			// If it wasn't found, try to find the literal name given
-			if (Field == NULL)
+			if (Type == NULL)
 			{
-				Field = FindField( StructClass, *ParentStructNameInScript, true, UScriptStruct::StaticClass(), TEXT("'extends'") );
+				Type = StructScope->FindTypeByName(*ParentStructNameInScript);
 			}
 
 			// Resolve structs declared in another class  //@TODO: UCREMOVAL: This seems extreme
-			if (Field == NULL)
+			if (Type == NULL)
 			{
 				if (bOverrideParentStructName)
 				{
-					Field = FindObject<UScriptStruct>(ANY_PACKAGE, *ParentStructNameStripped);
+					Type = FindObject<UScriptStruct>(ANY_PACKAGE, *ParentStructNameStripped);
 				}
 
-				if (Field == NULL)
+				if (Type == NULL)
 				{
-					Field = FindObject<UScriptStruct>(ANY_PACKAGE, *ParentStructNameInScript);
+					Type = FindObject<UScriptStruct>(ANY_PACKAGE, *ParentStructNameInScript);
 				}
 			}
 
 			// If the struct still wasn't found, throw an error
-			if ((Field == NULL) || !Field->IsA( UScriptStruct::StaticClass() ))
+			if (Type == NULL)
 			{
 				FError::Throwf(TEXT("'struct': Can't find struct '%s'"), *ParentStructNameInScript );
 			}
 			else
 			{
 				// If the struct was found, confirm it adheres to the correct syntax. This should always fail if we were expecting an override that was not found.
-				BaseStruct = ((UScriptStruct*)Field);
+				BaseStruct = ((UScriptStruct*)Type);
 				if( bOverrideParentStructName )
 				{
 					const TCHAR* PrefixCPP = StructsWithTPrefix.Contains(ParentStructNameStripped) ? TEXT("T") : BaseStruct->GetPrefixCPP();
 					if( ParentStructNameInScript != FString::Printf(TEXT("%s%s"), PrefixCPP, *ParentStructNameStripped) )
 					{
 						BaseStruct = NULL;
-						FError::Throwf(TEXT("Parent Struct '%s' is missing a valid Unreal prefix, expecting '%s'"), *ParentStructNameInScript, *FString::Printf(TEXT("%s%s"), PrefixCPP, *Field->GetName()));
+						FError::Throwf(TEXT("Parent Struct '%s' is missing a valid Unreal prefix, expecting '%s'"), *ParentStructNameInScript, *FString::Printf(TEXT("%s%s"), PrefixCPP, *Type->GetName()));
 					}
 				}
 				else
@@ -1328,35 +1337,35 @@ UScriptStruct* FHeaderParser::CompileStructDeclaration(FClasses& AllClasses, FCl
 	{
 		StructFlags |= (BaseStruct->StructFlags&STRUCT_Inherit);
 	}
-	
-	// If declared in a header without UClass, the outer is the class package.
-	UObject* StructOuter = FClassUtils::IsTemporaryClass(Scope) ? Scope->GetOuter() : Scope;
+
 	// Create.
-	UScriptStruct* Struct = new(StructOuter, *EffectiveStructName, RF_Public) UScriptStruct(FObjectInitializer(), BaseStruct);
-	Struct->Next = Scope->Children;
-	Scope->Children = Struct;
+	UScriptStruct* Struct = new(SourceFile.GetPackage(), *EffectiveStructName, RF_Public) UScriptStruct(FObjectInitializer(), BaseStruct);
+
+	AddModuleRelativePathToMetadata(Struct, MetaData);
+
+	Scope->AddType(Struct);
+	FScope::AddTypeScope(Struct, &SourceFile.GetScope().Get());
+
+	AddTypeDefinition(SourceFile, Struct, InputLine);
 
 	// Check to make sure the syntactic native prefix was set-up correctly.
 	// If this check results in a false positive, it will be flagged as an identifier failure.
-	if (bOverrideStructName)
+	FString DeclaredPrefix = GetClassPrefix( StructNameInScript );
+	if( DeclaredPrefix == Struct->GetPrefixCPP() || DeclaredPrefix == TEXT("T") )
 	{
-		FString DeclaredPrefix = GetClassPrefix( StructNameInScript );
-		if( DeclaredPrefix == Struct->GetPrefixCPP() || DeclaredPrefix == TEXT("T") )
+		// Found a prefix, do a basic check to see if it's valid
+		const TCHAR* ExpectedPrefixCPP = StructsWithTPrefix.Contains(StructNameStripped) ? TEXT("T") : Struct->GetPrefixCPP();
+		FString ExpectedStructName = FString::Printf(TEXT("%s%s"), ExpectedPrefixCPP, *StructNameStripped);
+		if (StructNameInScript != ExpectedStructName)
 		{
-			// Found a prefix, do a basic check to see if it's valid
-			const TCHAR* ExpectedPrefixCPP = StructsWithTPrefix.Contains(StructNameStripped) ? TEXT("T") : Struct->GetPrefixCPP();
-			FString ExpectedStructName = FString::Printf(TEXT("%s%s"), ExpectedPrefixCPP, *StructNameStripped);
-			if (StructNameInScript != ExpectedStructName)
-			{
-				FError::Throwf(TEXT("Struct '%s' has an invalid Unreal prefix, expecting '%s'"), *StructNameInScript, *ExpectedStructName);
-			}
+			FError::Throwf(TEXT("Struct '%s' has an invalid Unreal prefix, expecting '%s'"), *StructNameInScript, *ExpectedStructName);
 		}
-		else
-		{
-			const TCHAR* ExpectedPrefixCPP = StructsWithTPrefix.Contains(StructNameInScript) ? TEXT("T") : Struct->GetPrefixCPP();
-			FString ExpectedStructName = FString::Printf(TEXT("%s%s"), ExpectedPrefixCPP, *StructNameInScript);
-			FError::Throwf(TEXT("Struct '%s' is missing a valid Unreal prefix, expecting '%s'"), *StructNameInScript, *ExpectedStructName);
-		}			
+	}
+	else
+	{
+		const TCHAR* ExpectedPrefixCPP = StructsWithTPrefix.Contains(StructNameInScript) ? TEXT("T") : Struct->GetPrefixCPP();
+		FString ExpectedStructName = FString::Printf(TEXT("%s%s"), ExpectedPrefixCPP, *StructNameInScript);
+		FError::Throwf(TEXT("Struct '%s' is missing a valid Unreal prefix, expecting '%s'"), *StructNameInScript, *ExpectedStructName);
 	}
 
 	Struct->StructFlags = EStructFlags(Struct->StructFlags | StructFlags);
@@ -1376,10 +1385,9 @@ UScriptStruct* FHeaderParser::CompileStructDeclaration(FClasses& AllClasses, FCl
 	{
 		FToken StructToken;
 		StructToken.Struct = Struct;
-		StructToken.ExportInfo = CustomExportText;
 
 		// add this struct to the compiler's persistent tracking system
-		ClassData->AddStruct(StructToken);
+		GScriptHelper.AddClassData(StructToken.Struct);
 	}
 
 	int32 SavedLineNumber = InputLine;
@@ -1635,149 +1643,62 @@ void FHeaderParser::CheckAllow( const TCHAR* Thing, uint32 AllowFlags )
 	}
 }
 
-//
-// Check that a specified object is accessible from
-// this object's scope.
-//
-void FHeaderParser::CheckInScope( UObject* Obj )
-{
-	if ( Obj != NULL )
-	{
-		UClass* CheckClass = (Cast<UClass>(Obj) != NULL) ? Cast<UClass>(Obj) : Obj->GetClass();
-
-		if ( !AllowReferenceToClass(CheckClass) )
-		{
-			FError::Throwf(TEXT("Invalid reference to unparsed class: '%s'"), *CheckClass->GetPathName());
-		}
-	}
-}
-
-/**
- * Returns whether the specified class can be referenced from the class currently being compiled.
- *
- * @param	CheckClass	the class we want to reference
- *
- * @return	true if the specified class is an intrinsic type or if the class has successfully been parsed
- */
-bool FHeaderParser::AllowReferenceToClass( UClass* CheckClass ) const
+bool FHeaderParser::AllowReferenceToClass(UStruct* Scope, UClass* CheckClass) const
 {
 	check(CheckClass);
 
-	return	(Class->GetOuter() == CheckClass->GetOuter())
-		||	((CheckClass->ClassFlags&CLASS_Parsed) != 0)
-		||	((CheckClass->ClassFlags&CLASS_Intrinsic) != 0);
+	return	(Scope->GetOutermost() == CheckClass->GetOutermost())
+		|| ((CheckClass->ClassFlags&CLASS_Parsed) != 0)
+		|| ((CheckClass->ClassFlags&CLASS_Intrinsic) != 0);
 }
 
 /*-----------------------------------------------------------------------------
 	Nest management.
 -----------------------------------------------------------------------------*/
 
-/**
- * Increase the nesting level, setting the new top nesting level to
- * the one specified.  If pushing a function or state and it overrides a similar
- * thing declared on a lower nesting level, verifies that the override is legal.
- *
- * @param	NestType	the new nesting type
- * @param	ThisName	name of the new nest
- * @param	InNode		@todo
- */
-void FHeaderParser::PushNest( ENestType NestType, FName ThisName, UStruct* InNode )
+void FHeaderParser::PushNest(ENestType NestType, UStruct* InNode, FUnrealSourceFile* SourceFile)
 {
-	// Defaults.
-	UStruct* PrevTopNode = TopNode;
-	UStruct* PrevNode = NULL;
-	uint32 PrevAllow = 0;
-
-	UField* Existing = NULL;
-	if (NestType == NEST_FunctionDeclaration)
-	{
-		for (TFieldIterator<UField> It(TopNode,EFieldIteratorFlags::ExcludeSuper); It; ++It)
-		{
-			if (It->GetFName() == ThisName)
-			{
-				FError::Throwf(TEXT("'%s' conflicts with '%s'"), *ThisName.ToString(), *It->GetFullName() );
-			}
-		}
-	}
-
 	// Update pointer to top nesting level.
-	TopNest					= &Nest[NestLevel++];
-	TopNode					= NULL;
-	TopNest->Node			= InNode;
-	TopNest->NestType		= NestType;
+	TopNest = &Nest[NestLevel++];
+	TopNest->SetScope(NestType == NEST_GlobalScope ? &SourceFile->GetScope().Get() : &FScope::GetTypeScope(InNode).Get());
+	TopNest->NestType = NestType;
 
 	// Prevent overnesting.
-	if( NestLevel >= MAX_NEST_LEVELS )
+	if (NestLevel >= MAX_NEST_LEVELS)
 	{
-		FError::Throwf(TEXT("Maximum nesting limit exceeded") );
+		FError::Throwf(TEXT("Maximum nesting limit exceeded"));
 	}
 
 	// Inherit info from stack node above us.
-	const bool bIsNewNode = (NestType == NEST_Class) || (NestType == NEST_Interface) || (NestType == NEST_FunctionDeclaration);
-	if (NestLevel > 1)
+	if (NestLevel > 1 && NestType == NEST_GlobalScope)
 	{
-		if (bIsNewNode)
-		{
-			// Create a new stack node.
-			if ((NestType == NEST_Class) || (NestType == NEST_Interface))
-			{
-				TopNest->Node = TopNode = Class;
-			}
-			else if (NestType == NEST_FunctionDeclaration)
-			{
-				UFunction* Function = new(PrevTopNode ? (UObject*)PrevTopNode : (UObject*)Class, ThisName, RF_Public) UFunction( FObjectInitializer(), NULL );
-				TopNest->Node = TopNode = Function;
-				Function->RepOffset = MAX_uint16;
-				Function->ReturnValueOffset = MAX_uint16;
-				Function->FirstPropertyToInit = NULL;
-			}
-		}
-		else
-		{
-			// Use the existing stack node.
-			TopNest->Node = TopNest[-1].Node;
-			TopNode = TopNest->Node;
-		}
-
-		check(TopNode != NULL);
-		PrevNode  = TopNest[-1].Node;
-		PrevAllow = TopNest[-1].Allow;
+		// Use the existing stack node.
+		TopNest->SetScope(TopNest[-1].GetScope());
 	}
 
 	// NestType specific logic.
-	switch( NestType )
+	switch (NestType)
 	{
-		case NEST_GlobalScope:
-			check(PrevNode==NULL);
-			TopNest->Allow = ALLOW_Class | ALLOW_TypeDecl;
-			check(PrevNode == NULL);
-			TopNode = InNode;
-			break;
+	case NEST_GlobalScope:
+		TopNest->Allow = ALLOW_Class | ALLOW_TypeDecl | ALLOW_Function;
+		break;
 
-		case NEST_Class:
-			check(ThisName != NAME_None);
-			TopNest->Allow = ALLOW_VarDecl | ALLOW_Function | ALLOW_TypeDecl;
-			break;
+	case NEST_Class:
+		TopNest->Allow = ALLOW_VarDecl | ALLOW_Function | ALLOW_TypeDecl;
+		break;
 
-		// only function declarations are allowed inside interface nesting level
-		case NEST_Interface:
-			check(ThisName != NAME_None);
-			TopNest->Allow = ALLOW_Function | ALLOW_TypeDecl;
-			break;
+	case NEST_Interface:
+		TopNest->Allow = ALLOW_Function | ALLOW_TypeDecl;
+		break;
 
-		case NEST_FunctionDeclaration:
-			check(ThisName != NAME_None);
-			check(PrevNode != NULL);
-			TopNest->Allow = ALLOW_VarDecl;
-			
-			TopNode->Next = PrevNode->Children;
-			PrevNode->Children = TopNode;
-			
-			break;
+	case NEST_FunctionDeclaration:
+		TopNest->Allow = ALLOW_VarDecl;
 
-		default:
-			FError::Throwf(TEXT("Internal error in PushNest, type %i"), (uint8)NestType );
-			break;
+		break;
+
+	default:
+		FError::Throwf(TEXT("Internal error in PushNest, type %i"), (uint8)NestType);
+		break;
 	}
 }
 
@@ -1787,164 +1708,46 @@ void FHeaderParser::PushNest( ENestType NestType, FName ThisName, UStruct* InNod
  * @param	NestType	nesting type of the current node
  * @param	Descr		text to use in error message if any errors are encountered
  */
-void FHeaderParser::PopNest( FClasses& AllClasses, ENestType NestType, const TCHAR* Descr )
+void FHeaderParser::PopNest(ENestType NestType, const TCHAR* Descr)
 {
 	// Validate the nesting state.
 	if (NestLevel <= 0)
 	{
-		FError::Throwf(TEXT("Unexpected '%s' at global scope"), Descr, NestTypeName(NestType) );
-	}
-	else if (NestType == NEST_GlobalScope)
-	{
-		NestType = TopNest->NestType;
+		FError::Throwf(TEXT("Unexpected '%s' at global scope"), Descr, NestTypeName(NestType));
 	}
 	else if (TopNest->NestType != NestType)
 	{
-		FError::Throwf(TEXT("Unexpected end of %s in '%s' block"), Descr, NestTypeName(TopNest->NestType) );
-	}
-
-	// Remember code position.
-	if (NestType == NEST_FunctionDeclaration)
-	{
-		//@TODO: UCREMOVAL: Move this code to occur at delegate var declaration, and force delegates to be declared before variables that use them
-		if (ClassData->ContainsDelegates())
-		{
-			UFunction* TopFunction = CastChecked<UFunction>(TopNode);
-
-			// now validate all delegate variables declared in the class
-			TMap<FName, UFunction*> DelegateCache;
-			FixupDelegateProperties(AllClasses, TopFunction, TopFunction->GetOwnerClass(), DelegateCache);
-		}
+		FError::Throwf(TEXT("Unexpected end of %s in '%s' block"), Descr, NestTypeName(TopNest->NestType));
 	}
-	else if (NestType == NEST_Class)
-	{
-		// Todo - One drawback to doing this here is we don't get correct line numbers
-		// for conflicting implementations.
-		UClass* TopClass = CastChecked<UClass>(TopNode);
-
-		// first, fixup all delegate properties with the delegate types they're supposed to be bound to
-		if (ClassData->ContainsDelegates())
-		{
-			// now validate all delegate variables declared in the class
-			TMap<FName,UFunction*> DelegateCache;
-			FixupDelegateProperties(AllClasses, TopClass, TopClass, DelegateCache);
-		}
-
-		// Validate all the rep notify events here, to make sure they're implemented
-		VerifyRepNotifyCallbacks(TopClass);
-
-		// Iterate over all the interfaces we claim to implement
-		for (auto& Impl : TopClass->Interfaces)
-		{
-			// And their super-classes
-			for( UClass* Interface = Impl.Class; Interface; Interface = Interface->GetSuperClass() )
-			{
-				// If this interface is a common ancestor, skip it
-				if (TopNode->IsChildOf(Interface))
-					continue;
-
-				// So iterate over all functions this interface declares
-				for (auto InterfaceFunction : TFieldRange<UFunction>(Interface, EFieldIteratorFlags::ExcludeSuper))
-				{
-					bool Implemented = false;
-
-					// And try to find one that matches
-					for (UFunction* ClassFunction : TFieldRange<UFunction>(TopNode))
-					{
-						if (ClassFunction->GetFName() != InterfaceFunction->GetFName())
-							continue;
-
-						if ((InterfaceFunction->FunctionFlags & FUNC_Event) && !(ClassFunction->FunctionFlags & FUNC_Event))
-							FError::Throwf(TEXT("Implementation of function '%s' must be declared as 'event' to match declaration in interface '%s'"), *ClassFunction->GetName(), *Interface->GetName());
-
-						if ((InterfaceFunction->FunctionFlags & FUNC_Delegate) && !(ClassFunction->FunctionFlags & FUNC_Delegate))
-							FError::Throwf(TEXT("Implementation of function '%s' must be declared as 'delegate' to match declaration in interface '%s'"), *ClassFunction->GetName(), *Interface->GetName());
-
-						// Making sure all the parameters match up correctly
-						Implemented = true;
-
-						if (ClassFunction->NumParms != InterfaceFunction->NumParms)
-							FError::Throwf(TEXT("Implementation of function '%s' conflicts with interface '%s' - different number of parameters (%i/%i)"), *InterfaceFunction->GetName(), *Interface->GetName(), ClassFunction->NumParms, InterfaceFunction->NumParms);
-
-						int32 Count = 0;
-						for( TFieldIterator<UProperty> It1(InterfaceFunction),It2(ClassFunction); Count<ClassFunction->NumParms; ++It1,++It2,Count++ )
-						{
-							if( !FPropertyBase(*It1).MatchesType(FPropertyBase(*It2), 1) )
-							{
-								if( It1->PropertyFlags & CPF_ReturnParm )
-								{
-									FError::Throwf(TEXT("Implementation of function '%s' conflicts only by return type with interface '%s'"), *InterfaceFunction->GetName(), *Interface->GetName() );
-								}
-								else
-								{
-									FError::Throwf(TEXT("Implementation of function '%s' conflicts with interface '%s' - parameter %i '%s'"), *InterfaceFunction->GetName(), *Interface->GetName(), Count, *It1->GetName() );
-								}
-							}
-						}
-					}
-
-					// Delegate signature functions are simple stubs and aren't required to be implemented (they are not callable)
-					if (InterfaceFunction->FunctionFlags & FUNC_Delegate)
-					{
-						Implemented = true;
-					}
 
-					// Verify that if this has blueprint-callable functions that are not implementable events, we've implemented them as a UFunction in the target class
-					if( !Implemented 
-						&& !Interface->HasMetaData(TEXT("CannotImplementInterfaceInBlueprint"))  // FBlueprintMetadata::MD_CannotImplementInterfaceInBlueprint
-						&& InterfaceFunction->HasAnyFunctionFlags(FUNC_BlueprintCallable)
-						&& !InterfaceFunction->HasAnyFunctionFlags(FUNC_BlueprintEvent) )
-					{
-						FError::Throwf(TEXT("Missing UFunction implementation of function '%s' from interface '%s'.  This function needs a UFUNCTION() declaration."), *InterfaceFunction->GetName(), *Interface->GetName());
-					}
-				}
-			}
-		}
-	}
-	else if ( NestType == NEST_Interface )
-	{
-		if (ClassData->ContainsDelegates())
-		{
-			TMap<FName,UFunction*> DelegateCache;
-			FixupDelegateProperties(AllClasses, TopNode, ExactCast<UClass>(TopNode), DelegateCache);
-		}
-	}
-	else
+	if (NestType != NEST_GlobalScope && NestType != NEST_Class && NestType != NEST_Interface && NestType != NEST_FunctionDeclaration)
 	{
-		FError::Throwf(TEXT("Bad first pass NestType %i"), (uint8)NestType );
+		FError::Throwf(TEXT("Bad first pass NestType %i"), (uint8)NestType);
 	}
 
 	bool bLinkProps = true;
 	if (NestType == NEST_Class)
 	{
-		UClass* TopClass = CastChecked<UClass>(TopNode);
+		UClass* TopClass = CastChecked<UClass>(GetCurrentClass());
 		bLinkProps = !TopClass->HasAnyClassFlags(CLASS_Intrinsic);
 	}
-	TopNode->StaticLink(bLinkProps);
+
+	if (NestType != NEST_GlobalScope)
+	{
+		GetCurrentClass()->StaticLink(bLinkProps);
+	}
 
 	// Pop the nesting level.
 	NestType = TopNest->NestType;
 	NestLevel--;
 	TopNest--;
-	TopNode	= TopNest->Node;
 }
 
-/**
- * Binds all delegate properties declared in ValidationScope the delegate functions specified in the variable declaration, verifying that the function is a valid delegate
- * within the current scope.  This must be done once the entire class has been parsed because instance delegate properties must be declared before the delegate declaration itself.
- *
- * @todo: this function will no longer be required once the post-parse fixup phase is added (TTPRO #13256)
- *
- * @param	ValidationScope		the scope to validate delegate properties for
- * @param	OwnerClass			the class currently being compiled.
- * @param	DelegateCache		cached map of delegates that have already been found; used for faster lookup.
- */
-void FHeaderParser::FixupDelegateProperties( FClasses& AllClasses, UStruct* ValidationScope, UClass* OwnerClass, TMap<FName, UFunction*>& DelegateCache )
+void FHeaderParser::FixupDelegateProperties( FClasses& AllClasses, UStruct* Struct, FScope& Scope, TMap<FName, UFunction*>& DelegateCache )
 {
-	check(ValidationScope);
-	check(OwnerClass);
+	check(Struct);
 
-	for ( UField* Field = ValidationScope->Children; Field; Field = Field->Next )
+	for ( UField* Field = Struct->Children; Field; Field = Field->Next )
 	{
 		UProperty* Property = Cast<UProperty>(Field);
 		if ( Property != NULL )
@@ -1961,31 +1764,32 @@ void FHeaderParser::FixupDelegateProperties( FClasses& AllClasses, UStruct* Vali
 					MulticastDelegateProperty = Cast<UMulticastDelegateProperty>(ArrayProp->Inner);
 				}
 			}
-			if ( DelegateProperty != NULL || MulticastDelegateProperty != NULL )
+			if (DelegateProperty != nullptr || MulticastDelegateProperty != nullptr)
 			{
 				// this UDelegateProperty corresponds to an actual delegate variable (i.e. delegate<SomeDelegate> Foo); we need to lookup the token data for
 				// this property and verify that the delegate property's "type" is an actual delegate function
-				FTokenData* DelegatePropertyToken = ClassData->FindTokenData(Property);
+				FTokenData* DelegatePropertyToken = GScriptHelper.FindClassData(Struct)->FindTokenData(Property);
 				check(DelegatePropertyToken);
 
 				// attempt to find the delegate function in the map of functions we've already found
 				UFunction* SourceDelegateFunction = DelegateCache.FindRef(DelegatePropertyToken->Token.DelegateName);
-				if ( SourceDelegateFunction == NULL )
+				if (SourceDelegateFunction == nullptr)
 				{
 					FString NameOfDelegateFunction = DelegatePropertyToken->Token.DelegateName.ToString() + FString( HEADER_GENERATED_DELEGATE_SIGNATURE_SUFFIX );
 					if ( !NameOfDelegateFunction.Contains(TEXT(".")) )
 					{
 						// an unqualified delegate function name - search for a delegate function by this name within the current scope
-						SourceDelegateFunction = Cast<UFunction>(FindField(OwnerClass, *NameOfDelegateFunction, true, UFunction::StaticClass(), NULL));
-						if ( SourceDelegateFunction == NULL )
-						{
-							// convert this into a fully qualified path name for the error message.
-							NameOfDelegateFunction = OwnerClass->GetName() + TEXT(".") + NameOfDelegateFunction;
-						}
-						else
+						SourceDelegateFunction = Cast<UFunction>(Scope.FindTypeByName(*NameOfDelegateFunction));
+						if (SourceDelegateFunction == nullptr)
 						{
-							// convert this into a fully qualified path name for the error message.
-							NameOfDelegateFunction = SourceDelegateFunction->GetOwnerClass()->GetName() + TEXT(".") + NameOfDelegateFunction;
+							// Try to find in other packages.
+							SourceDelegateFunction = Cast<UFunction>(StaticFindObject(UFunction::StaticClass(), ANY_PACKAGE, *NameOfDelegateFunction));
+
+							if (SourceDelegateFunction == nullptr)
+							{
+								// convert this into a fully qualified path name for the error message.
+								NameOfDelegateFunction = Scope.GetName().ToString() + TEXT(".") + NameOfDelegateFunction;
+							}
 						}
 					}
 					else
@@ -2009,7 +1813,10 @@ void FHeaderParser::FixupDelegateProperties( FClasses& AllClasses, UStruct* Vali
 
 						// make sure that the class that contains the delegate can be referenced here
 						UClass* DelegateOwnerClass = AllClasses.FindScriptClassOrThrow(DelegateClassName);
-						CheckInScope(DelegateOwnerClass);
+						if (FScope::GetTypeScope(DelegateOwnerClass)->FindTypeByName(*DelegateName) != nullptr)
+						{
+							FError::Throwf(TEXT("Inaccessible type: '%s'"), *DelegateOwnerClass->GetPathName());
+						}
 						SourceDelegateFunction = Cast<UFunction>(FindField(DelegateOwnerClass, *DelegateName, false, UFunction::StaticClass(), NULL));	
 					}
 
@@ -2074,7 +1881,7 @@ void FHeaderParser::FixupDelegateProperties( FClasses& AllClasses, UStruct* Vali
 			UStruct* InternalStruct = Cast<UStruct>(Field);
 			if ( InternalStruct != NULL )
 			{
-				FixupDelegateProperties(AllClasses, InternalStruct, OwnerClass, DelegateCache);
+				FixupDelegateProperties(AllClasses, InternalStruct, Scope, DelegateCache);
 			}
 		}
 	}
@@ -2093,7 +1900,7 @@ void FHeaderParser::VerifyRepNotifyCallbacks( UClass* TargetClass )
 		UProperty* Prop = Cast<UProperty>(Field);
 		if( Prop && (Prop->GetPropertyFlags() & CPF_RepNotify) )
 		{
-			FTokenData* PropertyToken = ClassData->FindTokenData(Prop);
+			FTokenData* PropertyToken = GScriptHelper.FindClassData(TargetClass)->FindTokenData(Prop);
 			check(PropertyToken);
 
 			// Search through this class and its superclasses looking for the specified callback
@@ -2177,7 +1984,7 @@ void FHeaderParser::VerifyRepNotifyCallbacks( UClass* TargetClass )
 //
 // Process a compiler directive.
 //
-void FHeaderParser::CompileDirective(UClass* Class)
+void FHeaderParser::CompileDirective(FClasses& AllClasses, FUnrealSourceFile& SourceFile)
 {
 	FToken Directive;
 
@@ -2213,7 +2020,7 @@ void FHeaderParser::CompileDirective(UClass* Class)
 	}
 	else if (Directive.Matches(TEXT("include")))
 	{
-		FString ExpectedHeaderName = FString::Printf(TEXT("%s.generated.h"), *GClassHeaderNameWithNoPathMap[Class]);
+		FString ExpectedHeaderName = SourceFile.GetGeneratedHeaderFilename();
 		FToken IncludeName;
 		if (GetToken(IncludeName) && (IncludeName.TokenType == TOKEN_Const) && (IncludeName.Type == CPT_String))
 		{
@@ -2221,12 +2028,6 @@ void FHeaderParser::CompileDirective(UClass* Class)
 			{
 				bSpottedAutogeneratedHeaderInclude = true;
 			}
-			else
-			{
-				// Handle #include directive just like dependson. The included header name is stored with
-				// the extension so that we later know where this dependson came from.
-				GClassDependentOnMap.FindOrAdd(Class)->Add(*FPaths::GetCleanFilename(IncludeName.String));
-			}
 		}
 	}
 	else if (Directive.Matches(TEXT("if")))
@@ -2312,7 +2113,7 @@ void FHeaderParser::CompileDirective(UClass* Class)
 bool FHeaderParser::GetVarType
 (
 	FClasses&                       AllClasses,
-	UStruct*                        Scope,
+	FScope*							Scope,
 	FPropertyBase&                  VarProperty,
 	EObjectFlags&                   ObjectFlags,
 	uint64                          Disallow,
@@ -2323,8 +2124,7 @@ bool FHeaderParser::GetVarType
 	FIndexRange*                    ParsedVarIndexRange
 )
 {
-	check(Scope);
-
+	UStruct* OwnerStruct = Scope->IsFileScope() ? nullptr : ((FStructScope*)Scope)->GetStruct();
 	FName RepCallbackName      = FName(NAME_None);
 	bool  bIsMulticastDelegate = false;
 	bool  bIsWeak              = false;
@@ -2336,7 +2136,7 @@ bool FHeaderParser::GetVarType
 	ObjectFlags = RF_Public;
 	uint64 Flags = 0;
 	// force members to be 'blueprint read only' if in a const class
-	if (VariableCategory == EVariableCategory::Member && (Cast<UClass>(Scope) != NULL) && (((UClass*)Scope)->ClassFlags & CLASS_Const))
+	if (VariableCategory == EVariableCategory::Member && (Cast<UClass>(OwnerStruct) != nullptr) && (((UClass*)OwnerStruct)->ClassFlags & CLASS_Const))
 	{
 		Flags |= CPF_BlueprintReadOnly;
 	}
@@ -2540,7 +2340,7 @@ bool FHeaderParser::GetVarType
 			}
 			else if (Specifier == TEXT("Replicated") || Specifier == TEXT("ReplicatedUsing"))
 			{
-				if (Scope->GetClass() != UScriptStruct::StaticClass())
+				if (!OwnerStruct->IsA<UScriptStruct>())
 				{
 					Flags |= CPF_Net;
 
@@ -2558,7 +2358,7 @@ bool FHeaderParser::GetVarType
 			}
 			else if (Specifier == TEXT("NotReplicated"))
 			{
-				if (Scope->GetClass() == UScriptStruct::StaticClass())
+				if (OwnerStruct->IsA<UScriptStruct>())
 				{
 					Flags |= CPF_RepSkip;
 				}
@@ -2820,7 +2620,7 @@ bool FHeaderParser::GetVarType
 		uint64 OriginalVarTypeFlags = VarType.PropertyFlags;
 		VarType.PropertyFlags |= Flags;
 
-		GetVarType( AllClasses, Scope, VarProperty, InnerFlags, Disallow, TEXT("'tarray'"), &VarType, EPropertyDeclarationStyle::None, VariableCategory );
+		GetVarType(AllClasses, Scope, VarProperty, InnerFlags, Disallow, TEXT("'tarray'"), &VarType, EPropertyDeclarationStyle::None, VariableCategory);
 		if (VarProperty.ArrayType != EArrayType::None)
 		{
 			FError::Throwf(TEXT("Arrays within arrays not supported.") );
@@ -2981,11 +2781,10 @@ bool FHeaderParser::GetVarType
 		{
 			bHandledType = true;
 
-
 			// Struct keyword in front of a struct is legal, we 'consume' it
 			bUnconsumedStructKeyword = false;
 		}
-		else if ( UFunction* DelegateFunc = Cast<UFunction>(FindField(Scope, *(IdentifierStripped + HEADER_GENERATED_DELEGATE_SIGNATURE_SUFFIX))) )
+		else if (UFunction* DelegateFunc = Cast<UFunction>(Scope->FindTypeByName(*(IdentifierStripped + HEADER_GENERATED_DELEGATE_SIGNATURE_SUFFIX))))
 		{
 			bHandledType = true;
 
@@ -3085,9 +2884,6 @@ bool FHeaderParser::GetVarType
 			{
 				bHandledType = true;
 
-				// An object reference.
-				CheckInScope( TempClass );
-
 				bool bAllowWeak = !(Disallow & CPF_AutoWeak); // if it is not allowing anything, force it strong. this is probably a function arg
 				VarProperty = FPropertyBase( TempClass, NULL, bAllowWeak, bIsWeak, bWeakIsAuto, bIsLazy, bIsAsset );
 				if (TempClass->IsChildOf(UClass::StaticClass()))
@@ -3163,21 +2959,16 @@ bool FHeaderParser::GetVarType
 		// Resolve delegates declared in another class  //@TODO: UCREMOVAL: This seems extreme
 		if (!bHandledType)
 		{
-			for (FClass* OtherClass : AllClasses)
+			if (UFunction* DelegateFunc = (UFunction*)StaticFindObject(UFunction::StaticClass(), ANY_PACKAGE, *(IdentifierStripped + HEADER_GENERATED_DELEGATE_SIGNATURE_SUFFIX)))
 			{
-				if (UFunction* DelegateFunc = Cast<UFunction>(FindField(OtherClass, *(IdentifierStripped + HEADER_GENERATED_DELEGATE_SIGNATURE_SUFFIX))))
-				{
-					bHandledType = true;
+				bHandledType = true;
 
-					VarProperty = FPropertyBase( DelegateFunc->HasAnyFunctionFlags(FUNC_MulticastDelegate) ? CPT_MulticastDelegate : CPT_Delegate);
-					VarProperty.DelegateName = *FString::Printf(TEXT("%s.%s"), *OtherClass->GetNameWithPrefix(), *IdentifierStripped);
+				VarProperty = FPropertyBase( DelegateFunc->HasAnyFunctionFlags(FUNC_MulticastDelegate) ? CPT_MulticastDelegate : CPT_Delegate);
+				VarProperty.DelegateName = *IdentifierStripped;
 					
-					if (!(Disallow & CPF_InstancedReference))
-					{
-						Flags |= CPF_InstancedReference;
-					}
-
-					break;
+				if (!(Disallow & CPF_InstancedReference))
+				{
+					Flags |= CPF_InstancedReference;
 				}
 			}
 
@@ -3239,7 +3030,7 @@ bool FHeaderParser::GetVarType
 	}
 
 	//@TODO: UCREMOVAL: 'const' member variables that will get written post-construction by defaultproperties
-	if (Class->HasAnyClassFlags(CLASS_Const) && VariableCategory == EVariableCategory::Member)
+	if (VariableCategory == EVariableCategory::Member && OwnerStruct->IsA<UClass>() && ((UClass*)OwnerStruct)->HasAnyClassFlags(CLASS_Const))
 	{
 		// Eat a 'not quite truthful' const after the type; autogenerated for member variables of const classes.
 		MatchIdentifier(TEXT("const"));
@@ -3347,7 +3138,7 @@ bool FHeaderParser::GetVarType
 
 	// Check for invalid transients
 	uint64 Transients = VarProperty.PropertyFlags & (CPF_DuplicateTransient | CPF_TextExportTransient | CPF_NonPIEDuplicateTransient);
-	if (Transients && !Cast<UClass>(Scope))
+	if (Transients && !Cast<UClass>(OwnerStruct))
 	{
 		TArray<const TCHAR*> FlagStrs = ParsePropertyFlags(Transients);
 		FError::Throwf(TEXT("'%s' specifier(s) are only allowed on class member variables"), *FString::Join(FlagStrs, TEXT(", ")));
@@ -3414,12 +3205,7 @@ UProperty* FHeaderParser::GetVarNameAndDim
 {
 	check(Scope);
 
-	// Add metadata for the module relative path
-	const FString *ModuleRelativePath = GClassModuleRelativePathMap.Find(Class);
-	if(ModuleRelativePath != NULL)
-	{
-		VarProperty.MetaData.Add(TEXT("ModuleRelativePath"), *ModuleRelativePath);
-	}
+	AddModuleRelativePathToMetadata(Scope, VarProperty.MetaData);
 
 	// Get variable name.
 	if (HardcodedName != NULL)
@@ -3796,7 +3582,7 @@ UProperty* FHeaderParser::GetVarNameAndDim
 	}
 
 	VarProperty.TokenProperty = NewProperty;
-	ClassData->AddProperty(VarProperty);
+	GScriptHelper.FindClassData(Scope)->AddProperty(VarProperty);
 
 	// if we had any metadata, add it to the class
 	AddMetaDataToClassData(VarProperty.TokenProperty, VarProperty.MetaData);
@@ -3811,7 +3597,7 @@ UProperty* FHeaderParser::GetVarNameAndDim
 //
 // Compile a declaration in Token. Returns 1 if compiled, 0 if not.
 //
-bool FHeaderParser::CompileDeclaration( FClasses& AllClasses, FToken& Token )
+bool FHeaderParser::CompileDeclaration(FClasses& AllClasses, FUnrealSourceFile& SourceFile, FToken& Token)
 {
 	EAccessSpecifier AccessSpecifier = ParseAccessProtectionSpecifier(Token);
 	if (AccessSpecifier)
@@ -3820,16 +3606,11 @@ bool FHeaderParser::CompileDeclaration( FClasses& AllClasses, FToken& Token )
 		{
 			FError::Throwf(TEXT("Access specifier %s not allowed here."), Token.Identifier);
 		}
-		check( TopNest->NestType == NEST_Class || TopNest->NestType == NEST_Interface );
+		check(TopNest->NestType == NEST_Class || TopNest->NestType == NEST_Interface);
 		CurrentAccessSpecifier = AccessSpecifier;
 	}
 	else if (Token.Matches(TEXT("class")) && (TopNest->NestType == NEST_GlobalScope))
 	{
-		if (!bHaveSeenFirstInterfaceClass || bFinishedParsingInterfaceClasses)
-		{
-			return SkipDeclaration(Token);
-		}
-
 		// Make sure the previous class ended with valid nesting.
 		if (bEncounteredNewStyleClass_UnmatchedBrackets)
 		{
@@ -3838,27 +3619,34 @@ bool FHeaderParser::CompileDeclaration( FClasses& AllClasses, FToken& Token )
 
 		// Start parsing the second class
 		bEncounteredNewStyleClass_UnmatchedBrackets = true;
-		bHaveSeenSecondInterfaceClass = true;
-		ParseSecondInterfaceClass(AllClasses);
+
+		if (!TryParseIInterfaceClass(AllClasses))
+		{
+			bEncounteredNewStyleClass_UnmatchedBrackets = false;
+			UngetToken(Token);
+			return SkipDeclaration(Token);
+		}
 	}
-	else if (Token.Matches(TEXT("GENERATED_BODY")) && (TopNest->NestType != NEST_Interface
-		|| (bHaveSeenFirstInterfaceClass && !bHaveSeenSecondInterfaceClass)))
+	else if (Token.Matches(TEXT("GENERATED_BODY")))
 	{
 		if (TopNest->NestType != NEST_Class && TopNest->NestType != NEST_Interface)
 		{
 			FError::Throwf(TEXT("%s must occur inside the class or interface definition"), Token.Identifier);
 		}
-		if (!ClassDefinitionRanges.Contains(Class))
+		if (!ClassDefinitionRanges.Contains(GetCurrentClass()))
 		{
-			ClassDefinitionRanges.Add(Class, ClassDefinitionRange());
+			ClassDefinitionRanges.Add(GetCurrentClass(), ClassDefinitionRange());
 		}
 
-		ClassDefinitionRanges[Class].bHasGeneratedBody = true;
+		ClassDefinitionRanges[GetCurrentClass()].bHasGeneratedBody = true;
 
 		RequireSymbol(TEXT("("), Token.Identifier);
 		RequireSymbol(TEXT(")"), Token.Identifier);
 
-		GScriptHelper.FindClassData(Class)->GeneratedBodyMacroAccessSpecifier = CurrentAccessSpecifier;
+		auto* ClassData = GetCurrentClassData();
+
+		ClassData->GeneratedBodyMacroAccessSpecifier = CurrentAccessSpecifier;
+		ClassData->SetGeneratedBodyLine(InputLine);
 
 		if (TopNest->NestType == NEST_Class)
 		{
@@ -3874,20 +3662,24 @@ bool FHeaderParser::CompileDeclaration( FClasses& AllClasses, FToken& Token )
 		RequireSymbol(TEXT("("), Token.Identifier);
 		RequireSymbol(TEXT(")"), Token.Identifier);
 
+		GetCurrentClassData()->SetGeneratedBodyLine(InputLine);
+
 		// The body implementation macro always ends with a 'public:'
 		CurrentAccessSpecifier = ACCESS_Public;
 
 		bClassHasGeneratedBody = true;
 	}
-	else if (bHaveSeenSecondInterfaceClass && Token.Matches(TEXT("GENERATED_IINTERFACE_BODY")))
+	else if (Token.Matches(TEXT("GENERATED_IINTERFACE_BODY")))
 	{
 		RequireSymbol(TEXT("("), Token.Identifier);
 		RequireSymbol(TEXT(")"), Token.Identifier);
 
+		GetCurrentClassData()->SetInterfaceGeneratedBodyLine(InputLine);
+
 		// The body implementation macro always ends with a 'public:'
 		CurrentAccessSpecifier = ACCESS_Public;
 	}
-	else if (bHaveSeenFirstInterfaceClass && !bHaveSeenSecondInterfaceClass && Token.Matches(TEXT("GENERATED_UINTERFACE_BODY")))
+	else if (Token.Matches(TEXT("GENERATED_UINTERFACE_BODY")))
 	{
 		if (TopNest->NestType != NEST_Interface)
 		{
@@ -3896,69 +3688,62 @@ bool FHeaderParser::CompileDeclaration( FClasses& AllClasses, FToken& Token )
 		RequireSymbol(TEXT("("), Token.Identifier);
 		RequireSymbol(TEXT(")"), Token.Identifier);
 
+		GetCurrentClassData()->SetGeneratedBodyLine(InputLine);
+
 		// The body implementation macro always ends with a 'public:'
 		CurrentAccessSpecifier = ACCESS_Public;
 	}
 	else if (Token.Matches(TEXT("UCLASS"), ESearchCase::CaseSensitive) && (TopNest->Allow & ALLOW_Class))
 	{
-		if (bHaveSeenUClass)
-		{
-			FError::Throwf(TEXT("Can only declare one class per file (two for an interface)"));
-		}
 		bHaveSeenUClass = true;
-		bEncounteredNewStyleClass_UnmatchedBrackets = true;		
+		bEncounteredNewStyleClass_UnmatchedBrackets = true;
 		CompileClassDeclaration(AllClasses);
 	}
 	else if (Token.Matches(TEXT("UINTERFACE")) && (TopNest->Allow & ALLOW_Class))
 	{
-		if (bHaveSeenUClass)
-		{
-			FError::Throwf(TEXT("Can only declare one class per file (two for an interface)"));
-		}
 		bHaveSeenUClass = true;
 		bEncounteredNewStyleClass_UnmatchedBrackets = true;
-		bHaveSeenFirstInterfaceClass = true;
 		CompileInterfaceDeclaration(AllClasses);
 	}
 	else if (Token.Matches(TEXT("UFUNCTION"), ESearchCase::CaseSensitive))
 	{
-		CompileFunctionDeclaration(AllClasses);
+		CompileFunctionDeclaration(SourceFile, AllClasses);
 	}
 	else if (Token.Matches(TEXT("UDELEGATE")))
 	{
-		CompileDelegateDeclaration(AllClasses, Token.Identifier, EDelegateSpecifierAction::Parse);
+		CompileDelegateDeclaration(SourceFile, AllClasses, Token.Identifier, EDelegateSpecifierAction::Parse);
 	}
 	else if (IsValidDelegateDeclaration(Token)) // Legacy delegate parsing - it didn't need a UDELEGATE
 	{
-		CompileDelegateDeclaration(AllClasses, Token.Identifier);
+		CompileDelegateDeclaration(SourceFile, AllClasses, Token.Identifier);
 	}
 	else if (Token.Matches(TEXT("UPROPERTY"), ESearchCase::CaseSensitive))
 	{
-		CheckAllow( TEXT("'Member variable declaration'"), ALLOW_VarDecl );
-		check( TopNest->NestType == NEST_Class );
+		CheckAllow(TEXT("'Member variable declaration'"), ALLOW_VarDecl);
+		check(TopNest->NestType == NEST_Class);
 
-		CompileVariableDeclaration(AllClasses, TopNode, EPropertyDeclarationStyle::UPROPERTY);
+		CompileVariableDeclaration(AllClasses, GetCurrentClass(), EPropertyDeclarationStyle::UPROPERTY);
 	}
 	else if (Token.Matches(TEXT("UENUM")))
 	{
 		// Enumeration definition.
-		CompileEnum(Class);
+		CompileEnum(SourceFile);
 	}
 	else if (Token.Matches(TEXT("USTRUCT")))
 	{
 		// Struct definition.
-		CompileStructDeclaration(AllClasses, Class);
+		CompileStructDeclaration(AllClasses, SourceFile);
 	}
 	else if (Token.Matches(TEXT("#")))
 	{
 		// Compiler directive.
-		CompileDirective(Class);
+		CompileDirective(AllClasses, SourceFile);
 	}
 	else if (bEncounteredNewStyleClass_UnmatchedBrackets && Token.Matches(TEXT("}")))
 	{
-		if (ClassDefinitionRanges.Contains(Class))
+		if (ClassDefinitionRanges.Contains(GetCurrentClass()))
 		{
-			ClassDefinitionRanges[Class].End = &Input[InputPos];
+			ClassDefinitionRanges[GetCurrentClass()].End = &Input[InputPos];
 		}
 		MatchSemi();
 
@@ -3967,19 +3752,18 @@ bool FHeaderParser::CompileDeclaration( FClasses& AllClasses, FToken& Token )
 		// Instead, the whole statement-consumer loop should be in a nest
 		bEncounteredNewStyleClass_UnmatchedBrackets = false;
 
+		UClass* CurrentClass = GetCurrentClass();
+
 		// Pop nesting here to allow other non UClass declarations in the header file.
-		if (Class->ClassFlags & CLASS_Interface)
+		if (CurrentClass->ClassFlags & CLASS_Interface)
 		{
-			if (bHaveSeenSecondInterfaceClass)
-			{
-				// Remember all interface classes are now parsed so that we can parse other (non-U)classes.
-				bFinishedParsingInterfaceClasses = true;
-			}
-			PopNest( AllClasses, NEST_Interface, TEXT("'Interface'") );
+			PopNest(NEST_Interface, TEXT("'Interface'"));
+			PostPopNestInterface(AllClasses, CurrentClass);
 		}
 		else
 		{
-			PopNest( AllClasses, NEST_Class, TEXT("'Class'") );
+			PopNest(NEST_Class, TEXT("'Class'"));
+			PostPopNestClass(CurrentClass);
 
 			// Ensure classes have a GENERATED_UCLASS_BODY declaration
 			if (bHaveSeenUClass && !bClassHasGeneratedBody)
@@ -3989,21 +3773,23 @@ bool FHeaderParser::CompileDeclaration( FClasses& AllClasses, FToken& Token )
 
 			bClassHasGeneratedBody = false;
 		}
+
+		GetCurrentScope()->AddType(CurrentClass);
 	}
 	else if (Token.Matches(TEXT(";")))
 	{
-		if( GetToken(Token) )
+		if (GetToken(Token))
 		{
-			FError::Throwf(TEXT("Extra ';' before '%s'"), Token.Identifier );
+			FError::Throwf(TEXT("Extra ';' before '%s'"), Token.Identifier);
 		}
 		else
 		{
-			FError::Throwf(TEXT("Extra ';' before end of file") );
+			FError::Throwf(TEXT("Extra ';' before end of file"));
 		}
 	}
-	else if (bEncounteredNewStyleClass_UnmatchedBrackets && Token.Matches(NameLookupCPP.GetNameCPP(Class)))
+	else if (bEncounteredNewStyleClass_UnmatchedBrackets && IsInAClass() && GetCurrentClass() && Token.Matches(NameLookupCPP.GetNameCPP(GetCurrentClass())))
 	{
-		if(!TryToMatchConstructorParameterList(Token))
+		if (!TryToMatchConstructorParameterList(Token))
 		{
 			return SkipDeclaration(Token);
 		}
@@ -4134,10 +3920,15 @@ bool FHeaderParser::SafeMatchSymbol( const TCHAR* Match )
 	return false;
 }
 
-void FHeaderParser::ParseClassNameDeclaration(FClasses& AllClasses, FString& DeclaredClassName, FString& RequiredAPIMacroIfPresent)
+FClass* FHeaderParser::ParseClassNameDeclaration(FClasses& AllClasses, FString& DeclaredClassName, FString& RequiredAPIMacroIfPresent)
 {
 	ParseNameWithPotentialAPIMacroPrefix(/*out*/ DeclaredClassName, /*out*/ RequiredAPIMacroIfPresent, TEXT("class"));
 
+	FClass* FoundClass = AllClasses.FindClass(*GetClassNameWithPrefixRemoved(*DeclaredClassName));
+	check(FoundClass);
+
+	GScriptHelper.AddClassData(FoundClass);
+
 	// Get parent class.
 	bool bSpecifiesParentClass = false;
 
@@ -4148,7 +3939,7 @@ void FHeaderParser::ParseClassNameDeclaration(FClasses& AllClasses, FString& Dec
 	}
 
 	// Add class cast flag
-	Class->ClassCastFlags |= ClassCastFlagMap::Get().GetCastFlag(DeclaredClassName);
+	FoundClass->ClassCastFlags |= ClassCastFlagMap::Get().GetCastFlag(DeclaredClassName);
 
 	if (bSpecifiesParentClass)
 	{
@@ -4157,20 +3948,20 @@ void FHeaderParser::ParseClassNameDeclaration(FClasses& AllClasses, FString& Dec
 		// a class cannot 'extends' an interface, use 'implements'
 		if (TempClass->ClassFlags & CLASS_Interface)
 		{
-			FError::Throwf(TEXT("Class '%s' cannot extend interface '%s', use 'implements'"), *Class->GetName(), *TempClass->GetName() );
+			FError::Throwf(TEXT("Class '%s' cannot extend interface '%s', use 'implements'"), *FoundClass->GetName(), *TempClass->GetName());
 		}
 
-		UClass* SuperClass = Class->GetSuperClass();
+		UClass* SuperClass = FoundClass->GetSuperClass();
 		if( SuperClass == NULL )
 		{
-			Class->SetSuperStruct(TempClass);
+			FoundClass->SetSuperStruct(TempClass);
 		}
 		else if( SuperClass != TempClass )
 		{
-			FError::Throwf(TEXT("%s's superclass must be %s, not %s"), *Class->GetPathName(), *SuperClass->GetPathName(), *TempClass->GetPathName() );
+			FError::Throwf(TEXT("%s's superclass must be %s, not %s"), *FoundClass->GetPathName(), *SuperClass->GetPathName(), *TempClass->GetPathName());
 		}
 
-		Class->ClassCastFlags |= Class->GetSuperClass()->ClassCastFlags;
+		FoundClass->ClassCastFlags |= FoundClass->GetSuperClass()->ClassCastFlags;
 
 		// Handle additional inherited interface classes
 		while (MatchSymbol(TEXT(",")))
@@ -4210,16 +4001,18 @@ void FHeaderParser::ParseClassNameDeclaration(FClasses& AllClasses, FString& Dec
 				}
 			}
 
-			HandleOneInheritedClass(AllClasses, *InterfaceName);
+			HandleOneInheritedClass(AllClasses, FoundClass, *InterfaceName);
 		}
 	}
-	else if( Class->GetSuperClass() )
+	else if (FoundClass->GetSuperClass())
 	{
-		FError::Throwf(TEXT("class: missing 'Extends %s'"), *Class->GetSuperClass()->GetName() );
+		FError::Throwf(TEXT("class: missing 'Extends %s'"), *FoundClass->GetSuperClass()->GetName());
 	}
+
+	return FoundClass;
 }
 
-void FHeaderParser::HandleOneInheritedClass(FClasses& AllClasses, FString InterfaceName)
+void FHeaderParser::HandleOneInheritedClass(FClasses& AllClasses, UClass* Class, FString InterfaceName)
 {
 	// Check for UInterface derived interface inheritance
 	if (UClass* Interface = AllClasses.FindScriptClass(InterfaceName))
@@ -4230,57 +4023,69 @@ void FHeaderParser::HandleOneInheritedClass(FClasses& AllClasses, FString Interf
 			FError::Throwf(TEXT("Implements: Class %s is not an interface; Can only inherit from non-UObjects or UInterface derived interfaces"), *Interface->GetName() );
 		}
 
-		// Make sure this class or a parent of this class didn't already implement the interface
-		for (UClass* TestClass = Class; TestClass != NULL; TestClass = TestClass->GetSuperClass())
-		{
-			for (int32 i = 0; i < TestClass->Interfaces.Num(); i++)
-			{
-				if (TestClass->Interfaces[i].Class == Interface)
-				{
-					FError::Throwf(TEXT("Implements: Interface '%s' is already implemented by '%s'"), *Interface->GetName(), *TestClass->GetName());
-				}
-			}
-		}
-
 		// Propagate the inheritable ClassFlags
 		Class->ClassFlags |= (Interface->ClassFlags) & CLASS_ScriptInherit;
 
 		new (Class->Interfaces) FImplementedInterface(Interface, 0, false);
 		if (Interface->HasAnyClassFlags(CLASS_Native))
 		{
-			ClassData->AddInheritanceParent(Interface);
+			GScriptHelper.FindClassData(Class)->AddInheritanceParent(Interface);
 		}
 	}
 	else
 	{
 		// Non-UObject inheritance
-		ClassData->AddInheritanceParent(InterfaceName);
+		GScriptHelper.FindClassData(Class)->AddInheritanceParent(InterfaceName);
 	}
 }
 
 /**
- * Compiles a class declaration.
+ * Setups basic class settings after parsing.
  */
-void FHeaderParser::CompileClassDeclaration(FClasses& AllClasses)
+void PostParsingClassSetup(UClass* Class)
 {
-	// Start of a class block.
-	CheckAllow( TEXT("'class'"), ALLOW_Class );
+	// Cleanup after first pass.
+	FHeaderParser::ComputeFunctionParametersSize(Class);
 
-	// Get categories inherited from the parent.
-	TArray<FString> HideCategories;
-	TArray<FString> ShowSubCatgories;
-	TArray<FString> HideFunctions;
-	TArray<FString> AutoExpandCategories;
-	TArray<FString> AutoCollapseCategories;
-	Class->GetHideCategories(HideCategories);
-	Class->GetShowCategories(ShowSubCatgories);
-	Class->GetHideFunctions(HideFunctions);
-	Class->GetAutoExpandCategories(AutoExpandCategories);
-	Class->GetAutoCollapseCategories(AutoCollapseCategories);
+	// Set all optimization ClassFlags based on property types
+	for (TFieldIterator<UProperty> It(Class, EFieldIteratorFlags::ExcludeSuper); It; ++It)
+	{
+		if (It->IsLocalized())
+		{
+			Class->ClassFlags |= CLASS_Localized;
+		}
 
-	// Class attributes.
-	FClassMetaData* ClassData = GScriptHelper.FindClassData(Class);
-	check(ClassData);
+		if ((It->PropertyFlags & CPF_Config) != 0)
+		{
+			Class->ClassFlags |= CLASS_Config;
+		}
+
+		if (It->ContainsInstancedObjectProperty())
+		{
+			Class->ClassFlags |= CLASS_HasInstancedReference;
+		}
+	}
+
+	// Class needs to specify which ini file is going to be used if it contains config variables.
+	if ((Class->ClassFlags & CLASS_Config) && (Class->ClassConfigName == NAME_None))
+	{
+		// Inherit config setting from base class.
+		Class->ClassConfigName = Class->GetSuperClass() ? Class->GetSuperClass()->ClassConfigName : NAME_None;
+		if (Class->ClassConfigName == NAME_None)
+		{
+			FError::Throwf(TEXT("Classes with config / globalconfig member variables need to specify config file."));
+			Class->ClassConfigName = NAME_Engine;
+		}
+	}
+}
+
+/**
+ * Compiles a class declaration.
+ */
+void FHeaderParser::CompileClassDeclaration(FClasses& AllClasses)
+{
+	// Start of a class block.
+	CheckAllow(TEXT("'class'"), ALLOW_Class);
 
 	// New-style UCLASS() syntax
 	TMap<FName, FString> MetaData;
@@ -4288,6 +4093,8 @@ void FHeaderParser::CompileClassDeclaration(FClasses& AllClasses)
 	TArray<FPropertySpecifier> SpecifiersFound;
 	ReadSpecifierSetInsideMacro(SpecifiersFound, TEXT("Class"), MetaData);
 
+	auto PrologFinishLine = InputLine;
+
 	AddFormattedPrevCommentAsTooltipMetaData(MetaData);
 
 	// New style files have the class name / extends afterwards
@@ -4297,8 +4104,52 @@ void FHeaderParser::CompileClassDeclaration(FClasses& AllClasses)
 
 	FString DeclaredClassName;
 	FString RequiredAPIMacroIfPresent;
-	ParseClassNameDeclaration(AllClasses, /*out*/ DeclaredClassName, /*out*/ RequiredAPIMacroIfPresent);
+	
+	FClass* Class = ParseClassNameDeclaration(AllClasses, /*out*/ DeclaredClassName, /*out*/ RequiredAPIMacroIfPresent);
 	ClassDefinitionRanges.Add(Class, ClassDefinitionRange(&Input[InputPos], nullptr));
+
+	check(Class);
+
+	Class->ClassFlags |= CLASS_Parsed;
+
+	PushNest(ENestType::NEST_Class, Class);
+
+	uint32 PreviousClassFlags = Class->ClassFlags;
+
+	ResetClassData();
+
+	// Verify class variables haven't been filled in
+	check(Class->Children == NULL);
+	check(Class->Next == NULL);
+	check(Class->NetFields.Num() == 0);
+
+	// Make sure our parent classes is parsed.
+	for (UClass* Temp = Class->GetSuperClass(); Temp; Temp = Temp->GetSuperClass())
+	{
+		if (!(Temp->ClassFlags & (CLASS_Parsed | CLASS_Intrinsic)))
+		{
+			FError::Throwf(TEXT("'%s' can't be compiled: Parent class '%s' has errors"), *Class->GetName(), *Temp->GetName());
+		}
+	}
+
+	// Get categories inherited from the parent.
+	TArray<FString> HideCategories;
+	TArray<FString> ShowSubCatgories;
+	TArray<FString> HideFunctions;
+	TArray<FString> AutoExpandCategories;
+	TArray<FString> AutoCollapseCategories;
+	Class->GetHideCategories(HideCategories);
+	Class->GetShowCategories(ShowSubCatgories);
+	Class->GetHideFunctions(HideFunctions);
+	Class->GetAutoExpandCategories(AutoExpandCategories);
+	Class->GetAutoCollapseCategories(AutoCollapseCategories);
+
+	// Class attributes.
+	FClassMetaData* ClassData = GScriptHelper.FindClassData(Class);
+	check(ClassData);
+
+	ClassData->SetPrologLine(PrologFinishLine);
+
 	// Record that this class is RequiredAPI if the CORE_API style macro was present
 	if (!RequiredAPIMacroIfPresent.IsEmpty())
 	{
@@ -4347,7 +4198,7 @@ void FHeaderParser::CompileClassDeclaration(FClasses& AllClasses)
 			{
 				FError::Throwf(TEXT("Classes cannot be 'within' interfaces"));
 			}
-			else if (Class->ClassWithin == NULL || Class->ClassWithin==UObject::StaticClass() || RequiredWithinClass->IsChildOf(Class->ClassWithin))
+			else if (Class->ClassWithin == NULL || Class->ClassWithin == UObject::StaticClass() || RequiredWithinClass->IsChildOf(Class->ClassWithin))
 			{
 				Class->ClassWithin = RequiredWithinClass;
 			}
@@ -4620,9 +4471,9 @@ void FHeaderParser::CompileClassDeclaration(FClasses& AllClasses)
 		? Class->GetSuperClass()->ClassWithin
 		: UObject::StaticClass();
 
-	if( !Class->ClassWithin->IsChildOf(ExpectedWithin) )
+	if (!Class->ClassWithin->IsChildOf(ExpectedWithin))
 	{
-		FError::Throwf(TEXT("Parent class declared within '%s'.  Cannot override within class with '%s' since it isn't a child"), *ExpectedWithin->GetName(), *Class->ClassWithin->GetName() );
+		FError::Throwf(TEXT("Parent class declared within '%s'.  Cannot override within class with '%s' since it isn't a child"), *ExpectedWithin->GetName(), *Class->ClassWithin->GetName());
 		///Class->ClassWithin = ExpectedWithin;
 	}
 
@@ -4634,12 +4485,12 @@ void FHeaderParser::CompileClassDeclaration(FClasses& AllClasses)
 	}
 
 	// Validate.
-	if( (Class->ClassFlags&CLASS_NoExport) )
+	if ((Class->ClassFlags&CLASS_NoExport))
 	{
 		// if the class's class flags didn't contain CLASS_NoExport before it was parsed, it means either:
 		// a) the DECLARE_CLASS macro for this native class doesn't contain the CLASS_NoExport flag (this is an error)
 		// b) this is a new native class, which isn't yet hooked up to static registration (this is OK)
-		if ( !(Class->ClassFlags&CLASS_Intrinsic) && (PreviousClassFlags & CLASS_NoExport) == 0 &&
+		if (!(Class->ClassFlags&CLASS_Intrinsic) && (PreviousClassFlags & CLASS_NoExport) == 0 &&
 			(PreviousClassFlags&CLASS_Native) != 0 )	// a new native class (one that hasn't been compiled into C++ yet) won't have this set
 		{
 			FError::Throwf(TEXT("'noexport': Must include CLASS_NoExport in native class declaration"));
@@ -4665,19 +4516,8 @@ void FHeaderParser::CompileClassDeclaration(FClasses& AllClasses)
 		Class->ClassConfigName = NAME_None;
 	}
 
-	// Add metadata for the include path
-	const FString *IncludePath = GClassIncludePathMap.Find(Class);
-	if(IncludePath != NULL)
-	{
-		MetaData.Add(TEXT("IncludePath"), *IncludePath);
-	}
-
-	// Add metadata for the module relative path
-	const FString *ModuleRelativePath = GClassModuleRelativePathMap.Find(Class);
-	if(ModuleRelativePath != NULL)
-	{
-		MetaData.Add(TEXT("ModuleRelativePath"), *ModuleRelativePath);
-	}
+	AddIncludePathToMetadata(Class, MetaData);
+	AddModuleRelativePathToMetadata(Class, MetaData);
 
 	// Register the metadata
 	AddMetaDataToClassData(Class, MetaData);
@@ -4686,19 +4526,16 @@ void FHeaderParser::CompileClassDeclaration(FClasses& AllClasses)
 	RequireSymbol( TEXT("{"), TEXT("'Class'") );
 
 	// Make visible outside the package.
-	Class->ClearFlags( RF_Transient );
+	Class->ClearFlags(RF_Transient);
 	check(Class->HasAnyFlags(RF_Public));
 	check(Class->HasAnyFlags(RF_Standalone));
 
 	// Copy properties from parent class.
 	if (Class->GetSuperClass())
 	{
-		Class->SetPropertiesSize( Class->GetSuperClass()->GetPropertiesSize() );
+		Class->SetPropertiesSize(Class->GetSuperClass()->GetPropertiesSize());
 	}
 
-	// Push the class nesting.
-	PushNest( NEST_Class, Class->GetFName(), NULL );
-
 	// auto-create properties for all of the VFTables needed for the multiple inheritances
 	// get the inheritance parents
 	auto& InheritanceParents = ClassData->GetInheritanceParents();
@@ -4721,14 +4558,22 @@ void FHeaderParser::CompileClassDeclaration(FClasses& AllClasses)
 	}
 }
 
-void FHeaderParser::ParseInterfaceNameDeclaration(FClasses& AllClasses, FString& DeclaredInterfaceName, FString& RequiredAPIMacroIfPresent)
+FClass* FHeaderParser::ParseInterfaceNameDeclaration(FClasses& AllClasses, FString& DeclaredInterfaceName, FString& RequiredAPIMacroIfPresent)
 {
 	ParseNameWithPotentialAPIMacroPrefix(/*out*/ DeclaredInterfaceName, /*out*/ RequiredAPIMacroIfPresent, TEXT("interface"));
 
+	FClass* FoundClass = AllClasses.FindClass(*GetClassNameWithPrefixRemoved(*DeclaredInterfaceName));
+	if (FoundClass == nullptr)
+	{
+		return nullptr;
+	}
+
 	// Get super interface
 	bool bSpecifiesParentClass = MatchSymbol(TEXT(":"));
 	if (!bSpecifiesParentClass)
-		return;
+	{
+		return FoundClass;
+	}
 
 	RequireIdentifier(TEXT("public"), TEXT("class inheritance"));
 
@@ -4742,18 +4587,20 @@ void FHeaderParser::ParseInterfaceNameDeclaration(FClasses& AllClasses, FString&
 			FError::Throwf(TEXT("Interface class '%s' cannot inherit from non-interface class '%s'"), *DeclaredInterfaceName, *TempClass->GetName() );
 	}
 
-	UClass* SuperClass = Class->GetSuperClass();
+	UClass* SuperClass = FoundClass->GetSuperClass();
 	if (SuperClass == NULL)
 	{
-		Class->SetSuperStruct(TempClass);
+		FoundClass->SetSuperStruct(TempClass);
 	}
 	else if (SuperClass != TempClass)
 	{
-		FError::Throwf(TEXT("%s's superclass must be %s, not %s"), *Class->GetPathName(), *SuperClass->GetPathName(), *TempClass->GetPathName() );
+		FError::Throwf(TEXT("%s's superclass must be %s, not %s"), *FoundClass->GetPathName(), *SuperClass->GetPathName(), *TempClass->GetPathName());
 	}
+
+	return FoundClass;
 }
 
-void FHeaderParser::ParseSecondInterfaceClass(FClasses& AllClasses)
+bool FHeaderParser::TryParseIInterfaceClass(FClasses& AllClasses)
 {
 	FString ErrorMsg(TEXT("C++ interface mix-in class declaration"));
 
@@ -4762,20 +4609,35 @@ void FHeaderParser::ParseSecondInterfaceClass(FClasses& AllClasses)
 	// Get a class name
 	FString DeclaredInterfaceName;
 	FString RequiredAPIMacroIfPresent;
-	ParseInterfaceNameDeclaration(AllClasses, /*out*/ DeclaredInterfaceName, /*out*/ RequiredAPIMacroIfPresent);
+	if (ParseInterfaceNameDeclaration(AllClasses, /*out*/ DeclaredInterfaceName, /*out*/ RequiredAPIMacroIfPresent) == nullptr)
+	{
+		return false;
+	}
 
-	// All classes must start with a valid Unreal prefix
-	const FString ExpectedInterfaceName = Class->GetNameWithPrefix(EEnforceInterfacePrefix::I);
-	if (DeclaredInterfaceName != ExpectedInterfaceName)
+	if (MatchSymbol(TEXT(";")))
 	{
-		FError::Throwf(TEXT("Interface name '%s' is invalid, the second interface class should be identified as '%s'"), *DeclaredInterfaceName, *ExpectedInterfaceName);
+		// Forward declaration.
+		return false;
+	}
+
+	if (DeclaredInterfaceName[0] != 'I')
+	{
+		return false;
+	}
+
+	UClass* FoundClass = nullptr;
+	if ((FoundClass = AllClasses.FindClass(*DeclaredInterfaceName.Mid(1))) == nullptr)
+	{
+		return false;
 	}
 
 	// Continue parsing the second class as if it were a part of the first (for reflection data purposes, it is)
 	RequireSymbol(TEXT("{"), *ErrorMsg);
 
 	// Push the interface class nesting again.
-	PushNest( NEST_Interface, Class->GetFName(), NULL );
+	PushNest(NEST_Interface, FoundClass);
+
+	return true;
 }
 
 /**
@@ -4797,32 +4659,34 @@ void FHeaderParser::CompileInterfaceDeclaration(FClasses& AllClasses)
 	// New-style UINTERFACE() syntax
 	ReadSpecifierSetInsideMacro(SpecifiersFound, TEXT("Interface"), MetaData);
 
+	auto PrologFinishLine = InputLine;
+
 	// New style files have the interface name / extends afterwards
 	RequireIdentifier(TEXT("class"), TEXT("Interface declaration"));
-	ParseInterfaceNameDeclaration(AllClasses, /*out*/ DeclaredInterfaceName, /*out*/ RequiredAPIMacroIfPresent);
+	FClass* InterfaceClass = ParseInterfaceNameDeclaration(AllClasses, /*out*/ DeclaredInterfaceName, /*out*/ RequiredAPIMacroIfPresent);
 
 
 	// Record that this interface is RequiredAPI if the CORE_API style macro was present
 	if (!RequiredAPIMacroIfPresent.IsEmpty())
 	{
-		Class->ClassFlags |= CLASS_RequiredAPI;
+		InterfaceClass->ClassFlags |= CLASS_RequiredAPI;
 	}
 
 	// Set the appropriate interface class flags
-	Class->ClassFlags |= CLASS_Interface | CLASS_Abstract;
-	if (Class->GetSuperStruct() != NULL)
+	InterfaceClass->ClassFlags |= CLASS_Interface | CLASS_Abstract;
+	if (InterfaceClass->GetSuperStruct() != NULL)
 	{
-		Class->ClassCastFlags |= Class->GetSuperClass()->ClassCastFlags;
+		InterfaceClass->ClassCastFlags |= InterfaceClass->GetSuperClass()->ClassCastFlags;
 	}
 
 	// All classes that are parsed are expected to be native
-	if (Class->GetSuperClass() && !Class->GetSuperClass()->HasAnyClassFlags(CLASS_Native))
+	if (InterfaceClass->GetSuperClass() && !InterfaceClass->GetSuperClass()->HasAnyClassFlags(CLASS_Native))
 	{
 		FError::Throwf(TEXT("Native classes cannot extend non-native classes") );
 	}
 
-	Class->SetFlags(RF_Native);
-	Class->ClassFlags |= CLASS_Native;
+	InterfaceClass->SetFlags(RF_Native);
+	InterfaceClass->ClassFlags |= CLASS_Native;
 
 	// Process all of the interface specifiers
 	for (TArray<FPropertySpecifier>::TIterator SpecifierIt(SpecifiersFound); SpecifierIt; ++SpecifierIt)
@@ -4836,7 +4700,7 @@ void FHeaderParser::CompileInterfaceDeclaration(FClasses& AllClasses)
 		}
 		else if (Specifier == TEXT("MinimalAPI"))
 		{
-			Class->ClassFlags |= CLASS_MinimalAPI;
+			InterfaceClass->ClassFlags |= CLASS_MinimalAPI;
 		}
 		else if (Specifier == TEXT("ConversionRoot"))
 		{
@@ -4849,30 +4713,32 @@ void FHeaderParser::CompileInterfaceDeclaration(FClasses& AllClasses)
 	}
 
 	// All classes must start with a valid Unreal prefix
-	const FString ExpectedInterfaceName = Class->GetNameWithPrefix(EEnforceInterfacePrefix::U);
+	const FString ExpectedInterfaceName = InterfaceClass->GetNameWithPrefix(EEnforceInterfacePrefix::U);
 	if (DeclaredInterfaceName != ExpectedInterfaceName)
 	{
 		FError::Throwf(TEXT("Interface name '%s' is invalid, the first class should be identified as '%s'"), *DeclaredInterfaceName, *ExpectedInterfaceName );
 	}
 
 	// Try parsing metadata for the interface
-	FClassMetaData* ClassData = GScriptHelper.FindClassData(Class);
+	FClassMetaData* ClassData = GScriptHelper.AddClassData(InterfaceClass);
 	check(ClassData);
 
+	ClassData->SetPrologLine(PrologFinishLine);
+
 	// Register the metadata
-	AddMetaDataToClassData( Class, MetaData );
+	AddMetaDataToClassData(InterfaceClass, MetaData);
 
 	// Handle the start of the rest of the interface
 	RequireSymbol( TEXT("{"), TEXT("'Class'") );
 
 	// Make visible outside the package.
-	Class->ClearFlags( RF_Transient );
-	check(Class->HasAnyFlags(RF_Public));
-	check(Class->HasAnyFlags(RF_Standalone));
+	InterfaceClass->ClearFlags(RF_Transient);
+	check(InterfaceClass->HasAnyFlags(RF_Public));
+	check(InterfaceClass->HasAnyFlags(RF_Standalone));
 
 	// Push the interface class nesting.
 	// we need a more specific set of allow flags for NEST_Interface, only function declaration is allowed, no other stuff are allowed
-	PushNest( NEST_Interface, Class->GetFName(), NULL );
+	PushNest(NEST_Interface, InterfaceClass);
 }
 
 // Returns true if the token is a dynamic delegate declaration
@@ -4895,7 +4761,7 @@ void FHeaderParser::ParseParameterList(FClasses& AllClasses, UFunction* Function
 		// Get parameter type.
 		FToken Property(CPT_None);
 		EObjectFlags ObjectFlags;
-		GetVarType( AllClasses, TopNode, Property, ObjectFlags, ~(CPF_ParmFlags|CPF_AutoWeak|CPF_RepSkip), TEXT("Function parameter"), NULL, EPropertyDeclarationStyle::None, (Function->FunctionFlags & FUNC_Net) ? EVariableCategory::ReplicatedParameter: EVariableCategory::RegularParameter);
+		GetVarType(AllClasses, GetCurrentScope(), Property, ObjectFlags, ~(CPF_ParmFlags | CPF_AutoWeak | CPF_RepSkip), TEXT("Function parameter"), NULL, EPropertyDeclarationStyle::None, (Function->FunctionFlags & FUNC_Net) ? EVariableCategory::ReplicatedParameter : EVariableCategory::RegularParameter);
 		Property.PropertyFlags |= CPF_Parm;
 
 		if (bExpectCommaBeforeName)
@@ -4903,7 +4769,7 @@ void FHeaderParser::ParseParameterList(FClasses& AllClasses, UFunction* Function
 			RequireSymbol(TEXT(","), TEXT("Delegate definitions require a , between the parameter type and parameter name"));
 		}
 
-		UProperty* Prop = GetVarNameAndDim(TopNode, Property, ObjectFlags, /*NoArrays=*/ false, /*IsFunction=*/ true, NULL, TEXT("Function parameter"));
+		UProperty* Prop = GetVarNameAndDim(Function, Property, ObjectFlags, /*NoArrays=*/ false, /*IsFunction=*/ true, NULL, TEXT("Function parameter"));
 
 		Function->NumParms++;
 
@@ -5008,16 +4874,11 @@ void FHeaderParser::ParseParameterList(FClasses& AllClasses, UFunction* Function
 	RequireSymbol( TEXT(")"), TEXT("parameter list") );
 }
 
-void FHeaderParser::CompileDelegateDeclaration(FClasses& AllClasses, const TCHAR* DelegateIdentifier, EDelegateSpecifierAction::Type SpecifierAction)
+void FHeaderParser::CompileDelegateDeclaration(FUnrealSourceFile& SourceFile, FClasses& AllClasses, const TCHAR* DelegateIdentifier, EDelegateSpecifierAction::Type SpecifierAction)
 {
 	TMap<FName, FString> MetaData;
 
-	// Add metadata for the module relative path
-	const FString *ModuleRelativePath = GClassModuleRelativePathMap.Find(Class);
-	if(ModuleRelativePath != NULL)
-	{
-		MetaData.Add(TEXT("ModuleRelativePath"), *ModuleRelativePath);
-	}
+	AddModuleRelativePathToMetadata(SourceFile, MetaData);
 
 	FFuncInfo            FuncInfo;
 
@@ -5090,7 +4951,7 @@ void FHeaderParser::CompileDelegateDeclaration(FClasses& AllClasses, const TCHAR
 
 	if (bHasReturnValue)
 	{
-		GetVarType( AllClasses, TopNode, ReturnType, ReturnValueFlags, 0, NULL, NULL, EPropertyDeclarationStyle::None, EVariableCategory::Return );
+		GetVarType(AllClasses, GetCurrentScope(), ReturnType, ReturnValueFlags, 0, NULL, NULL, EPropertyDeclarationStyle::None, EVariableCategory::Return);
 		RequireSymbol(TEXT(","), CurrentScopeName);
 	}
 
@@ -5122,28 +4983,28 @@ void FHeaderParser::CompileDelegateDeclaration(FClasses& AllClasses, const TCHAR
 		FCString::Strcpy( FuncInfo.Function.Identifier, *Name );
 	}
 
-	// Allocate local property frame, push nesting level and verify
-	// uniqueness at this scope level.
-	PushNest( NEST_FunctionDeclaration, FuncInfo.Function.Identifier, NULL );
-	UFunction* DelegateSignatureFunction = ((UFunction*)TopNode);
+	FuncInfo.DelegateMacroLine = InputLine;
+	auto* DelegateSignatureFunction = CreateDelegateFunction(FuncInfo);
+
+	GScriptHelper.AddClassData(DelegateSignatureFunction);
 
 	DelegateSignatureFunction->FunctionFlags |= FuncInfo.FunctionFlags;
 
 	FuncInfo.FunctionReference = DelegateSignatureFunction;
 	FuncInfo.SetFunctionNames();
-	ClassData->AddFunction(FuncInfo);
+	FFunctionData::Add(FuncInfo);
+	if (FuncInfo.FunctionReference->HasAnyFunctionFlags(FUNC_Delegate) && !GetCurrentScope()->IsFileScope())
+	{
+		GetCurrentClassData()->MarkContainsDelegate();
+	}
+
+	GetCurrentScope()->AddType(DelegateSignatureFunction);
 
 	// determine whether this function should be 'const'
 	if (bDeclaredConst)
 	{
 		DelegateSignatureFunction->FunctionFlags |= FUNC_Const;
 	}
-	else if (Class->HasAnyClassFlags(CLASS_Const))
-	{
-		// non-static functions in a const class must be const themselves
-		//@TODO: UCREMOVAL: Should this really apply to delegate signatures which don't really live in the class?
-		FError::Throwf(TEXT("Delegates declared in a const class must also be marked as const"));
-	}
 
 	// Get parameter list.
 	if (FoundParamCount)
@@ -5167,7 +5028,7 @@ void FHeaderParser::CompileDelegateDeclaration(FClasses& AllClasses, const TCHAR
 	if (bHasReturnValue)
 	{
 		ReturnType.PropertyFlags |= CPF_Parm | CPF_OutParm | CPF_ReturnParm;
-		UProperty* ReturnProp = GetVarNameAndDim(TopNode, ReturnType, ReturnValueFlags, /*NoArrays=*/ true, /*IsFunction=*/ true, TEXT("ReturnValue"), TEXT("Function return type"));
+		UProperty* ReturnProp = GetVarNameAndDim(DelegateSignatureFunction, ReturnType, ReturnValueFlags, /*NoArrays=*/ true, /*IsFunction=*/ true, TEXT("ReturnValue"), TEXT("Function return type"));
 
 		DelegateSignatureFunction->NumParms++;
 	}
@@ -5186,12 +5047,13 @@ void FHeaderParser::CompileDelegateDeclaration(FClasses& AllClasses, const TCHAR
 	DelegateSignatureFunction->Bind();
 
 	// End the nesting
-	PopNest( AllClasses, NEST_FunctionDeclaration, CurrentScopeName );
+	PostPopFunctionDeclaration(AllClasses, DelegateSignatureFunction);
 
 	// Don't allow delegate signatures to be redefined.
-	for (TFieldIterator<UFunction> FunctionIt(CastChecked<UStruct>(DelegateSignatureFunction->GetOuter()), EFieldIteratorFlags::IncludeSuper); FunctionIt; ++FunctionIt)
+	auto FunctionIterator = GetCurrentScope()->GetTypeIterator<UFunction>();
+	while (FunctionIterator.MoveNext())
 	{
-		UFunction* TestFunc = *FunctionIt;
+		UFunction* TestFunc = *FunctionIterator;
 		if ((TestFunc->GetFName() == DelegateSignatureFunction->GetFName()) && (TestFunc != DelegateSignatureFunction))
 		{
 			FError::Throwf(TEXT("Can't override delegate signature function '%s'"), FuncInfo.Function.Identifier);
@@ -5202,18 +5064,13 @@ void FHeaderParser::CompileDelegateDeclaration(FClasses& AllClasses, const TCHAR
 /**
  * Parses and compiles a function declaration
  */
-void FHeaderParser::CompileFunctionDeclaration(FClasses& AllClasses)
+void FHeaderParser::CompileFunctionDeclaration(FUnrealSourceFile& SourceFile, FClasses& AllClasses)
 {
 	CheckAllow(TEXT("'Function'"), ALLOW_Function);
 
 	TMap<FName, FString> MetaData;
 
-	// Add metadata for the module relative path
-	const FString *ModuleRelativePath = GClassModuleRelativePathMap.Find(Class);
-	if(ModuleRelativePath != NULL)
-	{
-		MetaData.Add(TEXT("ModuleRelativePath"), *ModuleRelativePath);
-	}
+	AddModuleRelativePathToMetadata(SourceFile, MetaData);
 
 	// New-style UFUNCTION() syntax 
 	TArray<FPropertySpecifier> SpecifiersFound;
@@ -5222,7 +5079,7 @@ void FHeaderParser::CompileFunctionDeclaration(FClasses& AllClasses)
 	FScriptLocation FuncNameRetry;
 	InitScriptLocation(FuncNameRetry);
 
-	if (!Class->HasAnyClassFlags(CLASS_Native))
+	if (!GetCurrentClass()->HasAnyClassFlags(CLASS_Native))
 	{
 		FError::Throwf(TEXT("Should only be here for native classes!"));
 	}
@@ -5258,7 +5115,7 @@ void FHeaderParser::CompileFunctionDeclaration(FClasses& AllClasses)
 	}
 
 	// non-static functions in a const class must be const themselves
-	if (Class->HasAnyClassFlags(CLASS_Const))
+	if (GetCurrentClass()->HasAnyClassFlags(CLASS_Const))
 	{
 		FuncInfo.FunctionFlags |= FUNC_Const;
 	}
@@ -5271,7 +5128,7 @@ void FHeaderParser::CompileFunctionDeclaration(FClasses& AllClasses)
 
 	ProcessFunctionSpecifiers(FuncInfo, SpecifiersFound);
 
-	if( (FuncInfo.FunctionFlags & FUNC_BlueprintPure) && Class->HasAnyClassFlags(CLASS_Interface) )
+	if ((FuncInfo.FunctionFlags & FUNC_BlueprintPure) && GetCurrentClass()->HasAnyClassFlags(CLASS_Interface))
 	{
 		// Until pure interface casts are supported, we don't allow pures in interfaces
 		FError::Throwf(TEXT("BlueprintPure specifier is not allowed for interface functions"));
@@ -5312,9 +5169,9 @@ void FHeaderParser::CompileFunctionDeclaration(FClasses& AllClasses)
 	}
 
 	// Verify interfaces with respect to their blueprint accessible functions
-	if (Class->HasAnyClassFlags(CLASS_Interface))
+	if (GetCurrentClass()->HasAnyClassFlags(CLASS_Interface))
 	{
-		const bool bCanImplementInBlueprints = !Class->HasMetaData(TEXT("CannotImplementInterfaceInBlueprint"));  //FBlueprintMetadata::MD_CannotImplementInterfaceInBlueprint
+		const bool bCanImplementInBlueprints = !GetCurrentClass()->HasMetaData(TEXT("CannotImplementInterfaceInBlueprint"));  //FBlueprintMetadata::MD_CannotImplementInterfaceInBlueprint
 		if((FuncInfo.FunctionFlags & FUNC_BlueprintEvent) != 0)
 		{
 			// Ensure that blueprint events are only allowed in implementable interfaces. Internal only functions allowed
@@ -5348,7 +5205,7 @@ void FHeaderParser::CompileFunctionDeclaration(FClasses& AllClasses)
 					//@TODO: Validate the module name for RequiredAPIMacroIfPresent
 					bThrowTokenBack = false;
 
-					if (Class->HasAnyClassFlags(CLASS_RequiredAPI))
+					if (GetCurrentClass()->HasAnyClassFlags(CLASS_RequiredAPI))
 					{
 						FError::Throwf(TEXT("'%s' must not be used on methods of a class that is marked '%s' itself."), *RequiredAPIMacroIfPresent, *RequiredAPIMacroIfPresent);
 					}
@@ -5377,7 +5234,7 @@ void FHeaderParser::CompileFunctionDeclaration(FClasses& AllClasses)
 		bAutomaticallyFinal = false;
 
 		// if this is a BlueprintNativeEvent or BlueprintImplementableEvent in an interface, make sure it's not "virtual"
-		if ((Class->HasAnyClassFlags(CLASS_Interface)) && (FuncInfo.FunctionFlags & FUNC_BlueprintEvent))
+		if ((GetCurrentClass()->HasAnyClassFlags(CLASS_Interface)) && (FuncInfo.FunctionFlags & FUNC_BlueprintEvent))
 		{
 			FError::Throwf(TEXT("BlueprintImplementableEvents in Interfaces must not be declared 'virtual'"));
 		}
@@ -5397,7 +5254,7 @@ void FHeaderParser::CompileFunctionDeclaration(FClasses& AllClasses)
 	else
 	{
 		// if this is a function in an Interface, it must be marked 'virtual' unless it's an event
-		if (Class->HasAnyClassFlags(CLASS_Interface) && !(FuncInfo.FunctionFlags & FUNC_BlueprintEvent))
+		if (GetCurrentClass()->HasAnyClassFlags(CLASS_Interface) && !(FuncInfo.FunctionFlags & FUNC_BlueprintEvent))
 		{
 			FError::Throwf(TEXT("Interface functions that are not BlueprintImplementableEvents must be declared 'virtual'"));
 		}
@@ -5410,7 +5267,7 @@ void FHeaderParser::CompileFunctionDeclaration(FClasses& AllClasses)
 		FuncInfo.FunctionFlags |= FUNC_Final;
 		FuncInfo.FunctionExportFlags |= FUNCEXPORT_Final;
 
-		if (Class->HasAnyClassFlags(CLASS_Interface))
+		if (GetCurrentClass()->HasAnyClassFlags(CLASS_Interface))
 		{
 			FError::Throwf(TEXT("Interface functions cannot be declared 'final'"));
 		}
@@ -5424,7 +5281,7 @@ void FHeaderParser::CompileFunctionDeclaration(FClasses& AllClasses)
 	// C++ style functions always have a return value type, even if it's void
 	if (!MatchIdentifier(TEXT("void")))
 	{
-		bHasReturnValue = GetVarType( AllClasses, TopNode, ReturnType, ReturnValueFlags, 0, NULL, NULL, EPropertyDeclarationStyle::None, EVariableCategory::Return );
+		bHasReturnValue = GetVarType(AllClasses, GetCurrentScope(), ReturnType, ReturnValueFlags, 0, NULL, NULL, EPropertyDeclarationStyle::None, EVariableCategory::Return);
 	}
 
 	// Get function or operator name.
@@ -5473,17 +5330,20 @@ void FHeaderParser::CompileFunctionDeclaration(FClasses& AllClasses)
 		}
 	}
 
-	// Allocate local property frame, push nesting level and verify
-	// uniqueness at this scope level.
-	PushNest( NEST_FunctionDeclaration, FuncInfo.Function.Identifier, NULL );
-	UFunction* TopFunction = ((UFunction*)TopNode);
+	auto* TopFunction = CreateFunction(FuncInfo);
+
+	GScriptHelper.AddClassData(TopFunction);
 
 	TopFunction->FunctionFlags |= FuncInfo.FunctionFlags;
 
 	FuncInfo.FunctionReference = TopFunction;
 	FuncInfo.SetFunctionNames();
 
-	FFunctionData* StoredFuncData = ClassData->AddFunction(FuncInfo);
+	auto* StoredFuncData = FFunctionData::Add(FuncInfo);
+	if (FuncInfo.FunctionReference->HasAnyFunctionFlags(FUNC_Delegate))
+	{
+		GetCurrentClassData()->MarkContainsDelegate();
+	}
 
 	// Get parameter list.
 	ParseParameterList(AllClasses, TopFunction, false, &MetaData);
@@ -5492,7 +5352,7 @@ void FHeaderParser::CompileFunctionDeclaration(FClasses& AllClasses)
 	if (bHasReturnValue)
 	{
 		ReturnType.PropertyFlags |= CPF_Parm | CPF_OutParm | CPF_ReturnParm;
-		UProperty* ReturnProp = GetVarNameAndDim(TopNode, ReturnType, ReturnValueFlags, /*NoArrays=*/ true, /*IsFunction=*/ true, TEXT("ReturnValue"), TEXT("Function return type"));
+		UProperty* ReturnProp = GetVarNameAndDim(TopFunction, ReturnType, ReturnValueFlags, /*NoArrays=*/ true, /*IsFunction=*/ true, TEXT("ReturnValue"), TEXT("Function return type"));
 
 		TopFunction->NumParms++;
 	}
@@ -5551,20 +5411,13 @@ void FHeaderParser::CompileFunctionDeclaration(FClasses& AllClasses)
 	// Handle C++ style functions being declared as abstract
 	if (MatchSymbol(TEXT("=")))
 	{
-		if (bHaveSeenSecondInterfaceClass)
-		{
-			int32 ZeroValue = 1;
-			bool bGotZero = GetConstInt(/*out*/ZeroValue);
-			bGotZero = bGotZero && (ZeroValue == 0);
+		int32 ZeroValue = 1;
+		bool bGotZero = GetConstInt(/*out*/ZeroValue);
+		bGotZero = bGotZero && (ZeroValue == 0);
 
-			if (!bGotZero)
-			{
-				FError::Throwf(TEXT("Expected 0 to indicate function is abstract"));
-			}
-		}
-		else
+		if (!bGotZero)
 		{
-			FError::Throwf(TEXT("Only functions in the second interface class can be declared abstract"));
+			FError::Throwf(TEXT("Expected 0 to indicate function is abstract"));
 		}
 	}
 
@@ -5574,7 +5427,7 @@ void FHeaderParser::CompileFunctionDeclaration(FClasses& AllClasses)
 		// This is a final (prebinding, non-overridable) function
 		FuncInfo.FunctionFlags |= FUNC_Final;
 		FuncInfo.FunctionExportFlags |= FUNCEXPORT_Final;
-		if (Class->HasAnyClassFlags(CLASS_Interface))
+		if (GetCurrentClass()->HasAnyClassFlags(CLASS_Interface))
 		{
 			FError::Throwf(TEXT("Interface functions cannot be declared 'final'"));
 		}
@@ -5591,113 +5444,111 @@ void FHeaderParser::CompileFunctionDeclaration(FClasses& AllClasses)
 
 	// Verify parameter list and return type compatibility within the
 	// function, if any, that it overrides.
-	for( int32 i=NestLevel-2; i>=1; i-- )
+	auto FunctionIterator = GetCurrentScope()->GetTypeIterator<UFunction>();
+	while (FunctionIterator.MoveNext())
 	{
-		for (UFunction* Function : TFieldRange<UFunction>(Nest[i].Node))
-		{
-			if (Function->GetFName() != TopNode->GetFName() || Function == TopNode)
-				continue;
+		UFunction* Function = *FunctionIterator;
+		if (Function->GetFName() != TopFunction->GetFName() || Function == TopFunction)
+			continue;
 
-			// Don't allow private functions to be redefined.
-			if (Function->FunctionFlags & FUNC_Private)
-				FError::Throwf(TEXT("Can't override private function '%s'"), FuncInfo.Function.Identifier);
+		// Don't allow private functions to be redefined.
+		if (Function->FunctionFlags & FUNC_Private)
+			FError::Throwf(TEXT("Can't override private function '%s'"), FuncInfo.Function.Identifier);
 
-			// see if they both either have a return value or don't
-			if ((TopFunction->GetReturnProperty() != NULL) != (Function->GetReturnProperty() != NULL))
-			{
-				ReturnToLocation(FuncNameRetry);
-				FError::Throwf(TEXT("Redefinition of '%s %s' differs from original: return value mismatch"), TypeOfFunction, FuncInfo.Function.Identifier );
-			}
+		// see if they both either have a return value or don't
+		if ((TopFunction->GetReturnProperty() != NULL) != (Function->GetReturnProperty() != NULL))
+		{
+			ReturnToLocation(FuncNameRetry);
+			FError::Throwf(TEXT("Redefinition of '%s %s' differs from original: return value mismatch"), TypeOfFunction, FuncInfo.Function.Identifier );
+		}
 
-			// See if all parameters match.
-			if (TopFunction->NumParms!=Function->NumParms)
-			{
-				ReturnToLocation(FuncNameRetry);
-				FError::Throwf(TEXT("Redefinition of '%s %s' differs from original; different number of parameters"), TypeOfFunction, FuncInfo.Function.Identifier );
-			}
+		// See if all parameters match.
+		if (TopFunction->NumParms!=Function->NumParms)
+		{
+			ReturnToLocation(FuncNameRetry);
+			FError::Throwf(TEXT("Redefinition of '%s %s' differs from original; different number of parameters"), TypeOfFunction, FuncInfo.Function.Identifier );
+		}
 
-			// Check all individual parameters.
-			int32 Count=0;
-			for( TFieldIterator<UProperty> CurrentFuncParam(TopFunction),SuperFuncParam(Function); Count<Function->NumParms; ++CurrentFuncParam,++SuperFuncParam,++Count )
+		// Check all individual parameters.
+		int32 Count=0;
+		for( TFieldIterator<UProperty> CurrentFuncParam(TopFunction),SuperFuncParam(Function); Count<Function->NumParms; ++CurrentFuncParam,++SuperFuncParam,++Count )
+		{
+			if( !FPropertyBase(*CurrentFuncParam).MatchesType(FPropertyBase(*SuperFuncParam), 1) )
 			{
-				if( !FPropertyBase(*CurrentFuncParam).MatchesType(FPropertyBase(*SuperFuncParam), 1) )
-				{
-					if( CurrentFuncParam->PropertyFlags & CPF_ReturnParm )
-					{
-						ReturnToLocation(FuncNameRetry);
-						FError::Throwf(TEXT("Redefinition of %s %s differs only by return type"), TypeOfFunction, FuncInfo.Function.Identifier );
-					}
-					else
-					{
-						ReturnToLocation(FuncNameRetry);
-						FError::Throwf(TEXT("Redefinition of '%s %s' differs from original"), TypeOfFunction, FuncInfo.Function.Identifier );
-					}
-					break;
-				}
-				else if ( CurrentFuncParam->HasAnyPropertyFlags(CPF_OutParm) != SuperFuncParam->HasAnyPropertyFlags(CPF_OutParm) )
+				if( CurrentFuncParam->PropertyFlags & CPF_ReturnParm )
 				{
 					ReturnToLocation(FuncNameRetry);
-					FError::Throwf(TEXT("Redefinition of '%s %s' differs from original - 'out' mismatch on parameter %i"), TypeOfFunction, FuncInfo.Function.Identifier, Count + 1);
+					FError::Throwf(TEXT("Redefinition of %s %s differs only by return type"), TypeOfFunction, FuncInfo.Function.Identifier );
 				}
-				else if ( CurrentFuncParam->HasAnyPropertyFlags(CPF_ReferenceParm) != SuperFuncParam->HasAnyPropertyFlags(CPF_ReferenceParm) )
+				else
 				{
 					ReturnToLocation(FuncNameRetry);
-					FError::Throwf(TEXT("Redefinition of '%s %s' differs from original - 'ref' mismatch on parameter %i"), TypeOfFunction, FuncInfo.Function.Identifier, Count + 1);
+					FError::Throwf(TEXT("Redefinition of '%s %s' differs from original"), TypeOfFunction, FuncInfo.Function.Identifier );
 				}
+				break;
 			}
-
-			if( Count<TopFunction->NumParms )
+			else if ( CurrentFuncParam->HasAnyPropertyFlags(CPF_OutParm) != SuperFuncParam->HasAnyPropertyFlags(CPF_OutParm) )
 			{
-				continue;
+				ReturnToLocation(FuncNameRetry);
+				FError::Throwf(TEXT("Redefinition of '%s %s' differs from original - 'out' mismatch on parameter %i"), TypeOfFunction, FuncInfo.Function.Identifier, Count + 1);
 			}
-
-			// if super version is event, overridden version must be defined as event (check before inheriting FUNC_Event)
-			if ( (Function->FunctionFlags & FUNC_Event) && !(FuncInfo.FunctionFlags & FUNC_Event) )
+			else if ( CurrentFuncParam->HasAnyPropertyFlags(CPF_ReferenceParm) != SuperFuncParam->HasAnyPropertyFlags(CPF_ReferenceParm) )
 			{
-				FError::Throwf(TEXT("Superclass version is defined as an event so '%s' should be!"), FuncInfo.Function.Identifier);
+				ReturnToLocation(FuncNameRetry);
+				FError::Throwf(TEXT("Redefinition of '%s %s' differs from original - 'ref' mismatch on parameter %i"), TypeOfFunction, FuncInfo.Function.Identifier, Count + 1);
 			}
-			// Function flags to copy from parent.
-			FuncInfo.FunctionFlags |= (Function->FunctionFlags & FUNC_FuncInherit);
+		}
 
-			// Make sure the replication conditions aren't being redefined
-			if ((FuncInfo.FunctionFlags & FUNC_NetFuncFlags) != (Function->FunctionFlags & FUNC_NetFuncFlags))
-			{
-				FError::Throwf(TEXT("Redefinition of replication conditions for function '%s'"), FuncInfo.Function.Identifier);
-			}
-			FuncInfo.FunctionFlags |= (Function->FunctionFlags & FUNC_NetFuncFlags);
+		if( Count<TopFunction->NumParms )
+		{
+			continue;
+		}
 
-			// Are we overriding a function?
-			if( TopNode==Function->GetOuter() )
-			{
-				// Duplicate.
-				ReturnToLocation( FuncNameRetry );
-				FError::Throwf(TEXT("Duplicate function '%s'"), *Function->GetName() );
-			}
-			// Overriding an existing function.
-			else if( Function->FunctionFlags & FUNC_Final )
-			{
-				ReturnToLocation(FuncNameRetry);
-				FError::Throwf(TEXT("%s: Can't override a 'final' function"), *Function->GetName() );
-			}
-			// Native function overrides should be done in CPP text, not in a UFUNCTION() declaration (you can't change flags, and it'd otherwise be a burden to keep them identical)
-			else if( Cast<UClass>(TopFunction->GetOuter()) != NULL )
-			{
-				//ReturnToLocation(FuncNameRetry);
-				FError::Throwf(TEXT("%s: An override of a function cannot have a UFUNCTION() declaration above it; it will use the same parameters as the original base declaration."), *Function->GetName() );
-			}
+		// if super version is event, overridden version must be defined as event (check before inheriting FUNC_Event)
+		if ( (Function->FunctionFlags & FUNC_Event) && !(FuncInfo.FunctionFlags & FUNC_Event) )
+		{
+			FError::Throwf(TEXT("Superclass version is defined as an event so '%s' should be!"), FuncInfo.Function.Identifier);
+		}
+		// Function flags to copy from parent.
+		FuncInfo.FunctionFlags |= (Function->FunctionFlags & FUNC_FuncInherit);
 
-			// Balk if required specifiers differ.
-			if ((Function->FunctionFlags & FUNC_FuncOverrideMatch) != (FuncInfo.FunctionFlags & FUNC_FuncOverrideMatch))
-			{
-				FError::Throwf(TEXT("Function '%s' specifiers differ from original"), *Function->GetName());
-			}
+		// Make sure the replication conditions aren't being redefined
+		if ((FuncInfo.FunctionFlags & FUNC_NetFuncFlags) != (Function->FunctionFlags & FUNC_NetFuncFlags))
+		{
+			FError::Throwf(TEXT("Redefinition of replication conditions for function '%s'"), FuncInfo.Function.Identifier);
+		}
+		FuncInfo.FunctionFlags |= (Function->FunctionFlags & FUNC_NetFuncFlags);
 
-			// Here we have found the original.
-			TopNode->SetSuperStruct(Function);
-			goto Found;
+		// Are we overriding a function?
+		if (TopFunction == Function->GetOuter())
+		{
+			// Duplicate.
+			ReturnToLocation( FuncNameRetry );
+			FError::Throwf(TEXT("Duplicate function '%s'"), *Function->GetName() );
+		}
+		// Overriding an existing function.
+		else if( Function->FunctionFlags & FUNC_Final )
+		{
+			ReturnToLocation(FuncNameRetry);
+			FError::Throwf(TEXT("%s: Can't override a 'final' function"), *Function->GetName() );
+		}
+		// Native function overrides should be done in CPP text, not in a UFUNCTION() declaration (you can't change flags, and it'd otherwise be a burden to keep them identical)
+		else if( Cast<UClass>(TopFunction->GetOuter()) != NULL )
+		{
+			//ReturnToLocation(FuncNameRetry);
+			FError::Throwf(TEXT("%s: An override of a function cannot have a UFUNCTION() declaration above it; it will use the same parameters as the original base declaration."), *Function->GetName() );
 		}
+
+		// Balk if required specifiers differ.
+		if ((Function->FunctionFlags & FUNC_FuncOverrideMatch) != (FuncInfo.FunctionFlags & FUNC_FuncOverrideMatch))
+		{
+			FError::Throwf(TEXT("Function '%s' specifiers differ from original"), *Function->GetName());
+		}
+
+		// Here we have found the original.
+		TopFunction->SetSuperStruct(Function);
+		break;
 	}
-Found:
 
 	// Bind the function.
 	TopFunction->Bind();
@@ -5731,7 +5582,7 @@ Found:
 	}
 
 	// Just declaring a function, so end the nesting.
-	PopNest( AllClasses, NEST_FunctionDeclaration, TypeOfFunction );
+	PostPopFunctionDeclaration(AllClasses, TopFunction);
 
 	// Optionally consume a semicolon
 	// This is optional to allow inline function definitions, as long as the function body is inside #if CPP ... #endif
@@ -5759,13 +5610,13 @@ void FHeaderParser::ValidateMetaDataFormat(UField* Field, const FString& InKey,
 			if (Function->HasAnyFunctionFlags(FUNC_Static))
 			{
 				// Determine if it's a function library
-				UClass* Class = Function->GetOuterUClass();
-				while ((Class->GetSuperClass() != UObject::StaticClass()) && (Class != NULL))
+				UClass* Class = Cast<UClass>(Function->GetOuterUClass());
+				while (Class != nullptr && Class->GetSuperClass() != UObject::StaticClass())
 				{
 					Class = Class->GetSuperClass();
 				}
 
-				if ((Class != NULL) && (Class->GetName() == TEXT("BlueprintFunctionLibrary")))
+				if (Class != nullptr && Class->GetName() == TEXT("BlueprintFunctionLibrary"))
 				{
 					FError::Throwf(TEXT("%s doesn't make sense on static method '%s' in a blueprint function library"), *InKey, *Function->GetName());
 				}
@@ -5818,7 +5669,7 @@ void FHeaderParser::AddMetaDataToClassData(UField* Field, const TMap<FName, FStr
 
 	// Finish validating and associate the metadata with the field
 	ValidateMetaDataFormat(Field, RemappedPairs);
-	ClassData->AddMetaData(Field, RemappedPairs);
+	FClassMetaData::AddMetaData(Field, RemappedPairs);
 }
 
 // Ensures at script compile time that the metadata formatting is correct
@@ -5988,7 +5839,7 @@ void FHeaderParser::CompileVariableDeclaration(FClasses& AllClasses, UStruct* St
 	EObjectFlags ObjectFlags = RF_NoFlags;
 
 	FIndexRange TypeRange;
-	GetVarType( AllClasses, Struct, OriginalProperty, ObjectFlags, DisallowFlags, TEXT("Member variable declaration"), /*OuterPropertyType=*/ NULL, PropertyDeclarationStyle, EVariableCategory::Member, &TypeRange );
+	GetVarType( AllClasses, &FScope::GetTypeScope(Struct).Get(), OriginalProperty, ObjectFlags, DisallowFlags, TEXT("Member variable declaration"), /*OuterPropertyType=*/ NULL, PropertyDeclarationStyle, EVariableCategory::Member, &TypeRange );
 	OriginalProperty.PropertyFlags |= EdFlags;
 
 	FString* Category = OriginalProperty.MetaData.Find("Category");
@@ -5999,7 +5850,7 @@ void FHeaderParser::CompileVariableDeclaration(FClasses& AllClasses, UStruct* St
 		if (!Category && (OriginalProperty.PropertyFlags & (CPF_Edit|CPF_BlueprintVisible)))
 		{
 			static const FString AbsoluteEngineDir = FPaths::ConvertRelativePathToFull(FPaths::EngineDir());
-			FString SourceFilename = GClassSourceFileMap[Class];
+			FString SourceFilename = GetCurrentSourceFile()->GetFilename();
 			FPaths::NormalizeFilename(SourceFilename);
 			if (Struct->GetOutermost() != nullptr && !SourceFilename.StartsWith(AbsoluteEngineDir))
 			{
@@ -6131,7 +5982,7 @@ void FHeaderParser::CompileVariableDeclaration(FClasses& AllClasses, UStruct* St
 // Compile a statement: Either a declaration or a command.
 // Returns 1 if success, 0 if end of file.
 //
-bool FHeaderParser::CompileStatement(FClasses& AllClasses)
+bool FHeaderParser::CompileStatement(FClasses& AllClasses, FUnrealSourceFile& SourceFile)
 {
 	// Get a token and compile it.
 	FToken Token;
@@ -6140,7 +5991,7 @@ bool FHeaderParser::CompileStatement(FClasses& AllClasses)
 		// End of file.
 		return false;
 	}
-	else if (!CompileDeclaration( AllClasses, Token ))
+	else if (!CompileDeclaration(AllClasses, SourceFile, Token))
 	{
 		FError::Throwf(TEXT("'%s': Bad command or expression"), Token.Identifier );
 	}
@@ -6252,8 +6103,6 @@ void FHeaderParser::SkipStatements( int32 NestCount, const TCHAR* ErrorTag  )
 //
 void FHeaderParser::FinalizeScriptExposedFunctions(UClass* Class)
 {
-	check(Class->ClassFlags & CLASS_Parsed);
-
 	// Finalize all of the children introduced in this class
 	for (TFieldIterator<UStruct> ChildIt(Class, EFieldIteratorFlags::ExcludeSuper); ChildIt; ++ChildIt)
 	{
@@ -6280,114 +6129,59 @@ void FHeaderParser::FinalizeScriptExposedFunctions(UClass* Class)
 // Parses the header associated with the specified class.
 // Returns result enumeration.
 //
-ECompilationResult::Type FHeaderParser::ParseHeaderForOneClass(FClasses& AllClasses, FClass* InClass)
+ECompilationResult::Type FHeaderParser::ParseHeader(FClasses& AllClasses, FUnrealSourceFile& SourceFile)
 {
+	if (SourceFile.IsParsed())
+	{
+		return ECompilationResult::Succeeded;
+	}
+
+	SourceFile.MarkAsParsed();
+
 	// Early-out if this class has previously failed some aspect of parsing
-	if (FailedClassesAnnotation.Get(InClass))
+	if (FailedFilesAnnotation.Get(&SourceFile))
 	{
 		return ECompilationResult::OtherCompilationError;
 	}
 
 	// Reset the parser to begin a new class
-	Class                                       = InClass;
 	bEncounteredNewStyleClass_UnmatchedBrackets = false;
 	bSpottedAutogeneratedHeaderInclude          = false;
-	bHaveSeenFirstInterfaceClass                = false;
-	bHaveSeenSecondInterfaceClass               = false;
-	bFinishedParsingInterfaceClasses            = false;
 	bHaveSeenUClass                             = false;
 	bClassHasGeneratedBody                      = false;
 
 	ECompilationResult::Type Result = ECompilationResult::OtherCompilationError;
 
-	ClassData = GScriptHelper.AddClassData(Class);
-
 	// Message.
 	if (FParse::Param(FCommandLine::Get(), TEXT("VERBOSE")))
 	{
 		// Message.
-		Warn->Logf( TEXT("Parsing %s"), *Class->GetName() );
-	}
-
-	// Make sure our parent classes is parsed.
-	for (UClass* Temp = Class->GetSuperClass(); Temp; Temp=Temp->GetSuperClass())
-	{
-		if (!(Temp->ClassFlags & (CLASS_Parsed | CLASS_Intrinsic)))
-		{
-			FError::Throwf(TEXT("'%s' can't be compiled: Parent class '%s' has errors"), *Class->GetName(), *Temp->GetName() );
-		}
-	}
-
-	// First pass.
-
-	//@fixme - reset class default object state?
-
-	Class->PropertiesSize = 0;
-
-	// Set class flags and within.
-	PreviousClassFlags = Class->ClassFlags;
-	Class->ClassFlags &= ~CLASS_RecompilerClear;
-
-	UClass* SuperClass = Class->GetSuperClass();
-	if (SuperClass != NULL)
-	{
-		Class->ClassFlags |= (SuperClass->ClassFlags) & CLASS_ScriptInherit;
-		Class->ClassConfigName = SuperClass->ClassConfigName;
-		check(SuperClass->ClassWithin);
-		if (Class->ClassWithin == NULL)
-		{
-			Class->ClassWithin = SuperClass->ClassWithin;
-		}
-
-		// Copy special categories from parent
-		if (SuperClass->HasMetaData(TEXT("HideCategories")))
-		{
-			Class->SetMetaData(TEXT("HideCategories"), *SuperClass->GetMetaData("HideCategories"));
-		}
-		if (SuperClass->HasMetaData(TEXT("ShowCategories")))
-		{
-			Class->SetMetaData(TEXT("ShowCategories"), *SuperClass->GetMetaData("ShowCategories"));
-		}
-		if (SuperClass->HasMetaData(TEXT("HideFunctions")))
-		{
-			Class->SetMetaData(TEXT("HideFunctions"), *SuperClass->GetMetaData("HideFunctions"));
-		}
-		if (SuperClass->HasMetaData(TEXT("AutoExpandCategories")))
-		{
-			Class->SetMetaData(TEXT("AutoExpandCategories"), *SuperClass->GetMetaData("AutoExpandCategories"));
-		}
-		if (SuperClass->HasMetaData(TEXT("AutoCollapseCategories")))
-		{
-			Class->SetMetaData(TEXT("AutoCollapseCategories"), *SuperClass->GetMetaData("AutoCollapseCategories"));
-		}
+		Warn->Logf(TEXT("Parsing %s"), *SourceFile.GetFilename());
 	}
 
-	check(Class->ClassWithin);
-
 	// Init compiler variables.
-	ResetParser(*GClassStrippedHeaderTextMap[Class]);
+	ResetParser(*SourceFile.GetContent());
 
 	// Init nesting.
 	NestLevel = 0;
 	TopNest = NULL;
-	PushNest( NEST_GlobalScope, TEXT(""), Class );
-
-	// Verify class variables haven't been filled in
-	check(Class->Children == NULL);
-	check(Class->Next == NULL);
-	check(Class->NetFields.Num() == 0);
+	PushNest(NEST_GlobalScope, nullptr, &SourceFile);
 
 	// C++ classes default to private access level
 	CurrentAccessSpecifier = ACCESS_Private; 
 
 	// Try to compile it, and catch any errors.
 	bool bEmptyFile = true;
+
+	// Tells if this header defines no-export classes only.
+	bool bNoExportClassesOnly = true;
+
 #if !PLATFORM_EXCEPTIONS_DISABLED
 	try
 #endif
 	{
 		// Parse entire program.
-		while (CompileStatement(AllClasses))
+		while (CompileStatement(AllClasses, SourceFile))
 		{
 			bEmptyFile = false;
 
@@ -6396,11 +6190,28 @@ ECompilationResult::Type FHeaderParser::ParseHeaderForOneClass(FClasses& AllClas
 			StatementsParsed++;
 		}
 
+		PopNest(NEST_GlobalScope, TEXT("Global scope"));
+
+		auto ScopeTypeIterator = SourceFile.GetScope()->GetTypeIterator();
+		while (ScopeTypeIterator.MoveNext())
+		{
+			auto* Type = *ScopeTypeIterator;
+
+			if (!Type->IsA<UScriptStruct>() && !Type->IsA<UClass>())
+			{
+				continue;
+			}
+
+			UStruct* Struct = Cast<UStruct>(Type);
+
+			// now validate all delegate variables declared in the class
+			TMap<FName, UFunction*> DelegateCache;
+			FixupDelegateProperties(AllClasses, Struct, FScope::GetTypeScope(Struct).Get(), DelegateCache);
+		}
+
 		// Precompute info for runtime optimization.
 		LinesParsed += InputLine;
 
-		Class->ClassFlags |= CLASS_Parsed;
-
 		if (RPCsNeedingHookup.Num() > 0)
 		{
 			FString ErrorMsg(TEXT("Request functions missing response pairs:\r\n"));
@@ -6413,69 +6224,41 @@ ECompilationResult::Type FHeaderParser::ParseHeaderForOneClass(FClasses& AllClas
 			FError::Throwf(*ErrorMsg);
 		}
 
-		// Make sure both classes were declared for interfaces
-		if (bHaveSeenFirstInterfaceClass && !bHaveSeenSecondInterfaceClass)
-		{
-			FError::Throwf(TEXT("Expected two class declarations to complete an interface (UMyInterface and IMyInterface)"));
-		}
-
 		// Make sure the compilation ended with valid nesting.
 		if (bEncounteredNewStyleClass_UnmatchedBrackets)
 		{
 			FError::Throwf(TEXT("Missing } at end of class") );
 		}
 
-		if (NestLevel == 0)
+		if (NestLevel == 1)
 		{
 			FError::Throwf(TEXT("Internal nest inconsistency") );
 		}
-		else if (NestLevel > 1)
+		else if (NestLevel > 2)
 		{
 			FError::Throwf(TEXT("Unexpected end of script in '%s' block"), NestTypeName(TopNest->NestType) );
 		}
 
-		// Cleanup after first pass.
-		ComputeFunctionParametersSize( Class );
+		// First-pass success.
+		Result = ECompilationResult::Succeeded;
 
-		// Set all optimization ClassFlags based on property types
-		for (TFieldIterator<UProperty> It(Class, EFieldIteratorFlags::ExcludeSuper); It; ++It)
+		for (auto* Class : SourceFile.GetDefinedClasses())
 		{
-			if (It->IsLocalized())
-			{
-				Class->ClassFlags |= CLASS_Localized;
-			}
-
-			if ((It->PropertyFlags & CPF_Config) != 0)
-			{
-				Class->ClassFlags |= CLASS_Config;
-			}
+			PostParsingClassSetup(Class);
 
-			if (It->ContainsInstancedObjectProperty())
-			{
-				Class->ClassFlags |= CLASS_HasInstancedReference;
-			}
-		}
+			// Clean up and exit.
+			Class->Bind();
 
-		// Class needs to specify which ini file is going to be used if it contains config variables.
-		if( (Class->ClassFlags & CLASS_Config) && (Class->ClassConfigName == NAME_None) )
-		{
-			// Inherit config setting from base class.
-			Class->ClassConfigName = Class->GetSuperClass() ? Class->GetSuperClass()->ClassConfigName : NAME_None;
-			if( Class->ClassConfigName == NAME_None )
+			// Finalize functions
+			if (Result == ECompilationResult::Succeeded)
 			{
-				FError::Throwf(TEXT("Classes with config / globalconfig member variables need to specify config file.") );
-				Class->ClassConfigName = NAME_Engine;
+				FinalizeScriptExposedFunctions(Class);
 			}
-		}
 
-		// mark temporary classes as native
-		if (FClassUtils::IsTemporaryClass(Class))
-		{
-			Class->ClassFlags |= CLASS_Native;
+			bNoExportClassesOnly = bNoExportClassesOnly && Class->HasAnyClassFlags(CLASS_NoExport);
 		}
 
-		// First-pass success.
-		Result = ECompilationResult::Succeeded;
+		check(SourceFile.IsParsed());
 	}
 #if !PLATFORM_EXCEPTIONS_DISABLED
 	catch( TCHAR* ErrorMsg )
@@ -6483,30 +6266,20 @@ ECompilationResult::Type FHeaderParser::ParseHeaderForOneClass(FClasses& AllClas
 		// Handle compiler error.
 		{
 			TGuardValue<ELogTimes::Type> DisableLogTimes(GPrintLogTimes, ELogTimes::None);
-			FString FormattedErrorMessage            = FString::Printf(TEXT("Error: In %s: %s\r\n"), *Class->GetName(), ErrorMsg);
-			FString FormattedErrorMessageWithContext = FString::Printf(TEXT("%s: %s"), *GetContext(), *FormattedErrorMessage);
+			FString FormattedErrorMessageWithContext = FString::Printf(TEXT("%s: %s"), *GetContext(), ErrorMsg);
 
 			UE_LOG(LogCompile, Log,  TEXT("%s"), *FormattedErrorMessageWithContext );
-			Warn->Log(ELogVerbosity::Error, FormattedErrorMessage );
+			Warn->Log(ELogVerbosity::Error, ErrorMsg);
 		}
 
-		FailedClassesAnnotation.Set(Class);
+		FailedFilesAnnotation.Set(&SourceFile);
 		Result = GCompilationResult;
 	}
 #endif
 
-	// Clean up and exit.
-	Class->Bind();
-
-	// Finalize functions
-	if (Result == ECompilationResult::Succeeded)
-	{
-		FinalizeScriptExposedFunctions( Class );
-	}
-
-	if (!bSpottedAutogeneratedHeaderInclude && !bEmptyFile && !Class->HasAnyClassFlags(CLASS_NoExport))
+	if (!bSpottedAutogeneratedHeaderInclude && !bEmptyFile && !bNoExportClassesOnly)
 	{
-		const FString ExpectedHeaderName = FString::Printf(TEXT("%s.generated.h"), *GClassHeaderNameWithNoPathMap[Class]);
+		const FString ExpectedHeaderName = SourceFile.GetGeneratedHeaderFilename();
 		FError::Throwf(TEXT("Expected an include at the top of the header: '#include \"%s\"'"), *ExpectedHeaderName);
 	}
 
@@ -6517,130 +6290,144 @@ ECompilationResult::Type FHeaderParser::ParseHeaderForOneClass(FClasses& AllClas
 	Global functions.
 -----------------------------------------------------------------------------*/
 
-// Parse Class's annotated headers and optionally its child classes.  Marks the class as CLASS_Parsed.
-ECompilationResult::Type FHeaderParser::ParseHeaders(FClasses& AllClasses, FHeaderParser& HeaderParser, FClass* Class, bool bParseSubclasses)
+ECompilationResult::Type FHeaderParser::ParseRestOfModulesSourceFiles(FClasses& AllClasses, UPackage* ModulePackage, FHeaderParser& HeaderParser)
 {
-	if (!GClassStrippedHeaderTextMap.Contains(Class))
-		return ECompilationResult::Succeeded;
+	for (auto& Pair : GUnrealSourceFilesMap)
+	{
+		FUnrealSourceFile& SourceFile = Pair.Value.Get();
+
+		if (SourceFile.GetPackage() == ModulePackage && (!SourceFile.IsParsed() || SourceFile.GetDefinedClassesCount() == 0))
+		{
+			ECompilationResult::Type Result;
+			if ((Result = ParseHeaders(AllClasses, HeaderParser, SourceFile, true)) != ECompilationResult::Succeeded)
+			{
+				return Result;
+			}
+		}
+	}
+
+	return ECompilationResult::Succeeded;
+}
 
+/**
+ * Parses source file of subclasses of classes defined in given file.
+ *
+ * @param HeaderParser Current parser.
+ * @param AllClasses Current classes tree.
+ * @param SourceFile Given source file.
+ *
+ * @returns Compilation result enum.
+ */
+/*
+ECompilationResult::Type ParseSourceFileSubclasses(FHeaderParser& HeaderParser, FClasses& AllClasses, FUnrealSourceFile& SourceFile)
+{
 	ECompilationResult::Type Result = ECompilationResult::Succeeded;
 
-	// Handle all dependencies of the class.
-	auto& DependentOn = *GClassDependentOnMap.FindOrAdd(Class);
-	for (int32 NameIndex = 0; NameIndex < DependentOn.Num(); NameIndex++)
+	for (auto* Class : SourceFile.GetDefinedClasses())
 	{
-		FString DependentClassName = DependentOn[NameIndex].ToString();
+		for (auto SubClass : AllClasses.GetDerivedClasses((FClass*)Class))
+		{
+			auto* ClassDefinitionInfo = GTypeDefinitionInfoMap.Find(SubClass);
+
+			if (ClassDefinitionInfo == nullptr)
+			{
+				continue;
+			}
 
-		// Check if the dependent class name came from #include directive in which case we allow it to fail.
-		DependentClassNameFromHeader(*DependentClassName, DependentClassName);
+			FUnrealSourceFile& SubclassSourceFile = (*ClassDefinitionInfo)->GetUnrealSourceFile();
 
-		FClass* DependsOnClass = AllClasses.FindScriptClass(DependentClassName);
-		if (!DependsOnClass)
-		{
-			// Try again with actor class name
-			const FString DependentClassNameStripped  = GetClassNameWithoutPrefix(DependentClassName);
-			const FString ActorDependentClassName = FString(TEXT("A")) + DependentClassNameStripped;
+			//note: you must always pass in the root tree node here, since we may add new classes to the tree
+			// if an manual dependency (through dependson()) is encountered
+			ECompilationResult::Type ParseResult = FHeaderParser::ParseHeaders(AllClasses, HeaderParser, SubclassSourceFile, true);
 
-			DependsOnClass = AllClasses.FindScriptClass(ActorDependentClassName);
-			if (!DependsOnClass)
+			if (ParseResult == ECompilationResult::FailedDueToHeaderChange)
 			{
-				// Try again with temporary class name. This may be struct only header.
-				const FString TemporaryDependentClassName = FString(TEXT("U")) + GenerateTemporaryClassName(*DependentClassNameStripped);
+				return ParseResult;
+			}
 
-				DependsOnClass = AllClasses.FindScriptClass(TemporaryDependentClassName);
-				if (!DependsOnClass)
-				{
-					// Ingore the error and remove this entry from DependentOn array.
-					DependentOn.RemoveAt(NameIndex--);
-					continue;
-				}
+			if (ParseResult == ECompilationResult::OtherCompilationError)
+			{
+				Result = ECompilationResult::OtherCompilationError;
 			}
 		}
+	}
 
-		// Detect potentially unnecessary usage of dependson.
-		if (DependsOnClass->HasAnyClassFlags(CLASS_Parsed))
-			continue;
+	return Result;
+}
+*/
 
-		// Check for circular dependency. If the DependsOnClass is dependent on the SubClass, there is one.
-		if (DependsOnClass != Class && AllClasses.IsDependentOn(DependsOnClass, Class))
-		{
-			HeaderParser.Class = Class;
-			HeaderParser.InputLine = GClassDeclarationLineNumber[Class];
-			FError::Throwf(TEXT("Error: Class %s DependsOn(%s) is a circular dependency."),*DependsOnClass->GetName(),*Class->GetName());
-		}
+// Parse Class's annotated headers and optionally its child classes.
+ECompilationResult::Type FHeaderParser::ParseHeaders(FClasses& AllClasses, FHeaderParser& HeaderParser, FUnrealSourceFile& SourceFile, bool bParseSubclasses)
+{
+	ECompilationResult::Type Result = ECompilationResult::Succeeded;
 
-		// Consider all children of the suspect if any of them are dependent on the source, the suspect is
-		// too because it itself is dependent on its children.
-		if (!DependsOnClass->HasAnyClassFlags(CLASS_Interface) && !AllClasses.ContainsClass(DependsOnClass))
-			FError::Throwf(TEXT("Unparsed class '%s' found while validating DependsOn entries for '%s'"), *DependsOnClass->GetName(), *Class->GetName());
+	if (SourceFile.AreDependenciesResolved())
+	{
+		return Result;
+	}
 
-		// Find first base class of DependsOnClass that is not a base class of Class.
-		TArray<FClass*> ClassesToParse;
-		ClassesToParse.Add(DependsOnClass);
+	SourceFile.MarkDependenciesResolved();
 
-		for ( FClass* ParentClass = DependsOnClass->GetSuperClass(); ParentClass && !ParentClass->HasAnyClassFlags(CLASS_Parsed|CLASS_Intrinsic); ParentClass = ParentClass->GetSuperClass() )
-		{
-			ClassesToParse.Add(ParentClass);
-		}
+	TArray<FUnrealSourceFile*> SourceFilesRequired;
 
-		while (ClassesToParse.Num() > 0)
+	for (auto& Include : SourceFile.GetIncludes())
+	{
+		if (Include.GetId() == "Object.h")
 		{
-			FClass* NextClass = ClassesToParse.Pop();
+			continue;
+		}
 
-			ECompilationResult::Type ParseResult = ParseHeaders(AllClasses, HeaderParser, NextClass, true);
+		FUnrealSourceFile* DepFile = Include.Resolve();
 
-			if (ParseResult == ECompilationResult::Succeeded)
-			{
-				break;
-			}
+		if (DepFile)
+		{
+			SourceFilesRequired.Add(DepFile);
+		}
+	}
 
-			ParseResult = ParseHeaders(AllClasses, HeaderParser, NextClass, false);
+	auto Classes = SourceFile.GetDefinedClasses();
 
-			if (ParseResult != ECompilationResult::Succeeded)
-			{
-				Result = ParseResult;
-				break;
-			}
+	for (auto* Class : Classes)
+	{
+		for (auto* ParentClass = Class->GetSuperClass(); ParentClass && !ParentClass->HasAnyClassFlags(CLASS_Parsed | CLASS_Intrinsic); ParentClass = ParentClass->GetSuperClass())
+		{
+			SourceFilesRequired.Add(&GTypeDefinitionInfoMap[ParentClass]->GetUnrealSourceFile());
 		}
 	}
 
-	// Parse the class
-	if (!(Class->ClassFlags & CLASS_Parsed))
+	for (auto* RequiredFile : SourceFilesRequired)
 	{
-		UClass* CurrentSuperClass = Class->GetSuperClass();
-		ECompilationResult::Type OneClassResult = HeaderParser.ParseHeaderForOneClass(AllClasses, Class);
-		if (OneClassResult != ECompilationResult::Succeeded)
+		SourceFile.GetScope()->IncludeScope(&RequiredFile->GetScope().Get());
+
+		ECompilationResult::Type SuperClassParseResult = ParseHeaders(AllClasses, HeaderParser, *RequiredFile, true);
+
+		if (SuperClassParseResult == ECompilationResult::Succeeded)
 		{
-			// if we couldn't parse this class, we won't be able to parse its children
-			return OneClassResult;
+			continue;
 		}
 
-		if (CurrentSuperClass != Class->GetSuperClass())
+		SuperClassParseResult = ParseHeaders(AllClasses, HeaderParser, *RequiredFile, false);
+
+		if (SuperClassParseResult != ECompilationResult::Succeeded)
 		{
-			// detect a native class that has changed parents and update the tree
-			AllClasses.ChangeParentClass(Class);
-			AllClasses.Validate();
+			Result = SuperClassParseResult;
+			break;
 		}
 	}
 
-	// Parse subclasses if instructed to do so.
-	if (bParseSubclasses)
+	// Parse the file
 	{
-		for (auto SubClass : AllClasses.GetDerivedClasses(Class))
-		{
-			//note: you must always pass in the root tree node here, since we may add new classes to the tree
-			// if an manual dependency (through dependson()) is encountered
-			ECompilationResult::Type ParseResult = ParseHeaders(AllClasses, HeaderParser, SubClass, bParseSubclasses);
+		ECompilationResult::Type OneFileResult = HeaderParser.ParseHeader(AllClasses, SourceFile);
 
-			if (ParseResult == ECompilationResult::FailedDueToHeaderChange)
-			{
-				return ParseResult;
-			}
+		for (auto* Class : Classes)
+		{
+			Class->ClassFlags |= CLASS_Parsed;
+		}
 
-			if (ParseResult == ECompilationResult::OtherCompilationError)
-			{
-				Result = ECompilationResult::OtherCompilationError;
-			}
+		if (OneFileResult != ECompilationResult::Succeeded)
+		{
+			// if we couldn't parse this file fail.
+			return OneFileResult;
 		}
 	}
 
@@ -6661,26 +6448,64 @@ bool FHeaderParser::DependentClassNameFromHeader(const TCHAR* HeaderFilename, FS
 	return false;
 }
 
-// Begins the process of exporting C++ class declarations for native classes in the specified package
-void FHeaderParser::ExportNativeHeaders( UPackage* CurrentPackage, FClasses& AllClasses, bool bAllowSaveExportedHeaders )
+/**
+ * Gets source files ordered by UCLASSes inheritance.
+ *
+ * @param CurrentPackage Current package.
+ * @param AllClasses Current class tree.
+ *
+ * @returns Array of source files.
+ */
+TArray<FUnrealSourceFile*> GetSourceFilesWithInheritanceOrdering(UPackage* CurrentPackage, FClasses& AllClasses)
 {
-	// Build a list of header filenames
-	TArray<FString>	ClassHeaderFilenames;
-	new(ClassHeaderFilenames) FString(TEXT(""));
+	TArray<FUnrealSourceFile*> SourceFiles;
 
-	bool bExportingHeaders = false;
+	auto Classes = AllClasses.GetClassesInPackage();
 
-	for (TObjectIterator<UClass> ClassIt; ClassIt; ++ClassIt)
+	// First add source files with the inheritance order.
+	for (auto* Class : Classes)
 	{
-		UClass* Cls = *ClassIt;
-		if( (CurrentPackage == NULL || Cls->GetOuter()==CurrentPackage) && GClassStrippedHeaderTextMap.Contains(Cls) && Cls->HasAnyClassFlags(CLASS_Native) && !Cls->HasAnyClassFlags(CLASS_Intrinsic) )
+		auto* DefinitionInfoPtr = GTypeDefinitionInfoMap.Find(Class);
+		if (DefinitionInfoPtr == nullptr)
 		{
-			bExportingHeaders = true;
-			break;
+			continue;
+		}
+
+		auto& SourceFile = (*DefinitionInfoPtr)->GetUnrealSourceFile();
+
+		if (!SourceFiles.Contains(&SourceFile)
+			&& SourceFile.GetScope()->ContainsTypes())
+		{
+			SourceFiles.Add(&SourceFile);
 		}
 	}
 
-	if (bExportingHeaders)
+	// Then add the rest.
+	for (auto& Pair : GUnrealSourceFilesMap)
+	{
+		auto& SourceFile = Pair.Value.Get();
+
+		if (SourceFile.GetPackage() == CurrentPackage
+			&& !SourceFiles.Contains(&SourceFile)
+			&& SourceFile.GetScope()->ContainsTypes())
+		{
+			SourceFiles.Add(&SourceFile);
+		}
+	}
+
+	return SourceFiles;
+}
+
+// Begins the process of exporting C++ class declarations for native classes in the specified package
+void FHeaderParser::ExportNativeHeaders( UPackage* CurrentPackage, FClasses& AllClasses, bool bAllowSaveExportedHeaders )
+{
+	// Build a list of header filenames
+	TArray<FString>	ClassHeaderFilenames;
+	new(ClassHeaderFilenames) FString(TEXT(""));
+
+	auto SourceFiles = GetSourceFilesWithInheritanceOrdering(CurrentPackage, AllClasses);
+
+	if (SourceFiles.Num() > 0)
 	{
 		const static bool bQuiet = !FParse::Param(FCommandLine::Get(),TEXT("VERBOSE"));
 		if ( CurrentPackage != NULL )
@@ -6707,17 +6532,15 @@ void FHeaderParser::ExportNativeHeaders( UPackage* CurrentPackage, FClasses& All
 		}
 
 		// Export native class definitions to package header files.
-		FNativeClassHeaderGenerator(CurrentPackage, AllClasses, bAllowSaveExportedHeaders);
+		FNativeClassHeaderGenerator(CurrentPackage, SourceFiles, AllClasses, bAllowSaveExportedHeaders);
 	}
 }
 
 FHeaderParser::FHeaderParser(FFeedbackContext* InWarn)
 : FBaseParser                       ()
 , Warn                              (InWarn)
-, Class                             (NULL)
 , bSpottedAutogeneratedHeaderInclude(false)
 , TopNest                           (NULL)
-, TopNode                           (NULL)
 {
 	FScriptLocation::Compiler = this;
 
@@ -6799,8 +6622,17 @@ FString FHeaderParser::RequireExactlyOneSpecifierValue(const FPropertySpecifier&
 // Exports the class to all vailable plugins
 void ExportClassToScriptPlugins(UClass* Class, const FManifestModule& Module, IScriptGeneratorPluginInterface& ScriptPlugin)
 {
-	auto ClassHeaderInfo = GClassGeneratedFileMap.FindRef(Class);
-	ScriptPlugin.ExportClass(Class, ClassHeaderInfo.SourceFilename, ClassHeaderInfo.GeneratedFilename, ClassHeaderInfo.bHasChanged);
+	auto DefinitionInfoRef = GTypeDefinitionInfoMap.Find(Class);
+	if (DefinitionInfoRef == nullptr)
+	{
+		const FString Empty = TEXT("");
+		ScriptPlugin.ExportClass(Class, Empty, Empty, false);
+	}
+	else
+	{
+		auto& SourceFile = (*DefinitionInfoRef)->GetUnrealSourceFile();
+		ScriptPlugin.ExportClass(Class, SourceFile.GetFilename(), SourceFile.GetGeneratedFilename(), SourceFile.HasChanged());
+	}
 }
 
 // Exports class tree to all available plugins
@@ -6831,7 +6663,7 @@ ECompilationResult::Type FHeaderParser::ParseAllHeadersInside(FClasses& ModuleCl
 	Warn->SetContext(&HeaderParser);
 
 	// Set up a filename for the error context if we don't even get as far parsing a class
-	HeaderParser.Filename = IFileManager::Get().ConvertToAbsolutePathForExternalAppForRead(*GClassSourceFileMap[ModuleClasses.GetRootClass()]);
+	HeaderParser.Filename = IFileManager::Get().ConvertToAbsolutePathForExternalAppForRead(*GTypeDefinitionInfoMap[ModuleClasses.GetRootClass()]->GetUnrealSourceFile().GetFilename());
 
 	// Hierarchically parse all classes.
 	ECompilationResult::Type Result = ECompilationResult::Succeeded;
@@ -6839,10 +6671,24 @@ ECompilationResult::Type FHeaderParser::ParseAllHeadersInside(FClasses& ModuleCl
 	try
 #endif
 	{
-		// Parse the headers
-		const bool bParseSubclasses = true;
+		for (auto* SourceFilePtr : GPublicSourceFileSet)
+		{
+			FUnrealSourceFile& SourceFile = *SourceFilePtr;
 
-		Result = FHeaderParser::ParseHeaders(ModuleClasses, HeaderParser, ModuleClasses.GetRootClass(), bParseSubclasses);
+			if (SourceFile.GetPackage() == CurrentPackage && (!SourceFile.IsParsed() || SourceFile.GetDefinedClassesCount() == 0))
+			{
+				ECompilationResult::Type Result;
+				if ((Result = ParseHeaders(ModuleClasses, HeaderParser, SourceFile, true)) != ECompilationResult::Succeeded)
+				{
+					return Result;
+				}
+			}
+		}
+
+		if (Result == ECompilationResult::Succeeded)
+		{
+			Result = FHeaderParser::ParseRestOfModulesSourceFiles(ModuleClasses, CurrentPackage, HeaderParser);
+		}
 
 		// Export the autogenerated code wrappers
 		if (Result == ECompilationResult::Succeeded)
@@ -6950,15 +6796,13 @@ void FHeaderParser::ParseClassName(const TCHAR* Temp, FString& ClassName)
 }
 
 // Performs a preliminary parse of the text in the specified buffer, pulling out useful information for the header generation process
-void FHeaderParser::SimplifiedClassParse(const TCHAR* InBuffer, bool& bIsInterface, TArray<FName>& DependentOn, FString& out_ClassName, FString& out_ParentClassName, int32& out_ClassDeclLine, FStringOutputDevice& ClassHeaderTextStrippedOfCppText)
+void FHeaderParser::SimplifiedClassParse(const TCHAR* InBuffer, TArray<FSimplifiedParsingClassInfo>& OutParsedClassArray, TArray<FHeaderProvider>& DependentOn, FStringOutputDevice& ClassHeaderTextStrippedOfCppText)
 {
 	FHeaderPreParser Parser;
 	FString StrLine;
 	FString ClassName;
 	FString BaseClassName;
 
-	out_ClassDeclLine = 1;
-
 	// Two passes, preprocessor, then looking for the class stuff
 
 	// The layer of multi-line comment we are in.
@@ -7127,7 +6971,6 @@ void FHeaderParser::SimplifiedClassParse(const TCHAR* InBuffer, bool& bIsInterfa
 		{
 			if (bFoundGeneratedInclude)
 			{
-				out_ClassDeclLine = CurrentLine;
 				FError::Throwf(TEXT("#include found after .generated.h file - the .generated.h file should always be the last #include in a header"));
 			}
 
@@ -7143,7 +6986,7 @@ void FHeaderParser::SimplifiedClassParse(const TCHAR* InBuffer, bool& bIsInterfa
 				if (HeaderFilenameEnd != INDEX_NONE)
 				{
 					// Include the extension in the name so that we later know where this entry came from.
-					DependentOn.Add(*FPaths::GetCleanFilename(DependsOnHeaderName.Mid(1, HeaderFilenameEnd - 1)));
+					DependentOn.Add(FHeaderProvider(EHeaderProviderSourceType::FileName, *FPaths::GetCleanFilename(DependsOnHeaderName.Mid(1, HeaderFilenameEnd - 1))));
 				}
 			}
 		}
@@ -7237,34 +7080,27 @@ void FHeaderParser::SimplifiedClassParse(const TCHAR* InBuffer, bool& bIsInterfa
 			Str = *StrLine;
 
 			// Get class or interface name
-			if (ClassName.IsEmpty())
+			if (const TCHAR* UInterfaceMacroDecl = FCString::Strfind(Str, TEXT("UINTERFACE(")))
 			{
-				if (const TCHAR* UInterfaceMacroDecl = FCString::Strfind(Str, TEXT("UINTERFACE(")))
-				{
-					out_ClassDeclLine = CurrentLine;
-					Parser.ParseClassDeclaration(StartOfLine + (UInterfaceMacroDecl - Str), CurrentLine, TEXT("UINTERFACE"), /*out*/ ClassName, /*out*/ BaseClassName, /*inout*/ DependentOn);
-					bIsInterface = true;
-				}
+				Parser.ParseClassDeclaration(StartOfLine + (UInterfaceMacroDecl - Str), CurrentLine, TEXT("UINTERFACE"), /*out*/ ClassName, /*out*/ BaseClassName, /*out*/ DependentOn);
+				OutParsedClassArray.Add(FSimplifiedParsingClassInfo(MoveTemp(ClassName), MoveTemp(BaseClassName), CurrentLine, true));
+			}
 
-				if (const TCHAR* UClassMacroDecl = FCString::Strfind(Str, TEXT("UCLASS(")))
-				{
-					out_ClassDeclLine = CurrentLine;
-					Parser.ParseClassDeclaration(StartOfLine + (UClassMacroDecl - Str), CurrentLine, TEXT("UCLASS"), /*out*/ ClassName, /*out*/ BaseClassName, /*inout*/ DependentOn);
-				}
+			if (const TCHAR* UClassMacroDecl = FCString::Strfind(Str, TEXT("UCLASS(")))
+			{
+				Parser.ParseClassDeclaration(StartOfLine + (UClassMacroDecl - Str), CurrentLine, TEXT("UCLASS"), /*out*/ ClassName, /*out*/ BaseClassName, /*out*/ DependentOn);
+				OutParsedClassArray.Add(FSimplifiedParsingClassInfo(MoveTemp(ClassName), MoveTemp(BaseClassName), CurrentLine, false));
 			}
 		}
 	
 		StartOfLine = Buffer;
 	}
-
-	out_ClassName = ClassName;
-	out_ParentClassName = BaseClassName;
 }
 
 /////////////////////////////////////////////////////
 // FHeaderPreParser
 
-void FHeaderPreParser::ParseClassDeclaration(const TCHAR* InputText, int32 InLineNumber, const TCHAR* StartingMatchID, FString& out_ClassName, FString& out_BaseClassName, TArray<FName>& inout_ClassNames)
+void FHeaderPreParser::ParseClassDeclaration(const TCHAR* InputText, int32 InLineNumber, const TCHAR* StartingMatchID, FString& out_ClassName, FString& out_BaseClassName, TArray<FHeaderProvider>& out_ClassNames)
 {
 	FString ErrorMsg = TEXT("Class declaration");
 
@@ -7313,7 +7149,7 @@ void FHeaderPreParser::ParseClassDeclaration(const TCHAR* InputText, int32 InLin
 			}
 
 			FName InterfaceClassName(InterfaceClassNameToken.Identifier);
-			inout_ClassNames.Add(InterfaceClassName);
+			out_ClassNames.Add(FHeaderProvider(EHeaderProviderSourceType::ClassName, InterfaceClassName.ToString().Mid(1)));
 		}
 	}
 
@@ -7327,7 +7163,7 @@ void FHeaderPreParser::ParseClassDeclaration(const TCHAR* InputText, int32 InLin
 				const FString& Value = *It2;
 
 				FName DependentOnClassName = FName(*Value);
-				inout_ClassNames.Add(DependentOnClassName);
+				out_ClassNames.Add(FHeaderProvider(EHeaderProviderSourceType::ClassName, DependentOnClassName.ToString()));
 			}
 		}
 	}
@@ -7414,10 +7250,11 @@ bool FHeaderParser::DefaultValueStringCppFormatToInnerFormat(const UProperty* Pr
 	else 
 	{
 		// Cache off the struct types, in case we need them later
-		static const UScriptStruct* VectorStruct = FindObjectChecked<UScriptStruct>(UObject::StaticClass(), TEXT("Vector"));
-		static const UScriptStruct* Vector2DStruct = FindObjectChecked<UScriptStruct>(UObject::StaticClass(), TEXT("Vector2D"));
-		static const UScriptStruct* RotatorStruct = FindObjectChecked<UScriptStruct>(UObject::StaticClass(), TEXT("Rotator"));
-		static const UScriptStruct* LinearColorStruct = FindObjectChecked<UScriptStruct>(UObject::StaticClass(), TEXT("LinearColor"));
+		UPackage* CoreUObjectPackage = UObject::StaticClass()->GetOutermost();
+		static const UScriptStruct* VectorStruct = FindObjectChecked<UScriptStruct>(CoreUObjectPackage, TEXT("Vector"));
+		static const UScriptStruct* Vector2DStruct = FindObjectChecked<UScriptStruct>(CoreUObjectPackage, TEXT("Vector2D"));
+		static const UScriptStruct* RotatorStruct = FindObjectChecked<UScriptStruct>(CoreUObjectPackage, TEXT("Rotator"));
+		static const UScriptStruct* LinearColorStruct = FindObjectChecked<UScriptStruct>(CoreUObjectPackage, TEXT("LinearColor"));
 
 		const UStructProperty* StructProperty = CastChecked<UStructProperty>(Property);
 		if( StructProperty->Struct == VectorStruct )
@@ -7563,7 +7400,7 @@ bool FHeaderParser::TryToMatchConstructorParameterList(FToken Token)
 		return false;
 	}
 
-	auto* ClassData = GScriptHelper.FindClassData(Class);
+	auto* ClassData = GScriptHelper.FindClassData(GetCurrentClass());
 
 	ClassData->bConstructorDeclared = true;
 
@@ -7674,3 +7511,182 @@ void FHeaderParser::SkipDeprecatedMacroIfNecessary()
 
 	RequireSymbol(TEXT(")"), TEXT("DEPRECATED macro"));
 }
+
+void FHeaderParser::ResetClassData()
+{
+	UClass* CurrentClass = GetCurrentClass();
+	CurrentClass->PropertiesSize = 0;
+
+	// Set class flags and within.
+	CurrentClass->ClassFlags &= ~CLASS_RecompilerClear;
+
+	UClass* SuperClass = CurrentClass->GetSuperClass();
+	if (SuperClass != NULL)
+	{
+		CurrentClass->ClassFlags |= (SuperClass->ClassFlags) & CLASS_ScriptInherit;
+		CurrentClass->ClassConfigName = SuperClass->ClassConfigName;
+		check(SuperClass->ClassWithin);
+		if (CurrentClass->ClassWithin == NULL)
+		{
+			CurrentClass->ClassWithin = SuperClass->ClassWithin;
+		}
+
+		// Copy special categories from parent
+		if (SuperClass->HasMetaData(TEXT("HideCategories")))
+		{
+			CurrentClass->SetMetaData(TEXT("HideCategories"), *SuperClass->GetMetaData("HideCategories"));
+		}
+		if (SuperClass->HasMetaData(TEXT("ShowCategories")))
+		{
+			CurrentClass->SetMetaData(TEXT("ShowCategories"), *SuperClass->GetMetaData("ShowCategories"));
+		}
+		if (SuperClass->HasMetaData(TEXT("HideFunctions")))
+		{
+			CurrentClass->SetMetaData(TEXT("HideFunctions"), *SuperClass->GetMetaData("HideFunctions"));
+		}
+		if (SuperClass->HasMetaData(TEXT("AutoExpandCategories")))
+		{
+			CurrentClass->SetMetaData(TEXT("AutoExpandCategories"), *SuperClass->GetMetaData("AutoExpandCategories"));
+		}
+		if (SuperClass->HasMetaData(TEXT("AutoCollapseCategories")))
+		{
+			CurrentClass->SetMetaData(TEXT("AutoCollapseCategories"), *SuperClass->GetMetaData("AutoCollapseCategories"));
+		}
+	}
+
+	check(CurrentClass->ClassWithin);
+}
+
+void FHeaderParser::PostPopNestClass(UClass* CurrentClass)
+{
+	// Validate all the rep notify events here, to make sure they're implemented
+	VerifyRepNotifyCallbacks(CurrentClass);
+
+	// Iterate over all the interfaces we claim to implement
+	for (auto& Impl : CurrentClass->Interfaces)
+	{
+		// And their super-classes
+		for (UClass* Interface = Impl.Class; Interface; Interface = Interface->GetSuperClass())
+		{
+			// If this interface is a common ancestor, skip it
+			if (CurrentClass->IsChildOf(Interface))
+				continue;
+
+			// So iterate over all functions this interface declares
+			for (auto InterfaceFunction : TFieldRange<UFunction>(Interface, EFieldIteratorFlags::ExcludeSuper))
+			{
+				bool Implemented = false;
+
+				// And try to find one that matches
+				for (UFunction* ClassFunction : TFieldRange<UFunction>(CurrentClass))
+				{
+					if (ClassFunction->GetFName() != InterfaceFunction->GetFName())
+						continue;
+
+					if ((InterfaceFunction->FunctionFlags & FUNC_Event) && !(ClassFunction->FunctionFlags & FUNC_Event))
+						FError::Throwf(TEXT("Implementation of function '%s' must be declared as 'event' to match declaration in interface '%s'"), *ClassFunction->GetName(), *Interface->GetName());
+
+					if ((InterfaceFunction->FunctionFlags & FUNC_Delegate) && !(ClassFunction->FunctionFlags & FUNC_Delegate))
+						FError::Throwf(TEXT("Implementation of function '%s' must be declared as 'delegate' to match declaration in interface '%s'"), *ClassFunction->GetName(), *Interface->GetName());
+
+					// Making sure all the parameters match up correctly
+					Implemented = true;
+
+					if (ClassFunction->NumParms != InterfaceFunction->NumParms)
+						FError::Throwf(TEXT("Implementation of function '%s' conflicts with interface '%s' - different number of parameters (%i/%i)"), *InterfaceFunction->GetName(), *Interface->GetName(), ClassFunction->NumParms, InterfaceFunction->NumParms);
+
+					int32 Count = 0;
+					for (TFieldIterator<UProperty> It1(InterfaceFunction), It2(ClassFunction); Count < ClassFunction->NumParms; ++It1, ++It2, Count++)
+					{
+						if (!FPropertyBase(*It1).MatchesType(FPropertyBase(*It2), 1))
+						{
+							if (It1->PropertyFlags & CPF_ReturnParm)
+							{
+								FError::Throwf(TEXT("Implementation of function '%s' conflicts only by return type with interface '%s'"), *InterfaceFunction->GetName(), *Interface->GetName());
+							}
+							else
+							{
+								FError::Throwf(TEXT("Implementation of function '%s' conflicts with interface '%s' - parameter %i '%s'"), *InterfaceFunction->GetName(), *Interface->GetName(), Count, *It1->GetName());
+							}
+						}
+					}
+				}
+
+				// Delegate signature functions are simple stubs and aren't required to be implemented (they are not callable)
+				if (InterfaceFunction->FunctionFlags & FUNC_Delegate)
+				{
+					Implemented = true;
+				}
+
+				// Verify that if this has blueprint-callable functions that are not implementable events, we've implemented them as a UFunction in the target class
+				if (!Implemented
+					&& !Interface->HasMetaData(TEXT("CannotImplementInterfaceInBlueprint"))  // FBlueprintMetadata::MD_CannotImplementInterfaceInBlueprint
+					&& InterfaceFunction->HasAnyFunctionFlags(FUNC_BlueprintCallable)
+					&& !InterfaceFunction->HasAnyFunctionFlags(FUNC_BlueprintEvent))
+				{
+					FError::Throwf(TEXT("Missing UFunction implementation of function '%s' from interface '%s'.  This function needs a UFUNCTION() declaration."), *InterfaceFunction->GetName(), *Interface->GetName());
+				}
+			}
+		}
+	}
+}
+
+void FHeaderParser::PostPopFunctionDeclaration(FClasses& AllClasses, UFunction* PoppedFunction)
+{
+	//@TODO: UCREMOVAL: Move this code to occur at delegate var declaration, and force delegates to be declared before variables that use them
+	if (!GetCurrentScope()->IsFileScope() && GetCurrentClassData()->ContainsDelegates())
+	{
+		// now validate all delegate variables declared in the class
+		TMap<FName, UFunction*> DelegateCache;
+		FixupDelegateProperties(AllClasses, PoppedFunction, *GetCurrentScope(), DelegateCache);
+	}
+}
+
+void FHeaderParser::PostPopNestInterface(FClasses& AllClasses, UClass* CurrentInterface)
+{
+	if (GScriptHelper.FindClassData(CurrentInterface)->ContainsDelegates())
+	{
+		TMap<FName, UFunction*> DelegateCache;
+		FixupDelegateProperties(AllClasses, CurrentInterface, FScope::GetTypeScope(ExactCast<UClass>(CurrentInterface)).Get(), DelegateCache);
+	}
+}
+
+template <class TFunctionType>
+TFunctionType* CreateFunctionImpl(const FFuncInfo& FuncInfo, UObject* Outer, FScope* CurrentScope)
+{
+	// Allocate local property frame, push nesting level and verify
+	// uniqueness at this scope level.
+	{
+		auto TypeIterator = CurrentScope->GetTypeIterator();
+		while (TypeIterator.MoveNext())
+		{
+			UField* Type = *TypeIterator;
+			if (Type->GetFName() == FuncInfo.Function.Identifier)
+			{
+				FError::Throwf(TEXT("'%s' conflicts with '%s'"), FuncInfo.Function.Identifier, *Type->GetFullName());
+			}
+		}
+	}
+
+	TFunctionType* Function = new(Outer, FuncInfo.Function.Identifier, RF_Public) TFunctionType(FObjectInitializer(), nullptr);
+	Function->RepOffset = MAX_uint16;
+	Function->ReturnValueOffset = MAX_uint16;
+	Function->FirstPropertyToInit = nullptr;
+
+	return Function;
+}
+
+UFunction* FHeaderParser::CreateFunction(const FFuncInfo &FuncInfo) const
+{
+	auto* Function = CreateFunctionImpl<UFunction>(FuncInfo, GetCurrentClass(), GetCurrentScope());
+
+	Function->Next = GetCurrentClass()->Children;
+	GetCurrentClass()->Children = Function;
+
+	return Function;
+}
+
+UDelegateFunction* FHeaderParser::CreateDelegateFunction(const FFuncInfo &FuncInfo) const
+{
+	return CreateFunctionImpl<UDelegateFunction>(FuncInfo, IsInAClass() ? (UObject*)GetCurrentClass() : (UObject*)GetCurrentFileScope()->GetSourceFile()->GetPackage(), GetCurrentScope());
+}
\ No newline at end of file
