UHT: Allows defining multiple UCLASSes in one header file.

diff --git a/Engine/Source/Programs/UnrealHeaderTool/Private/HeaderParser.h b/Engine/Source/Programs/UnrealHeaderTool/Private/HeaderParser.h
--- a/Engine/Source/Programs/UnrealHeaderTool/Private/HeaderParser.h
+++ b/Engine/Source/Programs/UnrealHeaderTool/Private/HeaderParser.h
@@ -45,10 +45,27 @@ namespace EDelegateSpecifierAction
 }
 
 /** Information for a particular nesting level. */
-struct FNestInfo
+class FNestInfo
 {
 	/** Link to the stack node. */
-	UStruct* Node;
+	FScope* Scope;
+
+public:
+	/**
+	 * Gets nesting scope.
+	 */
+	FScope* GetScope() const
+	{
+		return Scope;
+	}
+
+	/**
+	 * Sets nesting scope.
+	 */
+	void SetScope(FScope* Scope)
+	{
+		this->Scope = Scope;
+	}
 
 	/** Statement that caused the nesting. */
 	ENestType NestType;
@@ -83,7 +100,65 @@ struct ClassDefinitionRange
 	bool bHasGeneratedBody;
 };
 
-extern TMap<FClass*, ClassDefinitionRange> ClassDefinitionRanges;
+extern TMap<UClass*, ClassDefinitionRange> ClassDefinitionRanges;
+
+/**
+ * Simplified parsing information about UClasses.
+ */
+class FSimplifiedParsingClassInfo
+{
+public:
+	// Constructor.
+	FSimplifiedParsingClassInfo(FString ClassName, FString BaseClassName, int32 ClassDefLine, bool bClassIsAnInterface)
+		: ClassName(MoveTemp(ClassName)), BaseClassName(MoveTemp(BaseClassName)), ClassDefLine(ClassDefLine), bClassIsAnInterface(bClassIsAnInterface)
+	{}
+
+	/**
+	 * Gets class name.
+	 */
+	const FString& GetClassName() const
+	{
+		return ClassName;
+	}
+
+	/**
+	 * Gets base class name.
+	 */
+	const FString& GetBaseClassName() const
+	{
+		return BaseClassName;
+	}
+
+	/**
+	 * Gets class definition line number.
+	 */
+	int32 GetClassDefLine() const
+	{
+		return ClassDefLine;
+	}
+
+	/**
+	 * Tells if this class is an interface.
+	 */
+	bool IsInterface() const
+	{
+		return bClassIsAnInterface;
+	}
+
+private:
+	// Name.
+	FString ClassName;
+
+	// Super-class name.
+	FString BaseClassName;
+
+	// Class definition line.
+	int32 ClassDefLine;
+
+	// Is this an interface class?
+	bool bClassIsAnInterface;
+};
+
 /////////////////////////////////////////////////////
 // FHeaderParser
 
@@ -94,6 +169,9 @@ extern TMap<FClass*, ClassDefinitionRange> ClassDefinitionRanges;
 class FHeaderParser : public FBaseParser, public FContextSupplier
 {
 public:
+	// Compute the function parameter size and save the return offset
+	static void ComputeFunctionParametersSize(UClass* InClass);
+
 	// Parse all headers for classes that are inside LimitOuter.
 	static ECompilationResult::Type ParseAllHeadersInside(FClasses& ModuleClasses, FFeedbackContext* Warn, UPackage* LimitOuter, const FManifestModule& Module, TArray<class IScriptGeneratorPluginInterface*>& ScriptPlugins);
 
@@ -104,9 +182,9 @@ public:
 	//   
 	//  It also splits the buffer up into:
 	//   ScriptText (text outside of #if CPP and #if DEFAULTS blocks)
-	static void SimplifiedClassParse(const TCHAR* Buffer, bool& bIsInterface, TArray<FName>& DependentOn, FString& out_ClassName, FString& out_ParentClassName, int32& out_ClassDeclLine, FStringOutputDevice& ScriptText);
-	
-	/** 
+	static void SimplifiedClassParse(const TCHAR* Buffer, TArray<FSimplifiedParsingClassInfo>& OutParsedClassArray, TArray<FHeaderProvider>& DependentOn, FStringOutputDevice& ScriptText);
+
+	/**
 	 * Returns True if the given class name includes a valid Unreal prefix and matches up with the given original class Name.
 	 *
 	 * @param InNameToCheck - Name w/ potential prefix to check
@@ -123,17 +201,6 @@ public:
 	 */
 	static bool DependentClassNameFromHeader(const TCHAR* HeaderFilename, FString& OutClassName);
 
-	/**
-	 * Generates temporary class name (without 'u' prefix) from header name (no path, no extension).
-	 *
-	 * @param HeaderName Header name without path nor extension.
-	 * @return Temporary class name.
-	 */
-	static FString GenerateTemporaryClassName(const TCHAR* HeaderName)
-	{
-		return FString(TEXT("TemporaryUHTHeader_")) + HeaderName;
-	}
-
 	/**
 	 * Transforms CPP-formated string containing default value, to inner formated string
 	 * If it cannot be transformed empty string is returned.
@@ -145,15 +212,24 @@ public:
 	 */
 	static bool DefaultValueStringCppFormatToInnerFormat(const UProperty* Property, const FString& CppForm, FString &InnerForm);
 
+	/**
+	 * Parse Class's annotated headers and optionally its child classes.  Marks the class as CLASS_Parsed.
+	 *
+	 * @param	AllClasses			the class tree containing all classes in the current package
+	 * @param	HeaderParser		the header parser
+	 * @param	SourceFile			Source file info.
+	 * @param	bParseSubclasses	true if we should parse all child classes of Class
+	 *
+	 * @return	Result enumeration.
+	 */
+	static ECompilationResult::Type ParseHeaders(FClasses& AllClasses, FHeaderParser& HeaderParser, FUnrealSourceFile& SourceFile, bool bParseSubclasses);
+
 protected:
 	friend struct FScriptLocation;
 
 	// For compiling messages and errors.
 	FFeedbackContext* Warn;
 
-	// Class that is currently being parsed.
-	FClass* Class;
-
 	// Filename currently being parsed
 	FString Filename;
 
@@ -166,8 +242,80 @@ protected:
 	// Top nesting level.
 	FNestInfo* TopNest;
 
-	// Top stack node.
-	UStruct* TopNode;
+	/**
+	 * Gets current nesting scope.
+	 */
+	FScope* GetCurrentScope() const
+	{
+		return TopNest->GetScope();
+	}
+
+	/**
+	 * Gets current file scope.
+	 */
+	FFileScope* GetCurrentFileScope() const
+	{
+		int32 Index = 0;
+		while (TopNest[Index].NestType != NEST_GlobalScope)
+		{
+			--Index;
+		}
+
+		return (FFileScope*)TopNest[Index].GetScope();
+	}
+
+	/**
+	 * Gets current source file.
+	 */
+	FUnrealSourceFile* GetCurrentSourceFile() const
+	{
+		return GetCurrentFileScope()->GetSourceFile();
+	}
+
+	/**
+	 * Gets current class scope.
+	 */
+	FStructScope* GetCurrentClassScope() const
+	{
+		check(TopNest->NestType == NEST_Class || TopNest->NestType == NEST_Interface);
+
+		return (FStructScope*)TopNest->GetScope();
+	}
+
+	/**
+	 * Tells if parser is currently in a class.
+	 */
+	bool IsInAClass() const
+	{
+		int32 Index = 0;
+		while (TopNest[Index].NestType != NEST_GlobalScope)
+		{
+			if (TopNest[Index].NestType == NEST_Class || TopNest->NestType == NEST_Interface)
+			{
+				return true;
+			}
+
+			--Index;
+		}
+
+		return false;
+	}
+
+	/**
+	 * Gets current class.
+	 */
+	UClass* GetCurrentClass() const
+	{
+		return (UClass*)GetCurrentClassScope()->GetStruct();
+	}
+
+	/**
+	 * Gets current class's metadata.
+	 */
+	FClassMetaData* GetCurrentClassData()
+	{
+		return GScriptHelper.FindClassData(GetCurrentClass());
+	}
 
 	// Information about all nesting levels.
 	FNestInfo Nest[MAX_NEST_LEVELS];
@@ -214,18 +362,6 @@ protected:
 	 */
 	uint32 PreviousClassFlags;
 
-	//
-	FClassMetaData* ClassData;
-
-	// Have seen the first of the two expected classes in the new style .h files?
-	bool bHaveSeenFirstInterfaceClass;
-
-	// Have seen the second of the two expected classes in the new style .h files?
-	bool bHaveSeenSecondInterfaceClass;
-
-	// Indicates that both interface classes have been parsed.
-	bool bFinishedParsingInterfaceClasses;
-
 	// For new-style classes, used to keep track of an unmatched {} pair
 	bool bEncounteredNewStyleClass_UnmatchedBrackets;
 
@@ -287,16 +423,15 @@ protected:
 	FString RequireExactlyOneSpecifierValue(const FPropertySpecifier& Specifier);
 
 	/**
-	 * Parse Class's annotated headers and optionally its child classes.  Marks the class as CLASS_Parsed.
+	 * Parse rest of the module's source files.
 	 *
 	 * @param	AllClasses			the class tree containing all classes in the current package
+	 * @param	ModulePackage		current package
 	 * @param	HeaderParser		the header parser
-	 * @param	Class				the class to parse
-	 * @param	bParseSubclasses	true if we should parse all child classes of Class
 	 *
 	 * @return	Result enumeration.
 	 */
-	static ECompilationResult::Type ParseHeaders(FClasses& AllClasses, FHeaderParser& HeaderParser, FClass* Class, bool bParseSubclasses);
+	static ECompilationResult::Type ParseRestOfModulesSourceFiles(FClasses& AllClasses, UPackage* ModulePackage, FHeaderParser& HeaderParser);
 
 	//@TODO: Remove this method
 	static void ParseClassName(const TCHAR* Temp, FString& ClassName);
@@ -320,19 +455,27 @@ protected:
 	// End of FContextSupplier interface.
 
 	// High-level compiling functions.
-	ECompilationResult::Type ParseHeaderForOneClass(FClasses& AllClasses, FClass* InClass);
-	void CompileDirective(UClass* Class);
+	/**
+	 * Parses given source file.
+	 *
+	 * @param AllClasses The class tree for current package.
+	 * @param SourceFile Source file to parse.
+	 *
+	 * @returns Compilation result enum.
+	 */
+	ECompilationResult::Type ParseHeader(FClasses& AllClasses, FUnrealSourceFile& SourceFile);
+	void CompileDirective(FClasses& AllClasses, FUnrealSourceFile& SourceFile);
 	void FinalizeScriptExposedFunctions(UClass* Class);
-	UEnum* CompileEnum(UClass* Owner);
-	UScriptStruct* CompileStructDeclaration(FClasses& AllClasses, FClass* Owner);
-	bool CompileDeclaration(FClasses& AllClasses, FToken& Token);
+	UEnum* CompileEnum(FUnrealSourceFile& SourceFile);
+	UScriptStruct* CompileStructDeclaration(FClasses& AllClasses, FUnrealSourceFile& SourceFile);
+	bool CompileDeclaration(FClasses& AllClasses, FUnrealSourceFile& SourceFile, FToken& Token);
 
 	/** Skip C++ (noexport) declaration. */
 	bool SkipDeclaration(FToken& Token);
 	/** Similar to MatchSymbol() but will return to the exact location as on entry if the symbol was not found. */
 	bool SafeMatchSymbol(const TCHAR* Match);
-	void HandleOneInheritedClass(FClasses& AllClasses, FString InterfaceName);
-	void ParseClassNameDeclaration(FClasses& AllClasses, FString& DeclaredClassName, FString& RequiredAPIMacroIfPresent);
+	void HandleOneInheritedClass(FClasses& AllClasses, UClass* Class, FString InterfaceName);
+	FClass* ParseClassNameDeclaration(FClasses& AllClasses, FString& DeclaredClassName, FString& RequiredAPIMacroIfPresent);
 
 	/** The property style of a variable declaration being parsed */
 	struct EPropertyDeclarationStyle
@@ -344,19 +487,31 @@ protected:
 		};
 	};
 
-	void CompileClassDeclaration    (FClasses& AllClasses);
-	void CompileDelegateDeclaration (FClasses& AllClasses, const TCHAR* DelegateIdentifier, EDelegateSpecifierAction::Type SpecifierAction = EDelegateSpecifierAction::DontParse);
-	void CompileFunctionDeclaration (FClasses& AllClasses);
+	/**
+	 * Resets current class data back to its defaults.
+	 */
+	void ResetClassData();
+
+	/**
+	 * Create new function object based on given info structure.
+	 */
+	UFunction* CreateFunction(const FFuncInfo &FuncInfo) const;
+
+	/**
+	 * Create new delegate function object based on given info structure.
+	 */
+	UDelegateFunction* CreateDelegateFunction(const FFuncInfo &FuncInfo) const;
+
+	void CompileClassDeclaration(FClasses& AllClasses);
+	void CompileDelegateDeclaration(FUnrealSourceFile& SourceFile, FClasses& AllClasses, const TCHAR* DelegateIdentifier, EDelegateSpecifierAction::Type SpecifierAction = EDelegateSpecifierAction::DontParse);
+	void CompileFunctionDeclaration(FUnrealSourceFile& SourceFile, FClasses& AllClasses);
 	void CompileVariableDeclaration (FClasses& AllClasses, UStruct* Struct, EPropertyDeclarationStyle::Type PropertyDeclarationStyle);
 	void CompileInterfaceDeclaration(FClasses& AllClasses);
 
-	void ParseInterfaceNameDeclaration(FClasses& AllClasses, FString& DeclaredInterfaceName, FString& RequiredAPIMacroIfPresent);
-	void ParseSecondInterfaceClass(FClasses& AllClasses);
+	FClass* ParseInterfaceNameDeclaration(FClasses& AllClasses, FString& DeclaredInterfaceName, FString& RequiredAPIMacroIfPresent);
+	bool TryParseIInterfaceClass(FClasses& AllClasses);
 
-	bool CompileStatement(FClasses& AllClasses);
-
-	// Compute the function parameter size and save the return offset
-	static void ComputeFunctionParametersSize( UClass* InClass );
+	bool CompileStatement(FClasses& AllClasses, FUnrealSourceFile& SourceFile);
 
 	// Checks to see if a particular kind of command is allowed on this nesting level.
 	bool IsAllowedInThisNesting(uint32 AllowFlags);
@@ -366,8 +521,6 @@ protected:
 	// nesting level.
 	void CheckAllow(const TCHAR* Thing, uint32 AllowFlags);
 
-	void CheckInScope( UObject* Obj );
-
 	UStruct* GetSuperScope( UStruct* CurrentScope, const FName& SearchName );
 
 	/**
@@ -416,7 +569,7 @@ protected:
 	 */
 	bool GetVarType(
 		FClasses&                       AllClasses,
-		UStruct*                        Scope,
+		FScope*							Scope,
 		FPropertyBase&                  VarProperty,
 		EObjectFlags&                   ObjectFlags,
 		uint64                          Disallow,
@@ -449,9 +602,15 @@ protected:
 		const TCHAR* HardcodedName,
 		const TCHAR* Thing);
 	
-	void CheckObscures(UStruct* Scope, const FString& ScriptName, const FString& FieldName);
-	
-	bool AllowReferenceToClass( UClass* CheckClass ) const;
+	/**
+	 * Returns whether the specified class can be referenced from the class currently being compiled.
+	 *
+	 * @param	Scope		The scope we are currently parsing.
+	 * @param	CheckClass	The class we want to reference.
+	 *
+	 * @return	true if the specified class is an intrinsic type or if the class has successfully been parsed
+	 */
+	bool AllowReferenceToClass(UStruct* Scope, UClass* CheckClass) const;
 
 	/**
 	 * @return	true if Scope has UProperty objects in its list of fields
@@ -501,9 +660,42 @@ protected:
 
 	FClass* GetQualifiedClass(const FClasses& AllClasses, const TCHAR* Thing);
 
-	// Nest management functions.
-	void PushNest( ENestType NestType, FName ThisName, UStruct* InNode );
-	void PopNest ( FClasses& AllClasses, ENestType NestType, const TCHAR* Descr );
+	/**
+	 * Increase the nesting level, setting the new top nesting level to
+	 * the one specified.  If pushing a function or state and it overrides a similar
+	 * thing declared on a lower nesting level, verifies that the override is legal.
+	 *
+	 * @param	NestType	the new nesting type
+	 * @param	InNode		@todo
+	 */
+	void PushNest(ENestType NestType, UStruct* InNode, FUnrealSourceFile* SourceFile = nullptr);
+	void PopNest(ENestType NestType, const TCHAR* Descr);
+
+	/**
+	 * Tasks that need to be done after popping function declaration
+	 * from parsing stack.
+	 *
+	 * @param AllClasses The class tree for current package.
+	 * @param PoppedFunction Function that have just been popped.
+	 */
+	void PostPopFunctionDeclaration(FClasses& AllClasses, UFunction* PoppedFunction);
+
+	/**
+	 * Tasks that need to be done after popping interface definition
+	 * from parsing stack.
+	 *
+	 * @param AllClasses The class tree for current package.
+	 * @param CurrentInterface Interface that have just been popped.
+	 */
+	void PostPopNestInterface(FClasses& AllClasses, UClass* CurrentInterface);
+
+	/**
+	 * Tasks that need to be done after popping class definition
+	 * from parsing stack.
+	 *
+	 * @param CurrentClass Class that have just been popped.
+	 */
+	void PostPopNestClass(UClass* CurrentClass);
 
 	/**
 	 * Binds all delegate properties declared in ValidationScope the delegate functions specified in the variable declaration, verifying that the function is a valid delegate
@@ -512,11 +704,11 @@ protected:
 	 * @todo: this function will no longer be required once the post-parse fixup phase is added (TTPRO #13256)
 	 *
 	 * @param	AllClasses			the class tree for CurrentPackage
-	 * @param	ValidationScope		the scope to validate delegate properties for
-	 * @param	OwnerClass			the class currently being compiled.
+	 * @param	Struct				the struct to validate delegate properties for
+	 * @param	Scope				the current scope
 	 * @param	DelegateCache		cached map of delegates that have already been found; used for faster lookup.
 	 */
-	void FixupDelegateProperties( FClasses& AllClasses, UStruct* ValidationScope, UClass* OwnerClass, TMap<FName, UFunction*>& DelegateCache );
+	void FixupDelegateProperties(FClasses& AllClasses, UStruct* ValidationScope, FScope& Scope, TMap<FName, UFunction*>& DelegateCache);
 
 	/**
 	 * Verifies that all specified class's UProperties with CFG_RepNotify have valid callback targets with no parameters nor return values
@@ -582,5 +774,5 @@ public:
 	{
 	}
 
-	void ParseClassDeclaration(const TCHAR* InputText, int32 InLineNumber, const TCHAR* StartingMatchID, FString& out_ClassName, FString& out_BaseClassName, TArray<FName>& inout_ClassNames);
+	void ParseClassDeclaration(const TCHAR* InputText, int32 InLineNumber, const TCHAR* StartingMatchID, FString& out_ClassName, FString& out_BaseClassName, TArray<FHeaderProvider>& out_ClassNames);
 };
diff --git a/Engine/Source/Programs/UnrealHeaderTool/Private/HeaderProvider.cpp b/Engine/Source/Programs/UnrealHeaderTool/Private/HeaderProvider.cpp
new file mode 100644
--- /dev/null
+++ b/Engine/Source/Programs/UnrealHeaderTool/Private/HeaderProvider.cpp
@@ -0,0 +1,97 @@
+// Copyright 1998-2014 Epic Games, Inc. All Rights Reserved.
+
+#include "UnrealHeaderTool.h"
+#include "HeaderProvider.h"
+#include "UnrealSourceFile.h"
+#include "ClassMaps.h"
+
+FHeaderProvider::FHeaderProvider(EHeaderProviderSourceType::Type Type, const FString& Id)
+	: Type(Type), Id(Id), Cache(nullptr)
+{
+
+}
+
+template <class Predicate>
+bool TryFindSourceFileWithPredicate(FUnrealSourceFile*& Out, Predicate Pred)
+{
+	for (auto MapPair : GUnrealSourceFilesMap)
+	{
+		if (Pred(MapPair.Value.Get()))
+		{
+			Out = &(MapPair.Value.Get());
+			return true;
+		}
+	}
+
+	return false;
+}
+
+FUnrealSourceFile* FHeaderProvider::Resolve()
+{
+	if (Type != EHeaderProviderSourceType::Resolved)
+	{
+		if (Type == EHeaderProviderSourceType::ClassName)
+		{
+			for (const auto& Pair : GTypeDefinitionInfoMap)
+			{
+				if (Pair.Key->GetFName() == *Id)
+				{
+					Cache = &Pair.Value->GetUnrealSourceFile();
+					break;
+				}
+			}
+		}
+		else
+		{
+			auto* SourceFileSharedPtr = GUnrealSourceFilesMap.Find(Id);
+
+			if (SourceFileSharedPtr != nullptr)
+			{
+				Cache = &SourceFileSharedPtr->Get();
+			}
+			else
+			{
+				TryFindSourceFileWithPredicate(Cache,
+					[=](const FUnrealSourceFile& SourceFile)
+					{
+						return SourceFile.GetIncludePath() == Id;
+					})
+				|| TryFindSourceFileWithPredicate(Cache,
+					[=](const FUnrealSourceFile& SourceFile)
+					{
+						return SourceFile.GetFilename().EndsWith(TEXT("/") + Id) || SourceFile.GetFilename().EndsWith(TEXT("\\") + Id);
+					}
+				);
+			}
+		}
+
+		Type = EHeaderProviderSourceType::Resolved;
+	}
+
+	return Cache;
+}
+
+FString FHeaderProvider::ToString() const
+{
+	return FString::Printf(TEXT("%s %s"), Type == EHeaderProviderSourceType::ClassName ? TEXT("class") : TEXT("file"), *Id);
+}
+
+const FString& FHeaderProvider::GetId() const
+{
+	return Id;
+}
+
+bool operator==(const FHeaderProvider& A, const FHeaderProvider& B)
+{
+	if (A.Cache != B.Cache)
+	{
+		return false;
+	}
+
+	if (A.Cache == nullptr)
+	{
+		return A.Id == B.Id && A.Type == B.Type;
+	}
+
+	return true;
+}
diff --git a/Engine/Source/Programs/UnrealHeaderTool/Private/HeaderProvider.h b/Engine/Source/Programs/UnrealHeaderTool/Private/HeaderProvider.h
new file mode 100644
--- /dev/null
+++ b/Engine/Source/Programs/UnrealHeaderTool/Private/HeaderProvider.h
@@ -0,0 +1,37 @@
+// Copyright 1998-2014 Epic Games, Inc. All Rights Reserved.
+
+#pragma once
+
+#include "UnrealString.h"
+
+class FUnrealSourceFile;
+
+namespace EHeaderProviderSourceType
+{
+	enum Type
+	{
+		ClassName,
+		FileName,
+		Resolved
+	};
+}
+
+class FHeaderProvider
+{
+	friend bool operator==(const FHeaderProvider& A, const FHeaderProvider& B);
+public:
+	FHeaderProvider(EHeaderProviderSourceType::Type Type, const FString& Id);
+
+	FUnrealSourceFile* Resolve();
+
+	FString ToString() const;
+
+	const FString& GetId() const;
+
+private:
+	EHeaderProviderSourceType::Type Type;
+	FString Id;
+	FUnrealSourceFile* Cache;
+};
+
+bool operator==(const FHeaderProvider& A, const FHeaderProvider& B);
\ No newline at end of file
diff --git a/Engine/Source/Programs/UnrealHeaderTool/Private/NativeClassExporter.h b/Engine/Source/Programs/UnrealHeaderTool/Private/NativeClassExporter.h
--- a/Engine/Source/Programs/UnrealHeaderTool/Private/NativeClassExporter.h
+++ b/Engine/Source/Programs/UnrealHeaderTool/Private/NativeClassExporter.h
@@ -35,9 +35,16 @@ class FModuleClasses;
 struct FNativeClassHeaderGenerator
 {
 private:
-	FClass*				CurrentClass;
-	TArray<FClass*>		Classes;
 	FString				API;
+
+	/**
+	 * Gets API string for this header.
+	 */
+	FString GetAPIString()
+	{
+		return FString::Printf(TEXT("%s_API "), *API);
+	}
+
 	FString			ClassesHeaderPath;
 	/** the name of the cpp file where the implementation functions are generated, without the .cpp extension */
 	FString				GeneratedCPPFilenameBase;
@@ -48,10 +55,9 @@ private:
 	UPackage*			Package;
 	FStringOutputDevice	PreHeaderText;
 	FStringOutputDevice	EnumForeachText;
-	FStringOutputDevice ListOfPublicClassesUObjectHeaderGroupIncludes;  // This is built up each time from scratch each time for each header groups
+	TArray<FUnrealSourceFile*> ListOfPublicHeaderGroupIncludes;			// This is built up each time from scratch each time for each header groups
 	FStringOutputDevice ListOfPublicClassesUObjectHeaderModuleIncludes; // This is built up for the entire module, and for all processed headers
 	FStringOutputDevice ListOfAllUObjectHeaderIncludes;
-	FStringOutputDevice FriendText;
 	FStringOutputDevice	GeneratedPackageCPP;
 	FStringOutputDevice	GeneratedHeaderText;
 	FStringOutputDevice	GeneratedHeaderTextBeforeForwardDeclarations;
@@ -158,13 +164,6 @@ private:
 	 */
 	bool ShouldExportFunction( UFunction* Function );
 
-	/**
-	 * Returns a C++ code string for declaring a class or variable as an export, if the class is configured for that
-	 *
-	 * @return	The API declaration string
-	 */
-	FString MakeAPIDecl() const;
-
 	/**
 	 * Exports the struct's C++ properties to the HeaderText output device and adds special
 	 * compiler directives for GCC to pack as we expect.
@@ -175,55 +174,57 @@ private:
 	 */
 	void ExportProperties( UStruct* Struct, int32 TextIndent, bool bAccessSpecifiers = true, class FStringOutputDevice* Output = NULL );
 
-	/** Return the name of the singleton function that returns the UObject for Item **/
+	/** Return the name of the singleton function that returns the UObject for Item */
 	FString GetSingletonName(UField* Item, bool bRequiresValidObject=true);
+	/** Return the name of the singleton function that returns the UObject for Item */
 	FString GetSingletonName(FClass* Item, bool bRequiresValidObject=true);
 
 	/** 
 	 * Export functions used to find and call C++ or script implementation of a script function in the interface 
 	 */
-	void ExportInterfaceCallFunctions( const TArray<UFunction*>& CallbackFunctions, FStringOutputDevice& HeaderOutput );
+	void ExportInterfaceCallFunctions(const TArray<UFunction*>& CallbackFunctions, UClass* Class, FClassMetaData* ClassData, FStringOutputDevice& HeaderOutput);
 
-	/**
-	 * Exports the C++ class declarations for a native interface class.
+	/** 
+	 * Export UInterface boilerplate.
+	 *
+	 * @param UInterfaceBoilerplate Device to export to.
+	 * @param Class Interface to export.
+	 * @param FriendText Friend text for this boilerplate.
 	 */
-	void ExportInterfaceClassDeclaration( FClass* Class );
+	void ExportUInterfaceBoilerplate(FStringOutputDevice &UInterfaceBoilerplate, FClass* Class, const FString& FriendText);
 
 	/**
 	 * Appends the header definition for an inheritance hierarchy of classes to the header.
 	 * Wrapper for ExportClassHeaderRecursive().
 	 *
 	 * @param	AllClasses			The classes being processed.
-	 * @param	Class				The class to be exported.
+	 * @param	SourceFile			The source file to export header for.
 	 */
-	void ExportClassHeader( FClasses& AllClasses, FClass* Class );
+	void ExportSourceFileHeader(FClasses& AllClasses, FUnrealSourceFile* SourceFile);
 
 	/**
 	 * After all of the dependency checking, and setup for isolating the generated code, actually export the class
 	 *
-	 * @param	Class				The class to be exported.
-	 * @param	bValidNonTemporaryClass	Is the class permanent, or a temporary shell?
+	 * @param SourceFile Source file to export.
 	 */
-	void ExportClassHeaderInner( FClass* Class, bool bValidNonTemporaryClass );
+	void ExportClassesFromSourceFileInner(FUnrealSourceFile& SourceFile);
 
 	/**
 	 * After all of the dependency checking, but before actually exporting the class, set up the generated code
 	 *
-	 * @param	Class				The class to be exported.
-	 * @param	bIsExportClass		true if this is an export class
+	 * @param SourceFile Source file to export.
 	 */
-	void ExportClassHeaderWrapper( FClass* Class, bool bIsExportClass );
+	void ExportClassesFromSourceFileWrapper(FUnrealSourceFile& SourceFile);
 
 	/**
 	 * Appends the header definition for an inheritance hierarchy of classes to the header.
 	 *
 	 * @param	AllClasses				The classes being processed.
-	 * @param	Class					The class to be exported.
-	 * @param	DependencyChain			Used for finding errors. Must be empty before the first call.
-	 * @param	VisitedSet				The set of classes visited so far. Must be empty before the first call.
+	 * @param	SourceFile				The source file for which to export header.
+	 * @param	VisitedSet				The set of source files visited so far. Must be empty before the first call.
 	 * @param	bCheckDependenciesOnly	Whether we should just keep checking for dependency errors, without exporting anything.
 	 */
-	void ExportClassHeaderRecursive( FClasses& AllClasses, FClass* Class, TArray<FClass*>& DependencyChain, TSet<const UClass*>& VisitedSet, bool bCheckDependenciesOnly );
+	void ExportSourceFileHeaderRecursive(FClasses& AllClasses, FUnrealSourceFile* SourceFile, TSet<const FUnrealSourceFile*>& VisitedSet, bool bCheckDependenciesOnly);
 
 	/**
 	 * Returns a string in the format CLASS_Something|CLASS_Something which represents all class flags that are set for the specified
@@ -231,23 +232,6 @@ private:
 	 */
 	FString GetClassFlagExportText( UClass* Class );
 
-	/**
-	 * Iterates through all fields of the specified class, and separates fields that should be exported with this class into the appropriate array.
-	 * 
-	 * @param	Class				the class to pull fields from
-	 * @param	Enums				[out] all enums declared in the specified class
-	 * @param	Structs				[out] list of structs declared in the specified class
-	 * @param	CallbackFunctions	[out] list of delegates and events declared in the specified class
-	 * @param	NativeFunctions		[out] list of native functions declared in the specified class
-	 * @param	DelegateProperties	[out] list of delegate properties declared in the specified class
-	 */
-	void RetrieveRelevantFields(UClass* Class, TArray<UEnum*>& Enums, TArray<UScriptStruct*>& Structs, TArray<UFunction*>& CallbackFunctions, TArray<UFunction*>& NativeFunctions, TArray<UFunction*>& DelegateProperties);
-
-	/**
-	 * Exports the header text for the names needed
-	 */
-	void ExportNames();
-
 	/**
 	 * Exports the header text for the list of enums specified
 	 * 
@@ -265,9 +249,10 @@ private:
 	/**
 	 * Exports the macro declarations for GENERATED_USTRUCT_BODY() for each Foo in the list of structs specified
 	 * 
+	 * @param	SourceFile	the source file
 	 * @param	Structs		the structs to export
 	 */
-	void ExportGeneratedStructBodyMacros(const TArray<UScriptStruct*>& NativeStructs);
+	void ExportGeneratedStructBodyMacros(FUnrealSourceFile& SourceFile, const TArray<UScriptStruct*>& NativeStructs);
 
 	/**
 	 * Exports a local mirror of the specified struct; used to get offsets for noexport structs
@@ -289,38 +274,52 @@ private:
 	/**
 	 * Exports C++ type definitions for delegates
 	 *
+	 * @param	SourceFile			Source file of the delegate.
 	 * @param	DelegateFunctions	the functions that have parameters which need to be exported
 	 * @param	bWrapperImplementationsOnly	 True if only function implementations for delegate wrappers should be exported
 	 */
-	void ExportDelegateDefinitions( const TArray<UFunction*>& DelegateFunctions, const bool bWrapperImplementationsOnly );
+	void ExportDelegateDefinitions(FUnrealSourceFile& SourceFile, const TArray<UDelegateFunction*>& DelegateFunctions, const bool bWrapperImplementationsOnly);
 
 	/**
 	 * Exports the parameter struct declarations for the list of functions specified
 	 * 
+	 * @param	Scope				Current scope.
 	 * @param	CallbackFunctions	the functions that have parameters which need to be exported
+	 * @param	Output				ALternate output device
 	 * @param	Indent				number of spaces to put before each line
 	 * @param	bOutputConstructor	If true, output a constructor for the parm struct
-	 * @param	Output				ALternate output device
 	 */
-	void ExportEventParms( const TArray<UFunction*>& CallbackFunctions, const TCHAR* MacroSuffix, int32 Indent = 0, bool bOutputConstructor=true, class FStringOutputDevice* Output = NULL );
+	void ExportEventParms(FScope& Scope, const TArray<UFunction*>& CallbackFunctions, FStringOutputDevice& Output, int32 Indent = 0, bool bOutputConstructor = true);
+
+	/**
+	 * Exports the parameter struct declarations for the given function.
+	 *
+	 * @param	Function			the function that have parameters which need to be exported
+	 * @param	HeaderOutput		output device
+	 * @param	Indent				number of spaces to put before each line
+	 * @param	bOutputConstructor	If true, output a constructor for the param struct
+	 */
+	void ExportEventParm(UFunction* Function, FStringOutputDevice &HeaderOutput, int32 Indent = 0, bool bOutputConstructor = true);
 
 	/**
 	 * Generate a .proto message declaration for any functions marked as requiring one
 	 * 
 	 * @param InCallbackFunctions array of functions for consideration to generate .proto definitions
+	 * @param ClassData class data
 	 * @param Indent starting indentation level
 	 * @param Output optional output redirect
 	 */
-	void ExportProtoMessage(const TArray<UFunction*>& InCallbackFunctions, int32 Indent = 0, class FStringOutputDevice* Output = NULL);
+	void ExportProtoMessage(const TArray<UFunction*>& InCallbackFunctions, FClassMetaData* ClassData, int32 Indent = 0, class FStringOutputDevice* Output = NULL);
 
 	/**
 	 * Generate a .java message declaration for any functions marked as requiring one
 	 * 
 	 * @param InCallbackFunctions array of functions for consideration to generate .proto definitions
+	 * @param ClassData class data
 	 * @param Indent starting indentation level
 	 * @param Output optional output redirect
 	 */
-	void ExportMCPMessage(const TArray<UFunction*>& InCallbackFunctions, int32 Indent = 0, class FStringOutputDevice* Output = NULL);
+	void ExportMCPMessage(const TArray<UFunction*>& InCallbackFunctions, FClassMetaData* ClassData, int32 Indent = 0, class FStringOutputDevice* Output = NULL);
 
 	/** 
 	* Exports the temp header files into the .h files, then deletes the temp files.
@@ -381,19 +380,22 @@ private:
 	/**
 	 * Exports the native stubs for the list of functions specified
 	 * 
-	 * @param	NativeFunctions	the functions to export
+	 * @param SourceFile	current source file
+	 * @param Class			class
+	 * @param ClassData		class data
 	 */
-	void ExportNativeFunctions(const TArray<UFunction*>& NativeFunctions);
+	void ExportNativeFunctions(FUnrealSourceFile& SourceFile, UClass* Class, FClassMetaData* ClassData);
 
 	/**
 	 * Export the actual internals to a standard thunk function
 	 *
 	 * @param RPCWrappers output device for writing
+	 * @param Function given function
 	 * @param FunctionData function data for the current function
 	 * @param Parameters list of parameters in the function
 	 * @param Return return parameter for the function
 	 */
-	void ExportFunctionThunk(FStringOutputDevice& RPCWrappers, const FFuncInfo& FunctionData, const TArray<UProperty*>& Parameters, UProperty* Return);
+	void ExportFunctionThunk(FStringOutputDevice& RPCWrappers, UFunction* Function, const FFuncInfo& FunctionData, const TArray<UProperty*>& Parameters, UProperty* Return);
 
 	/** Exports the native function registration code for the given class. */
 	void ExportNatives(FClass* Class);
@@ -402,8 +404,18 @@ private:
 	 * Exports generated singleton functions for UObjects that used to be stored in .u files.
 	 * 
 	 * @param	Class			Class to export
-	**/
-	void ExportNativeGeneratedInitCode(FClass* Class);
+	 * @param	OutFriendText	(Output parameter) Friend text
+	 */
+	void ExportNativeGeneratedInitCode(FClass* Class, FStringOutputDevice& OutFriendText);
+
+	/**
+	 * Export given function.
+	 *
+	 * @param Function Given function.
+	 * @param Scope Scope for this function.
+	 * @param bIsNoExport Is in NoExport class.
+	 */
+	void ExportFunction(UFunction* Function, FScope* Scope, bool bIsNoExport);
 
 	/**
 	 * Exports a generated singleton function to setup the package for compiled-in classes.
@@ -448,9 +460,11 @@ private:
 	/**
 	 * Exports the proxy definitions for the list of enums specified
 	 * 
-	 * @param	CallbackFunctions	the functions to export
+	 * @param SourceFile	current source file
+	 * @param Class			current class
+	 * @param ClassData		current class data
 	 */
-	void ExportCallbackFunctions( const TArray<UFunction*>& CallbackFunctions);
+	TArray<UFunction*> ExportCallbackFunctions(FUnrealSourceFile& SourceFile, UClass* Class, FClassMetaData* ClassData);
 
 	/**
 	 * Determines if the property has alternate export text associated with it and if so replaces the text in PropertyText with the
@@ -489,14 +503,20 @@ private:
 	/**
 	 * Exports macros that manages UObject constructors.
 	 * 
+	 * @param ConstructorsMacroPrefix Prefix for constructors macro.
 	 * @param Class Class for which to export macros.
 	 */
-	void ExportConstructorsMacros(FClass* Class);
+	void ExportConstructorsMacros(const FString& ConstructorsMacroPrefix, FClass* Class);
+
+	/**
+	 * Gets list of public headers for the given package in the form of string of includes.
+	 */
+	FString GetListOfPublicHeaderGroupIncludesString(UPackage* Package);
 
 public:
 
 	// Constructor
-	FNativeClassHeaderGenerator( UPackage* InPackage, FClasses& AllClasses, bool InAllowSaveExportedHeaders );
+	FNativeClassHeaderGenerator(UPackage* InPackage, const TArray<FUnrealSourceFile*>& SourceFiles, FClasses& AllClasses, bool InAllowSaveExportedHeaders);
 
 	/**
 	 * Gets string with function return type.
diff --git a/Engine/Source/Programs/UnrealHeaderTool/Private/ParserClass.cpp b/Engine/Source/Programs/UnrealHeaderTool/Private/ParserClass.cpp
--- a/Engine/Source/Programs/UnrealHeaderTool/Private/ParserClass.cpp
+++ b/Engine/Source/Programs/UnrealHeaderTool/Private/ParserClass.cpp
@@ -62,18 +62,6 @@ FClass* FClass::GetClassWithin() const
 	return (FClass*)ClassWithin;
 }
 
-TArray<FName> FClass::GetDependentNames() const
-{
-	TArray<FName> Result;
-
-	for (auto Name : *GClassDependentOnMap.FindOrAdd(const_cast<FClass*>(this)))
-	{
-		Result.Add(Name);
-	}
-
-	return Result;
-}
-
 TArray<FClass*> FClass::GetInterfaceTypes() const
 {
 	TArray<FClass*> Result;
diff --git a/Engine/Source/Programs/UnrealHeaderTool/Private/ParserClass.h b/Engine/Source/Programs/UnrealHeaderTool/Private/ParserClass.h
--- a/Engine/Source/Programs/UnrealHeaderTool/Private/ParserClass.h
+++ b/Engine/Source/Programs/UnrealHeaderTool/Private/ParserClass.h
@@ -48,8 +48,6 @@ public:
 	 */
 	FClass* GetClassWithin() const;
 
-	TArray<FName> GetDependentNames() const;
-
 	TArray<FClass*> GetInterfaceTypes() const;
 
 	void GetHideCategories(TArray<FString>& OutHideCategories) const;
diff --git a/Engine/Source/Programs/UnrealHeaderTool/Private/ParserHelper.cpp b/Engine/Source/Programs/UnrealHeaderTool/Private/ParserHelper.cpp
--- a/Engine/Source/Programs/UnrealHeaderTool/Private/ParserHelper.cpp
+++ b/Engine/Source/Programs/UnrealHeaderTool/Private/ParserHelper.cpp
@@ -8,29 +8,6 @@
 /////////////////////////////////////////////////////
 // FClassMetaData
 
-FFunctionData* FClassMetaData::FindFunctionData( UFunction* Function )
-{
-	FFunctionData* Result = NULL;
-	
-	TScopedPointer<FFunctionData>* FuncData = FunctionData.Find(Function);
-	if ( FuncData != NULL )
-	{
-		Result = FuncData->GetOwnedPointer();
-	}
-
-	if ( Result == NULL )
-	{
-		UClass* OwnerClass = Function->GetOwnerClass();
-		FClassMetaData* OwnerClassData = GScriptHelper.FindClassData(OwnerClass);
-		if ( OwnerClassData && OwnerClassData != this )
-		{
-			Result = OwnerClassData->FindFunctionData(Function);
-		}
-	}
-
-	return Result;
-}
-
 /**
  * Finds the metadata for the struct specified
  * 
@@ -78,15 +55,21 @@ FTokenData* FClassMetaData::FindTokenData( UProperty* Prop )
 {
 	check(Prop);
 
-	FTokenData* Result = NULL;
+	FTokenData* Result = nullptr;
 	UObject* Outer = Prop->GetOuter();
-	UClass* OuterClass = Cast<UClass>(Outer);
-	if ( OuterClass != NULL )
+	UClass* OuterClass = nullptr;
+	if (Outer->IsA<UStruct>())
 	{
 		Result = GlobalPropertyData.Find(Prop);
-		if ( Result == NULL && OuterClass->GetSuperClass() != OuterClass )
+
+		if (Result == nullptr)
 		{
-			OuterClass = OuterClass->GetSuperClass();
+			OuterClass = Cast<UClass>(Outer);
+
+			if (Result == nullptr && OuterClass != nullptr && OuterClass->GetSuperClass() != OuterClass)
+			{
+				OuterClass = OuterClass->GetSuperClass();
+			}
 		}
 	}
 	else
@@ -95,12 +78,9 @@ FTokenData* FClassMetaData::FindTokenData( UProperty* Prop )
 		if ( OuterFunction != NULL )
 		{
 			// function parameter, return, or local property
-			TScopedPointer<FFunctionData>* pFuncData = FunctionData.Find(OuterFunction);
-			if ( pFuncData != NULL )
+			FFunctionData* FuncData = nullptr;
+			if (FFunctionData::TryFindForFunction(OuterFunction, FuncData))
 			{
-				FFunctionData* FuncData = pFuncData->GetOwnedPointer();
-				check(FuncData);
-
 				FPropertyData& FunctionParameters = FuncData->GetParameterData();
 				Result = FunctionParameters.Find(Prop);
 				if ( Result == NULL )
@@ -135,45 +115,16 @@ FTokenData* FClassMetaData::FindTokenData( UProperty* Prop )
 		}
 	}
 
-	if ( Result == NULL && OuterClass != NULL )
+	if (Result == nullptr && OuterClass != nullptr)
 	{
 		FClassMetaData* SuperClassData = GScriptHelper.FindClassData(OuterClass);
-		if ( SuperClassData && SuperClassData != this )
+		if (SuperClassData && SuperClassData != this)
 		{
 			Result = SuperClassData->FindTokenData(Prop);
 		}
 	}
-	return Result;
-}
-
-void FClassMetaData::Dump( int32 Indent )
-{
-	UE_LOG(LogCompile, Log, TEXT("Globals:"));
-	GlobalPropertyData.Dump(Indent+4);
 
-	UE_LOG(LogCompile, Log, TEXT("Structs:"));
-	for ( TMap<UScriptStruct*, TScopedPointer<FStructData> >::TIterator It(StructData); It; ++It )
-	{
-		UScriptStruct* Struct = It.Key();
-
-		TScopedPointer<FStructData>& PointerVal = It.Value();
-		FStructData* Data = PointerVal.GetOwnedPointer();
-
-		UE_LOG(LogCompile, Log, TEXT("%s%s"), FCString::Spc(Indent), *Struct->GetName());
-		Data->Dump(Indent+4);
-	}
-
-	UE_LOG(LogCompile, Log, TEXT("Functions:"));
-	for ( TMap<UFunction*, TScopedPointer<FFunctionData> >::TIterator It(FunctionData); It; ++It )
-	{
-		UFunction* Function = It.Key();
-
-		TScopedPointer<FFunctionData>& PointerVal = It.Value();
-		FFunctionData* Data = PointerVal.GetOwnedPointer();
-
-		UE_LOG(LogCompile, Log, TEXT("%s%s"), FCString::Spc(Indent), *Function->GetName());
-		Data->Dump(Indent+4);
-	}
+	return Result;
 }
 
 /////////////////////////////////////////////////////
@@ -295,4 +246,42 @@ bool FAdvancedDisplayParameterHandler::ShouldMarkParameter(const FString& Parame
 bool FAdvancedDisplayParameterHandler::CanMarkMore() const
 {
 	return bUseNumber ? (NumberLeaveUnmarked > 0) : (0 != ParametersNames.Num());
-}
\ No newline at end of file
+}
+
+TMap<UFunction*, TSharedRef<FFunctionData> > FFunctionData::FunctionDataMap;
+
+FFunctionData* FFunctionData::FindForFunction(UFunction* Function)
+{
+	auto* Output = FunctionDataMap.Find(Function);
+
+	check(Output);
+
+	return &(*Output).Get();
+}
+
+FFunctionData* FFunctionData::Add(UFunction* Function)
+{
+	auto Output = FunctionDataMap.Add(Function, MakeShareable(new FFunctionData()));
+
+	return &Output.Get();
+}
+
+FFunctionData* FFunctionData::Add(const FFuncInfo& FunctionInfo)
+{
+	auto Output = FunctionDataMap.Add(FunctionInfo.FunctionReference, MakeShareable(new FFunctionData(FunctionInfo)));
+
+	return &Output.Get();
+}
+
+bool FFunctionData::TryFindForFunction(UFunction* Function, FFunctionData*& OutData)
+{
+	auto* Output = FunctionDataMap.Find(Function);
+
+	if (!Output)
+	{
+		return false;
+	}
+
+	OutData = &(*Output).Get();
+	return true;
+}
diff --git a/Engine/Source/Programs/UnrealHeaderTool/Private/ParserHelper.h b/Engine/Source/Programs/UnrealHeaderTool/Private/ParserHelper.h
--- a/Engine/Source/Programs/UnrealHeaderTool/Private/ParserHelper.h
+++ b/Engine/Source/Programs/UnrealHeaderTool/Private/ParserHelper.h
@@ -60,26 +60,6 @@ struct ERefQualifier
 	};
 };
 
-/** Misc helper functions */
-struct FClassUtils
-{
-	FORCEINLINE static bool IsTemporaryClass(UClass* InClass)
-	{
-		return GTemporaryClasses.Contains(InClass);
-	}
-	FORCEINLINE static void MarkAsTemporaryClass(UClass* InClass)
-	{
-		if (!IsTemporaryClass(InClass))
-		{
-			GTemporaryClasses.Add(InClass);
-		}
-	}
-	FORCEINLINE static bool IsNoExportOrTemporaryClass(UClass* InClass)
-	{
-		return InClass->HasAnyClassFlags(CLASS_NoExport) || IsTemporaryClass(InClass);
-	}
-};
-
 #ifndef CASE_TEXT
 #define CASE_TEXT(txt) case txt: return TEXT(#txt)
 #endif
@@ -903,6 +883,8 @@ struct FFuncInfo
 	int16		RPCResponseId;
 	/** Whether this function represents a sealed event */
 	bool		bSealedEvent;
+	/** Delegate macro line in header. */
+	int32		DelegateMacroLine;
 	//@}
 
 	/** Constructor. */
@@ -917,6 +899,7 @@ struct FFuncInfo
 	,	RPCId               (0)
 	,	RPCResponseId       (0)
 	,	bSealedEvent        (false)
+	,	DelegateMacroLine(-1)
 	{}
 
 	FFuncInfo( const FFuncInfo& Other )
@@ -928,6 +911,7 @@ struct FFuncInfo
 	,	bCppValidationImplNameEndsWith_Validate(Other.bCppValidationImplNameEndsWith_Validate)
 	,	RPCId(Other.RPCId)
 	,	RPCResponseId(Other.RPCResponseId)
+	,	DelegateMacroLine(Other.DelegateMacroLine)
 	{
 		Function.Clone(Other.Function);
 		if (FunctionReference)
@@ -1284,6 +1268,29 @@ public:
 	{
 		FunctionData.FunctionExportFlags &= ~ClearFlags;
 	}
+
+	/**
+	 * Finds function data for given function object.
+	 */
+	static FFunctionData* FindForFunction(UFunction* Function);
+
+	/**
+	 * Adds function data object for given function object.
+	 */
+	static FFunctionData* Add(UFunction* Function);
+
+	/**
+	 * Adds function data object for given function object.
+	 */
+	static FFunctionData* Add(const FFuncInfo& FunctionInfo);
+
+	/**
+	 * Tries to find function data for given function object.
+	 */
+	static bool TryFindForFunction(UFunction* Function, FFunctionData*& OutData);
+
+private:
+	static TMap<UFunction*, TSharedRef<FFunctionData> > FunctionDataMap;
 };
 
 /**
@@ -1326,8 +1333,6 @@ class FClassMetaData
 
 	/** structs declared in this class */
 	TMap< UScriptStruct*, TScopedPointer<FStructData> >		StructData;
-	/** functions of this class */
-	TMap< UFunction*, TScopedPointer<FFunctionData> >		FunctionData;
 
 	/** base classes to multiply inherit from (other than the main base class */
 	TArray<FMultipleInheritanceBaseClass>					MultipleInheritanceParents;
@@ -1335,10 +1340,22 @@ class FClassMetaData
 	/** whether this class declares delegate functions or properties */
 	bool													bContainsDelegates;
 
+	/** The line of UCLASS/UINTERFACE macro in this class. */
+	int32 PrologLine;
+
+	/** The line of GENERATED_BODY/GENERATED_UCLASS_BODY macro in this class. */
+	int32 GeneratedBodyLine;
+
+	/** Same as above, but for interface class associated with this class. */
+	int32 InterfaceGeneratedBodyLine;
+
 public:
 	/** Default constructor */
 	FClassMetaData()
 		: bContainsDelegates(false)
+		, PrologLine(-1)
+		, GeneratedBodyLine(-1)
+		, InterfaceGeneratedBodyLine(-1)
 		, bConstructorDeclared(false)
 		, bDefaultConstructorDeclared(false)
 		, bObjectInitializerConstructorDeclared(false)
@@ -1347,34 +1364,65 @@ public:
 	}
 
 	/**
-	 * Adds a new function to be tracked
-	 * 
-	 * @param	FunctionInfo	the function to add
-	 *
-	 * @return	a pointer to the newly added FFunctionData
+	 * Gets prolog line number for this class.
 	 */
-	FFunctionData* AddFunction( const FFuncInfo& FunctionInfo )
+	int32 GetPrologLine() const
 	{
-		check(FunctionInfo.FunctionReference!=NULL);
+		check(PrologLine > 0);
+		return PrologLine;
+	}
 
-		FFunctionData* Result = NULL;
+	/**
+	 * Gets generated body line number for this class.
+	 */
+	int32 GetGeneratedBodyLine() const
+	{
+		check(GeneratedBodyLine > 0);
+		return GeneratedBodyLine;
+	}
 
-		TScopedPointer<FFunctionData>* pFuncData = FunctionData.Find(FunctionInfo.FunctionReference);
-		if ( pFuncData != NULL )
-		{
-			Result = pFuncData->GetOwnedPointer();
-			*Result = FFunctionData(FunctionInfo);
-		}
-		else
-		{
-			pFuncData = &FunctionData.Emplace(FunctionInfo.FunctionReference, new FFunctionData(FunctionInfo));
-			Result = pFuncData->GetOwnedPointer();
-		}
+	/**
+	 * Gets interface generated body line number for this class.
+	 */
+	int32 GetInterfaceGeneratedBodyLine() const
+	{
+		check(InterfaceGeneratedBodyLine > 0);
+		return InterfaceGeneratedBodyLine;
+	}
 
-		// update optimization flags
-		bContainsDelegates = bContainsDelegates || FunctionInfo.FunctionReference->HasAnyFunctionFlags(FUNC_Delegate);
+	/**
+	 * Sets prolog line number for this class.
+	 */
+	void SetPrologLine(int32 Line)
+	{
+		check(Line > 0);
+		PrologLine = Line;
+	}
 
-		return Result;
+	/**
+	 * Sets generated body line number for this class.
+	 */
+	void SetGeneratedBodyLine(int32 Line)
+	{
+		check(Line > 0);
+		GeneratedBodyLine = Line;
+	}
+
+	/**
+	 * Sets interface generated body line number for this class.
+	 */
+	void SetInterfaceGeneratedBodyLine(int32 Line)
+	{
+		check(Line > 0);
+		InterfaceGeneratedBodyLine = Line;
+	}
+
+	/**
+	 * Sets contains delegates flag for this class.
+	 */
+	void MarkContainsDelegate()
+	{
+		bContainsDelegates = true;
 	}
 
 	/**
@@ -1417,7 +1465,7 @@ public:
 		check(Prop);
 
 		UObject* Outer = Prop->GetOuter();
-		UClass* OuterClass = Cast<UClass>(Outer);
+		UStruct* OuterClass = Cast<UStruct>(Outer);
 		if ( OuterClass != NULL )
 		{
 			// global property
@@ -1425,14 +1473,12 @@ public:
 		}
 		else
 		{
+			checkNoEntry();
 			UFunction* OuterFunction = Cast<UFunction>(Outer);
 			if ( OuterFunction != NULL )
 			{
 				// function parameter, return, or local property
-				TScopedPointer<FFunctionData>* FuncData = FunctionData.Find(OuterFunction);
-				check(FuncData != NULL);
-
-				(*FuncData)->AddProperty(PropertyToken);
+				FFunctionData::FindForFunction(OuterFunction)->AddProperty(PropertyToken);
 			}
 			else
 			{
@@ -1476,7 +1522,7 @@ public:
 	 * @param	Field		the property or function to add to
 	 * @param	InMetaData	the metadata to add
 	 */
-	void AddMetaData(UField* Field, const TMap<FName, FString>& InMetaData)
+	static void AddMetaData(UField* Field, const TMap<FName, FString>& InMetaData)
 	{
 		// only add if we have some!
 		if (InMetaData.Num())
@@ -1578,7 +1624,6 @@ public:
 	{
 		GlobalPropertyData.Shrink();
 		StructData.Shrink();
-		FunctionData.Shrink();
 		MultipleInheritanceParents.Shrink();
 	}
 
@@ -1600,22 +1645,22 @@ public:
  * The type of data tracked by this class is data that would otherwise only be accessible by adding a 
  * member property to UFunction/UProperty.  
  */
-class FCompilerMetadataManager : protected TMap<UClass*, TScopedPointer<FClassMetaData> >
+class FCompilerMetadataManager : protected TMap<UStruct*, TScopedPointer<FClassMetaData> >
 {
 public:
 	/**
 	 * Adds a new class to be tracked
 	 * 
-	 * @param	Cls	the UClass to add
+	 * @param	Struct	the UStruct to add
 	 *
 	 * @return	a pointer to the newly added metadata for the class specified
 	 */
-	FClassMetaData* AddClassData( UClass* Cls )
+	FClassMetaData* AddClassData(UStruct* Struct)
 	{
-		TScopedPointer<FClassMetaData>* pClassData = Find(Cls);
-		if ( pClassData == NULL )
+		TScopedPointer<FClassMetaData>* pClassData = Find(Struct);
+		if (pClassData == NULL)
 		{
-			pClassData = &Emplace(Cls, new FClassMetaData());
+			pClassData = &Emplace(Struct, new FClassMetaData());
 		}
 
 		return *pClassData;
@@ -1624,16 +1669,16 @@ public:
 	/**
 	 * Find the metadata associated with the class specified
 	 * 
-	 * @param	Cls	the UClass to add
+	 * @param	Struct	the UStruct to add
 	 *
 	 * @return	a pointer to the newly added metadata for the class specified
 	 */
-	FClassMetaData* FindClassData( UClass* Cls )
+	FClassMetaData* FindClassData(UStruct* Struct)
 	{
 		FClassMetaData* Result = NULL;
 
-		TScopedPointer<FClassMetaData>* pClassData = Find(Cls);
-		if ( pClassData )
+		TScopedPointer<FClassMetaData>* pClassData = Find(Struct);
+		if (pClassData)
 		{
 			Result = pClassData->GetOwnedPointer();
 		}
@@ -1641,29 +1686,13 @@ public:
 		return Result;
 	}
 
-	/**
-	 * (debug) Dumps the values of this FFunctionData to the log file
-	 * 
-	 * @param	Indent	number of spaces to insert at the beginning of each line
-	 */	
-	void Dump()
-	{
-		for ( TMap<UClass*,TScopedPointer<FClassMetaData> >::TIterator It(*this); It; ++It )
-		{
-			UClass* Cls = It.Key();
-			TScopedPointer<FClassMetaData>& Data = It.Value();
-			UE_LOG(LogCompile, Log, TEXT("=== %s ==="), *Cls->GetName());
-			Data->Dump(4);
-		}
-	}
-
 	/**
 	 * Shrink TMaps to avoid slack in Pairs array.
 	 */
 	void Shrink()
 	{
-		TMap<UClass*,TScopedPointer<FClassMetaData> >::Shrink();
-		for( TMap<UClass*,TScopedPointer<FClassMetaData> >::TIterator It(*this); It; ++It )
+		TMap<UStruct*, TScopedPointer<FClassMetaData> >::Shrink();
+		for (TMap<UStruct*, TScopedPointer<FClassMetaData> >::TIterator It(*this); It; ++It)
 		{
 			FClassMetaData* MetaData = It.Value();
 			MetaData->Shrink();
@@ -1741,14 +1770,6 @@ struct FNameLookupCPP
 			return NameCPP;
 
 		FString DesiredStructName = Struct->GetName();
-		if (UClass* TestClass = Cast<UClass>(Struct))
-		{
-			if (FClassUtils::IsTemporaryClass(TestClass))
-			{
-				DesiredStructName = GClassHeaderNameWithNoPathMap[TestClass];
-			}
-		}
-
 		FString	TempName = FString(bForceInterface ? TEXT("I") : Struct->GetPrefixCPP()) + DesiredStructName;
 		int32 StringLength = TempName.Len();
 
diff --git a/Engine/Source/Programs/UnrealHeaderTool/Private/Scope.cpp b/Engine/Source/Programs/UnrealHeaderTool/Private/Scope.cpp
new file mode 100644
--- /dev/null
+++ b/Engine/Source/Programs/UnrealHeaderTool/Private/Scope.cpp
@@ -0,0 +1,185 @@
+// Copyright 1998-2015 Epic Games, Inc. All Rights Reserved.
+
+#include "UnrealHeaderTool.h"
+#include "ParserHelper.h"
+
+extern FCompilerMetadataManager GScriptHelper;
+
+FScope::FScope(FScope* Parent)
+	: Parent(Parent)
+{
+	check(Parent);
+}
+
+FScope::FScope()
+	: Parent(nullptr)
+{
+
+}
+
+void FScope::AddType(UField* Type)
+{
+	TypeMap.Add(Type->GetFName(), Type);
+}
+
+/**
+ * Dispatch type to one of three arrays Enums, Structs and DelegateFunctions.
+ *
+ * @param Type Input type.
+ * @param Enums (Output parameter) Array to fill with enums.
+ * @param Structs (Output parameter) Array to fill with structs.
+ * @param DelegateFunctions (Output parameter) Array to fill with delegate functions.
+ */
+void DispatchType(UField* Type, TArray<UEnum*> &Enums, TArray<UScriptStruct*> &Structs, TArray<UDelegateFunction*> &DelegateFunctions)
+{
+	UClass* TypeClass = Type->GetClass();
+
+	if (TypeClass == UClass::StaticClass() || TypeClass == UStruct::StaticClass())
+	{
+		// Inner scopes.
+		FScope::GetTypeScope((UStruct*)Type)->SplitTypesIntoArrays(Enums, Structs, DelegateFunctions);
+	}
+	else if (TypeClass == UEnum::StaticClass())
+	{
+		UEnum* Enum = (UEnum*)Type;
+		Enums.Add(Enum);
+	}
+	else if (TypeClass == UScriptStruct::StaticClass())
+	{
+		UScriptStruct* Struct = (UScriptStruct*)Type;
+		Structs.Add(Struct);
+	}
+	else if (TypeClass == UDelegateFunction::StaticClass())
+	{
+		bool bAdded = false;
+		UDelegateFunction* Function = (UDelegateFunction*)Type;
+
+		if (Function->GetSuperFunction() == NULL)
+		{
+			DelegateFunctions.Add(Function);
+			bAdded = true;
+		}
+
+		check(bAdded);
+	}
+}
+
+void FScope::SplitTypesIntoArrays(TArray<UEnum*>& Enums, TArray<UScriptStruct*>& Structs, TArray<UDelegateFunction*>& DelegateFunctions)
+{
+	for (auto& TypePair : TypeMap)
+	{
+		auto* Type = TypePair.Value;
+		DispatchType(Type, Enums, Structs, DelegateFunctions);
+	}
+}
+
+TSharedRef<FScope> FScope::GetTypeScope(UStruct* Type)
+{
+	auto* ScopeRefPtr = ScopeMap.Find(Type);
+	if (!ScopeRefPtr)
+	{
+		FError::Throwf(TEXT("Couldn't find scope for the type %s."), *Type->GetName());
+	}
+
+	return *ScopeRefPtr;
+}
+
+TSharedRef<FScope> FScope::AddTypeScope(UStruct* Type, FScope* ParentScope)
+{
+	TSharedRef<FScope> Scope = MakeShareable(new FStructScope(Type, ParentScope));
+
+	ScopeMap.Add(Type, Scope);
+
+	return Scope;
+}
+
+UField* FScope::FindTypeByName(FName Name)
+{
+	auto TypeIterator = TDeepScopeTypeIterator<UField, false>(this);
+
+	while (TypeIterator.MoveNext())
+	{
+		auto* Type = *TypeIterator;
+		if (Type->GetFName() == Name)
+		{
+			return Type;
+		}
+	}
+
+	return nullptr;
+}
+
+const UField* FScope::FindTypeByName(FName Name) const
+{
+	auto TypeIterator = GetTypeIterator();
+
+	while (TypeIterator.MoveNext())
+	{
+		auto* Type = *TypeIterator;
+		if (Type->GetFName() == Name)
+		{
+			return Type;
+		}
+	}
+
+	return nullptr;
+}
+
+bool FScope::ContainsType(UField* Type)
+{
+	return FindTypeByName(Type->GetFName()) != nullptr;
+}
+
+bool FScope::IsFileScope() const
+{
+	return Parent == nullptr;
+}
+
+bool FScope::ContainsTypes() const
+{
+	return TypeMap.Num() > 0;
+}
+
+TMap<UStruct*, TSharedRef<FScope> > FScope::ScopeMap;
+
+FFileScope::FFileScope(FName Name, FUnrealSourceFile* SourceFile)
+	: SourceFile(SourceFile), Name(Name)
+{
+
+}
+
+void FFileScope::IncludeScope(FFileScope* IncludedScope)
+{
+	IncludedScopes.Add(IncludedScope);
+}
+
+FUnrealSourceFile* FFileScope::GetSourceFile() const
+{
+	return SourceFile;
+}
+
+FName FFileScope::GetName() const
+{
+	return Name;
+}
+
+UStruct* FStructScope::GetStruct() const
+{
+	return Struct;
+}
+
+FName FStructScope::GetName() const
+{
+	return Struct->GetFName();
+}
+
+FStructScope::FStructScope(UStruct* Struct, FScope* Parent)
+	: FScope(Parent), Struct(Struct)
+{
+
+}
+
+FClassMetaData* FStructScope::GetClassMetaData() const
+{
+	return GScriptHelper.FindClassData(Struct);
+}
diff --git a/Engine/Source/Programs/UnrealHeaderTool/Private/Scope.h b/Engine/Source/Programs/UnrealHeaderTool/Private/Scope.h
new file mode 100644
--- /dev/null
+++ b/Engine/Source/Programs/UnrealHeaderTool/Private/Scope.h
@@ -0,0 +1,461 @@
+// Copyright 1998-2015 Epic Games, Inc. All Rights Reserved.
+
+#pragma once
+
+#include "CoreUObject.h"
+
+// Forward declarations.
+class UStruct;
+class FClassMetaData;
+class FUnrealSourceFile;
+
+// Traits to achieve conditional types for const/non-const iterators.
+template <bool Condition, class TypeIfTrue, class TypeIfFalse>
+class TConditionalType
+{
+public:
+	typedef TypeIfFalse Type;
+};
+
+template <class TypeIfTrue, class TypeIfFalse>
+class TConditionalType<true, TypeIfTrue, TypeIfFalse>
+{
+public:
+	typedef TypeIfTrue Type;
+};
+
+// Base class representing type scope.
+class FScope
+{
+public:
+	// Default constructor i.e. Parent == nullptr
+	FScope();
+
+	// Constructor that's providing parent scope.
+	FScope(FScope* Parent);
+
+	// Virtual destructor.
+	virtual ~FScope()
+	{ };
+
+	/**
+	 * Adds type to the scope.
+	 *
+	 * @param Type Type to add.
+	 */
+	void AddType(UField* Type);
+
+	/**
+	 * Finds type by its name.
+	 *
+	 * @param Name Name to look for.
+	 *
+	 * @returns Found type or nullptr on failure.
+	 */
+	UField* FindTypeByName(FName Name);
+
+	/**
+	 * Finds type by its name.
+	 *
+	 * Const version.
+	 *
+	 * @param Name Name to look for.
+	 *
+	 * @returns Found type or nullptr on failure.
+	 */
+	const UField* FindTypeByName(FName Name) const;
+
+	/**
+	 * Checks if scope contains this type.
+	 *
+	 * @param Type Type to look for.
+	 */
+	bool ContainsType(UField* Type);
+
+	/**
+	 * Checks if scope contains type that satisfies a predicate.
+	 *
+	 * @param Predicate Predicate to satisfy.
+	 */
+	template <typename TPredicateType>
+	bool Contains(TPredicateType Predicate)
+	{
+		for (const auto& NameTypePair : TypeMap)
+		{
+			if (Predicate.operator()(NameTypePair.Value))
+			{
+				return true;
+			}
+		}
+
+		return false;
+	}
+
+	/**
+	 * Gets struct's or class's scope.
+	 *
+	 * @param Type A struct or class for which to return the scope.
+	 *
+	 * @returns Struct's or class's scope.
+	 */
+	static TSharedRef<FScope> GetTypeScope(UStruct* Type);
+
+	/**
+	 * Adds struct's or class's scope.
+	 *
+	 * @param Type A struct or class for which to add the scope.
+	 * @param ParentScope A scope in which this scope is contained.
+	 *
+	 * @returns Newly added scope.
+	 */
+	static TSharedRef<FScope> AddTypeScope(UStruct* Type, FScope* ParentScope);
+
+	/**
+	 * Gets structs, enums and delegate functions from this scope.
+	 *
+	 * @param Enums (Output parameter) List of enums from this scope.
+	 * @param Structs (Output parameter) List of structs from this scope.
+	 * @param DelegateFunctions (Output parameter) List of delegate properties from this scope.
+	 */
+	void SplitTypesIntoArrays(TArray<UEnum*>& Enums, TArray<UScriptStruct*>& Structs, TArray<UDelegateFunction*>& DelegateFunctions);
+
+	/**
+	 * Gets scope name.
+	 */
+	virtual FName GetName() const = 0;
+
+	/**
+	 * Gets scope types filtered by type.
+	 *
+	 * @param OutArray (Output parameter) Array to fill with types from this scope.
+	 */
+	template <class TTypeFilter>
+	void GetTypes(TArray<TTypeFilter*> OutArray)
+	{
+		for (const auto& NameTypePair : TypeMap)
+		{
+			if (NameTypePair.Value->IsA<TTypeFilter>())
+			{
+				OutArray.Add((TTypeFilter*)NameTypePair.Value);
+			}
+		}
+	}
+
+	/**
+	 * Gets scope's parent.
+	 */
+	const FScope* GetParent() const
+	{
+		return Parent;
+	}
+
+	/**
+	 * Scope type iterator.
+	 */
+	template <class TType, bool TIsConst>
+	class TScopeTypeIterator
+	{
+	public:
+		// Conditional typedefs.
+		typedef typename TConditionalType<TIsConst, TMap<FName, UField*>::TConstIterator, TMap<FName, UField*>::TIterator>::Type MapIteratorType;
+		typedef typename TConditionalType<TIsConst, const FScope, FScope>::Type ScopeType;
+
+		// Constructor
+		TScopeTypeIterator(ScopeType* Scope)
+			: TypeIterator(Scope->TypeMap)
+		{
+			bBeforeStart = true;
+		}
+
+		/**
+		 * Indirection operator.
+		 *
+		 * @returns Type this iterator currently point to.
+		 */
+		TType* operator*()
+		{
+			return bBeforeStart ? nullptr : (TType*)(*TypeIterator).Value;
+		}
+
+		/**
+		 * Moves this iterator to next type.
+		 *
+		 * @returns True on success. False otherwise.
+		 */
+		bool MoveNext()
+		{
+			if (!bBeforeStart)
+			{
+				++TypeIterator;
+			}
+
+			bBeforeStart = false;
+			return TypeIterator.operator bool();
+		}
+
+	private:
+
+		// Current type.
+		MapIteratorType TypeIterator;
+
+		// Tells if this iterator is in initial state.
+		bool bBeforeStart;
+	};
+
+	/**
+	 * Gets type iterator.
+	 *
+	 * @return Type iterator.
+	 */
+	template <class TType = UField>
+	TScopeTypeIterator<TType, false> GetTypeIterator()
+	{
+		return TScopeTypeIterator<TType, false>(this);
+	}
+
+	/**
+	 * Gets const type iterator.
+	 *
+	 * @return Const type iterator.
+	 */
+	template <class TType = UField>
+	TScopeTypeIterator<TType, true> GetTypeIterator() const
+	{
+		return TScopeTypeIterator<TType, true>(this);
+	}
+
+	/**
+	 * Tells if this scope is a file scope.
+	 *
+	 * @returns True if this is a file scope. False otherwise.
+	 */
+	bool IsFileScope() const;
+
+	/**
+	 * Tells if this scope contains any type.
+	 *
+	 * @returns True if this scope contains type. False otherwise.
+	 */
+	bool ContainsTypes() const;
+
+private:
+	// This scopes parent.
+	const FScope* Parent;
+
+	// Map of types in this scope.
+	TMap<FName, UField*> TypeMap;
+
+	// Global map type <-> scope.
+	static TMap<UStruct*, TSharedRef<FScope> > ScopeMap;
+};
+
+/**
+ * Represents a scope associated with source file.
+ */
+class FFileScope : public FScope
+{
+public:
+	// Constructor.
+	FFileScope(FName Name, FUnrealSourceFile* SourceFile);
+
+	/**
+	 * Includes other file scopes into this scope (#include).
+	 *
+	 * @param IncludedScope File scope to include.
+	 */
+	void IncludeScope(FFileScope* IncludedScope);
+
+	/**
+	 * Gets scope name.
+	 */
+	FName GetName() const override;
+
+	/**
+	 * Gets source file associated with this scope.
+	 */
+	FUnrealSourceFile* GetSourceFile() const;
+
+	/**
+	 * Appends included file scopes into given array.
+	 *
+	 * @param Out (Output parameter) Array to append scopes.
+	 */
+	void AppendIncludedFileScopes(TArray<FScope*>& Out)
+	{
+		if (!Out.Contains(this))
+		{
+			Out.Add(this);
+		}
+
+		for (auto* IncludedScope : IncludedScopes)
+		{
+			if (!Out.Contains(IncludedScope))
+			{
+				Out.Add(IncludedScope);
+				IncludedScope->AppendIncludedFileScopes(Out);
+			}
+		}
+	}
+
+private:
+	// Source file.
+	FUnrealSourceFile* SourceFile;
+
+	// Scope name.
+	FName Name;
+
+	// Included scopes list.
+	TArray<FFileScope*> IncludedScopes;
+};
+
+
+/**
+ * Data structure representing scope of a struct or class.
+ */
+class FStructScope : public FScope
+{
+public:
+	// Constructor.
+	FStructScope(UStruct* Struct, FScope* Parent);
+
+	/**
+	 * Gets struct associated with this scope.
+	 */
+	UStruct* GetStruct() const;
+
+	/**
+	 * Gets metadata associated with class of this scope.
+	 */
+	FClassMetaData* GetClassMetaData() const;
+
+	/**
+	 * Gets scope name.
+	 */
+	FName GetName() const override;
+
+private:
+	// Struct associated with this scope.
+	UStruct* Struct;
+};
+
+/**
+ * Deep scope type iterator. It looks for type in the whole scope hierarchy.
+ * First going up inheritance from inner structs to outer. Then through all
+ * included scopes.
+ */
+template <class TType, bool TIsConst>
+class TDeepScopeTypeIterator
+{
+public:
+	// Conditional typedefs.
+	typedef typename TConditionalType<TIsConst, TMap<FName, UField*>::TConstIterator, TMap<FName, UField*>::TIterator>::Type MapIteratorType;
+	typedef typename TConditionalType<TIsConst, const FScope, FScope>::Type ScopeType;
+	typedef typename TConditionalType<TIsConst, const FFileScope, FFileScope>::Type FileScopeType;
+	typedef typename TConditionalType<TIsConst, typename TArray<ScopeType*>::TConstIterator, typename TArray<ScopeType*>::TIterator>::Type ScopeArrayIteratorType;
+
+	// Constructor.
+	TDeepScopeTypeIterator(ScopeType* Scope)
+	{
+		const auto* CurrentScope = Scope;
+
+		while (!CurrentScope->IsFileScope())
+		{
+			ScopesToTraverse.Add(Scope);
+
+			auto* Struct = ((FStructScope*)CurrentScope)->GetStruct();
+
+			if (Struct->IsA<UClass>())
+			{
+				auto* Class = ((UClass*)Struct)->GetSuperClass();
+
+				while (Class && !(Class->ClassFlags & EClassFlags::CLASS_Intrinsic))
+				{
+					ScopesToTraverse.Add(&FScope::GetTypeScope(Class).Get());
+					Class = Class->GetSuperClass();
+				}
+			}
+
+			CurrentScope = CurrentScope->GetParent();
+		}
+
+		((FileScopeType*)CurrentScope)->AppendIncludedFileScopes(ScopesToTraverse);
+	}
+
+	/**
+	 * Iterator increment.
+	 *
+	 * @returns True if moved iterator to another position. False otherwise.
+	 */
+	bool MoveNext()
+	{
+		if (!ScopeIterator.IsValid() && !MoveToNextScope())
+		{
+			return false;
+		}
+
+		if (ScopeIterator->MoveNext())
+		{
+			return true;
+		}
+		else
+		{
+			do
+			{
+				ScopeIterator = nullptr;
+				if (!MoveToNextScope())
+				{
+					return false;
+				}
+			}
+			while(!ScopeIterator->MoveNext());
+
+			return true;
+		}
+	}
+
+	/**
+	 * Current type getter.
+	 *
+	 * @returns Current type.
+	 */
+	TType* operator*()
+	{
+		return ScopeIterator->operator*();
+	}
+
+private:
+	/**
+	 * Moves the iterator to the next scope.
+	 *
+	 * @returns True if succeeded. False otherwise.
+	 */
+	bool MoveToNextScope()
+	{
+		if (!ScopesIterator.IsValid())
+		{
+			ScopesIterator = MakeShareable(new ScopeArrayIteratorType(ScopesToTraverse));
+		}
+		else
+		{
+			ScopesIterator->operator++();
+		}
+
+		if (!ScopesIterator->operator bool())
+		{
+			return false;
+		}
+
+		ScopeIterator = MakeShareable(new FScope::TScopeTypeIterator<TType, TIsConst>(ScopesIterator->operator*()));
+		return true;
+	}
+
+	// List of scopes to traverse.
+	TArray<ScopeType*> ScopesToTraverse;
+
+	// Current scope iterator.
+	TSharedPtr<FScope::TScopeTypeIterator<TType, TIsConst> > ScopeIterator;
+
+	// Scopes list iterator.
+	TSharedPtr<ScopeArrayIteratorType> ScopesIterator;
+};
+
diff --git a/Engine/Source/Programs/UnrealHeaderTool/Private/UnrealSourceFile.cpp b/Engine/Source/Programs/UnrealHeaderTool/Private/UnrealSourceFile.cpp
new file mode 100644
--- /dev/null
+++ b/Engine/Source/Programs/UnrealHeaderTool/Private/UnrealSourceFile.cpp
@@ -0,0 +1,114 @@
+// Copyright 1998-2015 Epic Games, Inc. All Rights Reserved.
+
+#include "UnrealHeaderTool.h"
+#include "UnrealSourceFile.h"
+#include "ParserHelper.h"
+
+void FUnrealSourceFile::AddDefinedClass(UClass* Class)
+{
+	DefinedClasses.Add(Class);
+}
+
+FString FUnrealSourceFile::GetFileId() const
+{
+	FString StdFilename = Filename;
+
+	FPaths::MakeStandardFilename(StdFilename);
+
+	FStringOutputDevice Out;
+
+	// Standard filename will always start with "..\..\..\", so we are removing it.
+	StdFilename = StdFilename.Mid(9);
+	for (auto Char : StdFilename)
+	{
+		if (Char == 0)
+		{
+			break;
+		}
+		else if (FChar::IsAlnum(Char))
+		{
+			Out.AppendChar(Char);
+		}
+		else
+		{
+			Out.AppendChar('_');
+		}
+	}
+
+	return Out;
+}
+
+FString FUnrealSourceFile::GetStrippedFilename() const
+{
+	return FPaths::GetBaseFilename(Filename);
+}
+
+FString FUnrealSourceFile::GetGeneratedMacroName(FClassMetaData* ClassData, const TCHAR* Suffix) const
+{
+	return GetGeneratedMacroName(ClassData->GetGeneratedBodyLine(), Suffix);
+}
+
+FString FUnrealSourceFile::GetGeneratedMacroName(int32 LineNumber, const TCHAR* Suffix) const
+{
+	if (Suffix != nullptr)
+	{
+		return FString::Printf(TEXT("%s_%d%s"), *GetFileId(), LineNumber, Suffix);
+	}
+
+	return FString::Printf(TEXT("%s_%d"), *GetFileId(), LineNumber);
+}
+
+FString FUnrealSourceFile::GetGeneratedBodyMacroName(int32 LineNumber, bool bLegacy) const
+{
+	return GetGeneratedMacroName(LineNumber, *FString::Printf(TEXT("%s%s"), TEXT("_GENERATED_BODY"), bLegacy ? TEXT("_LEGACY") : TEXT("")));
+}
+
+void FUnrealSourceFile::SetGeneratedFilename(FString GeneratedFilename)
+{
+	this->GeneratedFilename = MoveTemp(GeneratedFilename);
+}
+
+void FUnrealSourceFile::SetHasChanged(bool bHasChanged)
+{
+	this->bHasChanged = bHasChanged;
+}
+
+void FUnrealSourceFile::SetModuleRelativePath(FString ModuleRelativePath)
+{
+	this->ModuleRelativePath = MoveTemp(ModuleRelativePath);
+}
+
+void FUnrealSourceFile::SetIncludePath(FString IncludePath)
+{
+	this->IncludePath = MoveTemp(IncludePath);
+}
+
+const FString& FUnrealSourceFile::GetContent() const
+{
+	return Content;
+}
+
+void FUnrealSourceFile::MarkDependenciesResolved()
+{
+	bDependenciesResolved = true;
+}
+
+bool FUnrealSourceFile::AreDependenciesResolved() const
+{
+	return bDependenciesResolved;
+}
+
+void FUnrealSourceFile::MarkAsParsed()
+{
+	bParsed = true;
+}
+
+bool FUnrealSourceFile::IsParsed() const
+{
+	return bParsed;
+}
+
+bool FUnrealSourceFile::HasChanged() const
+{
+	return bHasChanged;
+}
diff --git a/Engine/Source/Programs/UnrealHeaderTool/Private/UnrealSourceFile.h b/Engine/Source/Programs/UnrealHeaderTool/Private/UnrealSourceFile.h
new file mode 100644
--- /dev/null
+++ b/Engine/Source/Programs/UnrealHeaderTool/Private/UnrealSourceFile.h
@@ -0,0 +1,250 @@
+// Copyright 1998-2015 Epic Games, Inc. All Rights Reserved.
+
+#pragma once
+
+#include "Scope.h"
+#include "HeaderProvider.h"
+#include "UnrealTypeDefinitionInfo.h"
+
+class FClassMetaData;
+
+/**
+ * Contains information about source file that defines various UHT aware types.
+ */
+class FUnrealSourceFile
+{
+public:
+	// Constructor.
+	FUnrealSourceFile(UPackage* Package, FString Filename, FString Content)
+		: Scope(MakeShareable(new FFileScope(*(FString(TEXT("__")) + FPaths::GetBaseFilename(Filename) + FString(TEXT("__File"))), this))),
+		Filename(MoveTemp(Filename)), Package(Package), bHasChanged(false), Content(MoveTemp(Content)), bParsed(false), bDependenciesResolved(false)
+	{
+		if (GetStrippedFilename() != "Object")
+		{
+			Includes.Add(FHeaderProvider(EHeaderProviderSourceType::FileName, "Object.h"));
+		}
+	}
+
+	/**
+	 * Appends array with classes defined in this source file.
+	 *
+	 * @param OutClasses (Output parameter) Array to fill with classes defined.
+	 */
+	void AppendDefinedClasses(TArray<UClass*>& OutClasses) const
+	{
+		for (auto* Class : DefinedClasses)
+		{
+			OutClasses.Add(Class);
+		}
+	}
+
+	/**
+	 * Gets array with classes defined in this source file.
+	 *
+	 * @returns Array with classes defined in this source file.
+	 */
+	const TArray<UClass*>& GetDefinedClasses() const
+	{
+		return DefinedClasses;
+	}
+
+	/**
+	 * Gets number of types defined in this source file.
+	 */
+	int32 GetDefinedClassesCount() const
+	{
+		return GetDefinedClasses().Num();
+	}
+
+	/**
+	 * Gets generated header filename.
+	 */
+	FString GetGeneratedHeaderFilename() const
+	{
+		return FString::Printf(TEXT("%s.generated.h"), *FPaths::GetBaseFilename(Filename));
+	}
+
+	/**
+	 * Gets module relative path.
+	 */
+	const FString& GetModuleRelativePath() const
+	{
+		return ModuleRelativePath;
+	}
+
+	/**
+	 * Gets stripped filename.
+	 */
+	FString GetStrippedFilename() const;
+
+	/**
+	 * Gets unique file id.
+	 */
+	FString GetFileId() const;
+
+	/**
+	 * Gets file-wise generated body macro name.
+	 *
+	 * @param LineNumber Number at which generated body macro is.
+	 * @param bLegacy Tells if method should get legacy generated body macro.
+	 */
+	FString GetGeneratedBodyMacroName(int32 LineNumber, bool bLegacy = false) const;
+
+	/**
+	 * Gets file-wise generated body macro name.
+	 *
+	 * @param LineNumber Number at which generated body macro is.
+	 * @param Suffix Suffix to add to generated body macro name.
+	 */
+	FString GetGeneratedMacroName(int32 LineNumber, const TCHAR* Suffix) const;
+
+	/**
+	 * Gets file-wise generated body macro name.
+	 *
+	 * @param ClassData Class metadata for which to get generated body macro name.
+	 * @param Suffix Suffix to add to generated body macro name.
+	 */
+	FString GetGeneratedMacroName(FClassMetaData* ClassData, const TCHAR* Suffix = nullptr) const;
+
+	/**
+	 * Adds given class to class definition list for this source file.
+	 *
+	 * @param Class Class to add to list.
+	 */
+	void AddDefinedClass(UClass* Class);
+
+	/**
+	 * Gets scope for this file.
+	 */
+	TSharedRef<FFileScope> GetScope() const
+	{
+		return Scope;
+	}
+
+	/**
+	 * Gets package this file is in.
+	 */
+	UPackage* GetPackage() const
+	{
+		return Package;
+	}
+
+	/**
+	 * Gets filename.
+	 */
+	const FString& GetFilename() const
+	{
+		return Filename;
+	}
+
+	/**
+	 * Gets generated filename.
+	 */
+	const FString& GetGeneratedFilename() const
+	{
+		return GeneratedFilename;
+	}
+
+	/**
+	 * Gets include path.
+	 */
+	const FString& GetIncludePath() const
+	{
+		return IncludePath;
+	}
+
+	/**
+	 * Gets content.
+	 */
+	const FString& GetContent() const;
+
+	/**
+	 * Gets includes.
+	 */
+	TArray<FHeaderProvider>& GetIncludes()
+	{
+		return Includes;
+	}
+
+	/**
+	 * Sets generated filename.
+	 */
+	void SetGeneratedFilename(FString GeneratedFilename);
+
+	/**
+	 * Sets has changed flag.
+	 */
+	void SetHasChanged(bool bHasChanged);
+
+	/**
+	 * Sets module relative path.
+	 */
+	void SetModuleRelativePath(FString ModuleRelativePath);
+
+	/**
+	 * Sets include path.
+	 */
+	void SetIncludePath(FString IncludePath);
+
+	/**
+	 * Mark this file as parsed.
+	 */
+	void MarkAsParsed();
+
+	/**
+	 * Checks if this file is parsed.
+	 */
+	bool IsParsed() const;
+
+	/**
+	 * Checks if generated file has been changed.
+	 */
+	bool HasChanged() const;
+
+	/**
+	 * Mark that this file has resolved dependencies.
+	 */
+	void MarkDependenciesResolved();
+
+	/**
+	 * Checks if dependencies has been resolved.
+	 */
+	bool AreDependenciesResolved() const;
+
+private:
+	// File scope.
+	TSharedRef<FFileScope> Scope;
+
+	// Path of this file.
+	FString Filename;
+
+	// Package of this file.
+	UPackage* Package;
+
+	// File name of the generated header file associated with this file.
+	FString GeneratedFilename;
+
+	// Tells if generated header file was changed.
+	bool bHasChanged;
+
+	// Module relative path.
+	FString ModuleRelativePath;
+
+	// Include path.
+	FString IncludePath;
+
+	// Source file content.
+	FString Content;
+
+	// Tells if this file was parsed.
+	bool bParsed;
+
+	// Tells if dependencies has been resolved already.
+	bool bDependenciesResolved;
+
+	// This source file includes.
+	TArray<FHeaderProvider> Includes;
+
+	// List of classes defined in this source file.
+	TArray<UClass*> DefinedClasses;
+};
\ No newline at end of file
diff --git a/Engine/Source/Programs/UnrealHeaderTool/Private/UnrealTypeDefinitionInfo.h b/Engine/Source/Programs/UnrealHeaderTool/Private/UnrealTypeDefinitionInfo.h
new file mode 100644
--- /dev/null
+++ b/Engine/Source/Programs/UnrealHeaderTool/Private/UnrealTypeDefinitionInfo.h
@@ -0,0 +1,39 @@
+// Copyright 1998-2015 Epic Games, Inc. All Rights Reserved.
+
+#pragma once
+
+// Forward declarations.
+class FUnrealSourceFile;
+
+/**
+ * Class that stores information about type (USTRUCT/UCLASS) definition.
+ */
+class FUnrealTypeDefinitionInfo
+{
+public:
+	// Constructor
+	FUnrealTypeDefinitionInfo(FUnrealSourceFile& SourceFile, int32 LineNumber)
+		: SourceFile(SourceFile), LineNumber(LineNumber)
+	{}
+
+	/**
+	 * Gets the line number in source file this type was defined in.
+	 */
+	int32 GetLineNumber() const
+	{
+		return LineNumber;
+	}
+
+	/**
+	 * Gets the reference to FUnrealSourceFile object that stores information about
+	 * source file this type was defined in.
+	 */
+	FUnrealSourceFile& GetUnrealSourceFile()
+	{
+		return SourceFile;
+	}
+
+private:
+	FUnrealSourceFile& SourceFile;
+	int32 LineNumber;
+};
\ No newline at end of file
diff --git a/Engine/Source/Runtime/Core/Private/Misc/Paths.cpp b/Engine/Source/Runtime/Core/Private/Misc/Paths.cpp
--- a/Engine/Source/Runtime/Core/Private/Misc/Paths.cpp
+++ b/Engine/Source/Runtime/Core/Private/Misc/Paths.cpp
@@ -965,3 +965,18 @@ void FPaths::CombineInternal(FString& OutPath, const TCHAR** Pathes, int32 NumPa
 	}
 }
 
+bool FPaths::IsSamePath(const FString& PathA, const FString& PathB)
+{
+	FString TmpA = PathA;
+	FString TmpB = PathB;
+
+	MakeStandardFilename(TmpA);
+	MakeStandardFilename(TmpB);
+
+#if defined(PLATFORM_WINDOWS) || defined(PLATFORM_XBOXONE)
+	return FCString::Stricmp(*TmpA, *TmpB) == 0;
+#else
+	return FCString::Strcmp(*TmpA, *TmpB) == 0;
+#endif
+}
+
