UHT: Multiple UCLASS in header. Minor tweaks and comments.

diff --git a/Engine/Source/Programs/UnrealHeaderTool/Private/CodeGenerator.cpp b/Engine/Source/Programs/UnrealHeaderTool/Private/CodeGenerator.cpp
--- a/Engine/Source/Programs/UnrealHeaderTool/Private/CodeGenerator.cpp
+++ b/Engine/Source/Programs/UnrealHeaderTool/Private/CodeGenerator.cpp
@@ -4968,6 +4968,12 @@ UClass* ProcessParsedClass(bool bClassIsAnInterface, TArray<FHeaderProvider> &De
 {
 	FString ClassNameStripped = GetClassNameWithPrefixRemoved(*ClassName);
 
+	// All classes must start with a valid unreal prefix
+	if (!FHeaderParser::ClassNameHasValidPrefix(ClassName, ClassNameStripped))
+	{
+		FError::Throwf(TEXT("Invalid class name '%s'. The class name must have an appropriate prefix added (A for Actors, U for other classes)."), *ClassName);
+	}
+
 	// Ensure the base class has any valid prefix and exists as a valid class. Checking for the 'correct' prefix will occur during compilation
 	FString BaseClassNameStripped;
 	if (!BaseClassName.IsEmpty())
diff --git a/Engine/Source/Programs/UnrealHeaderTool/Private/HeaderParser.cpp b/Engine/Source/Programs/UnrealHeaderTool/Private/HeaderParser.cpp
--- a/Engine/Source/Programs/UnrealHeaderTool/Private/HeaderParser.cpp
+++ b/Engine/Source/Programs/UnrealHeaderTool/Private/HeaderParser.cpp
@@ -6259,10 +6259,22 @@ ECompilationResult::Type FHeaderParser::ParseHeader(FClasses& AllClasses, FUnrea
 		}
 
 		check(SourceFile.IsParsed());
+
+		if (!bSpottedAutogeneratedHeaderInclude && !bEmptyFile && !bNoExportClassesOnly)
+		{
+			const FString ExpectedHeaderName = SourceFile.GetGeneratedHeaderFilename();
+			FError::Throwf(TEXT("Expected an include at the top of the header: '#include \"%s\"'"), *ExpectedHeaderName);
+		}
 	}
 #if !PLATFORM_EXCEPTIONS_DISABLED
 	catch( TCHAR* ErrorMsg )
 	{
+		if (NestLevel == 0)
+		{
+			// Pushing nest so there is a file context for this error.
+			PushNest(NEST_GlobalScope, nullptr, &SourceFile);
+		}
+
 		// Handle compiler error.
 		{
 			TGuardValue<ELogTimes::Type> DisableLogTimes(GPrintLogTimes, ELogTimes::None);
@@ -6277,12 +6289,6 @@ ECompilationResult::Type FHeaderParser::ParseHeader(FClasses& AllClasses, FUnrea
 	}
 #endif
 
-	if (!bSpottedAutogeneratedHeaderInclude && !bEmptyFile && !bNoExportClassesOnly)
-	{
-		const FString ExpectedHeaderName = SourceFile.GetGeneratedHeaderFilename();
-		FError::Throwf(TEXT("Expected an include at the top of the header: '#include \"%s\"'"), *ExpectedHeaderName);
-	}
-
 	return Result; //@TODO: UCREMOVAL: This function is always returning succeeded even on a compiler error; should this continue?
 }
 
diff --git a/Engine/Source/Runtime/CoreUObject/Private/UObject/UObjectGlobals.cpp b/Engine/Source/Runtime/CoreUObject/Private/UObject/UObjectGlobals.cpp
--- a/Engine/Source/Runtime/CoreUObject/Private/UObject/UObjectGlobals.cpp
+++ b/Engine/Source/Runtime/CoreUObject/Private/UObject/UObjectGlobals.cpp
@@ -169,7 +169,11 @@ UObject* StaticFindObject( UClass* ObjectClass, UObject* InObjectPackage, const
 	FName ObjectName(*InName, FNAME_Add, true);
 	MatchingObject = StaticFindObjectFast( ObjectClass, ObjectPackage, ObjectName, ExactClass, bAnyPackage );
 
+	// This is another look-up for native enums, structs or delegate signatures, cause they're path changed
+	// and old packages can have invalid ones. The path now does not have a UCLASS as an outer. All mentioned
+	// types are just children of package of the file there were defined in.
 	if (!MatchingObject && ObjectPackage != nullptr &&
+		!FPlatformProperties::RequiresCookedData() && // Cooked platforms will have all paths resolved.
 			(
 				ObjectClass == UEnum::StaticClass()	// Enums
 				|| ObjectClass == UScriptStruct::StaticClass() // Structs
