UHT: Fixes fully qualified path into relative ones. Fully paths was breaking cross-compilation in some cases.

diff --git a/Engine/Source/Programs/UnrealHeaderTool/Private/CodeGenerator.cpp b/Engine/Source/Programs/UnrealHeaderTool/Private/CodeGenerator.cpp
--- a/Engine/Source/Programs/UnrealHeaderTool/Private/CodeGenerator.cpp
+++ b/Engine/Source/Programs/UnrealHeaderTool/Private/CodeGenerator.cpp
@@ -1937,6 +1937,22 @@ void ExportEnhancedConstructorsMacro(FStringOutputDevice& Out, FClass* Class, co
 	ExportDefaultConstructorCallDefinition(Out, Class);
 }
 
+/**
+ * Gets a package relative inclusion path of the given source file for build.
+ *
+ * @param SourceFile Given source file.
+ *
+ * @returns Inclusion path.
+ */
+FString GetBuildPath(FUnrealSourceFile& SourceFile)
+{
+	FString Out = SourceFile.GetFilename();
+
+	ConvertToBuildIncludePath(SourceFile.GetPackage(), Out);
+
+	return Out;
+}
+
 FString FNativeClassHeaderGenerator::GetListOfPublicHeaderGroupIncludesString(UPackage* Package)
 {
 	FStringOutputDevice Out;
@@ -1952,7 +1968,7 @@ FString FNativeClassHeaderGenerator::GetListOfPublicHeaderGroupIncludesString(UP
 
 	for (auto* SourceFile : ListOfPublicHeaderGroupIncludes)
 	{
-		Out.Logf(TEXT("#include \"%s\"") LINE_TERMINATOR, *SourceFile->GetFilename());
+		Out.Logf(TEXT("#include \"%s\"") LINE_TERMINATOR, *GetBuildPath(*SourceFile));
 	}
 
 	Out.Log(LINE_TERMINATOR);
