Redundant parameter removed from CompileVariableDeclaration.

diff --git a/Engine/Source/Programs/UnrealHeaderTool/Private/HeaderParser.cpp b/Engine/Source/Programs/UnrealHeaderTool/Private/HeaderParser.cpp
--- a/Engine/Source/Programs/UnrealHeaderTool/Private/HeaderParser.cpp
+++ b/Engine/Source/Programs/UnrealHeaderTool/Private/HeaderParser.cpp
@@ -1408,7 +1408,7 @@ UScriptStruct* FHeaderParser::CompileStructDeclaration(FClasses& AllClasses, FUn
 		}
 		else if (Token.Matches(TEXT("UPROPERTY"), ESearchCase::CaseSensitive))
 		{
-			CompileVariableDeclaration(AllClasses, Struct, EPropertyDeclarationStyle::UPROPERTY);
+			CompileVariableDeclaration(AllClasses, Struct);
 		}
 		else if (Token.Matches(TEXT("UFUNCTION"), ESearchCase::CaseSensitive))
 		{
@@ -3722,7 +3722,7 @@ bool FHeaderParser::CompileDeclaration(FClasses& AllClasses, FUnrealSourceFile&
 		CheckAllow(TEXT("'Member variable declaration'"), ALLOW_VarDecl);
 		check(TopNest->NestType == NEST_Class);
 
-		CompileVariableDeclaration(AllClasses, GetCurrentClass(), EPropertyDeclarationStyle::UPROPERTY);
+		CompileVariableDeclaration(AllClasses, GetCurrentClass());
 	}
 	else if (Token.Matches(TEXT("UENUM")))
 	{
@@ -5829,7 +5829,7 @@ struct FExposeOnSpawnValidator
 	}
 };
 
-void FHeaderParser::CompileVariableDeclaration(FClasses& AllClasses, UStruct* Struct, EPropertyDeclarationStyle::Type PropertyDeclarationStyle)
+void FHeaderParser::CompileVariableDeclaration(FClasses& AllClasses, UStruct* Struct)
 {
 	uint64 DisallowFlags = CPF_ParmFlags;
 	uint64 EdFlags       = 0;
@@ -5839,39 +5839,36 @@ void FHeaderParser::CompileVariableDeclaration(FClasses& AllClasses, UStruct* St
 	EObjectFlags ObjectFlags = RF_NoFlags;
 
 	FIndexRange TypeRange;
-	GetVarType( AllClasses, &FScope::GetTypeScope(Struct).Get(), OriginalProperty, ObjectFlags, DisallowFlags, TEXT("Member variable declaration"), /*OuterPropertyType=*/ NULL, PropertyDeclarationStyle, EVariableCategory::Member, &TypeRange );
+	GetVarType( AllClasses, &FScope::GetTypeScope(Struct).Get(), OriginalProperty, ObjectFlags, DisallowFlags, TEXT("Member variable declaration"), /*OuterPropertyType=*/ NULL, EPropertyDeclarationStyle::UPROPERTY, EVariableCategory::Member, &TypeRange );
 	OriginalProperty.PropertyFlags |= EdFlags;
 
 	FString* Category = OriginalProperty.MetaData.Find("Category");
 
-	if (PropertyDeclarationStyle == EPropertyDeclarationStyle::UPROPERTY)
+	// First check if the category was specified at all and if the property was exposed to the editor.
+	if (!Category && (OriginalProperty.PropertyFlags & (CPF_Edit|CPF_BlueprintVisible)))
 	{
-		// First check if the category was specified at all and if the property was exposed to the editor.
-		if (!Category && (OriginalProperty.PropertyFlags & (CPF_Edit|CPF_BlueprintVisible)))
+		static const FString AbsoluteEngineDir = FPaths::ConvertRelativePathToFull(FPaths::EngineDir());
+		FString SourceFilename = GetCurrentSourceFile()->GetFilename();
+		FPaths::NormalizeFilename(SourceFilename);
+		if (Struct->GetOutermost() != nullptr && !SourceFilename.StartsWith(AbsoluteEngineDir))
 		{
-			static const FString AbsoluteEngineDir = FPaths::ConvertRelativePathToFull(FPaths::EngineDir());
-			FString SourceFilename = GetCurrentSourceFile()->GetFilename();
-			FPaths::NormalizeFilename(SourceFilename);
-			if (Struct->GetOutermost() != nullptr && !SourceFilename.StartsWith(AbsoluteEngineDir))
-			{
-				OriginalProperty.MetaData.Add("Category", Struct->GetFName().ToString());
-				Category = OriginalProperty.MetaData.Find("Category");
-			}
-			else
-			{
-				FError::Throwf(TEXT("Property is exposed to the editor or blueprints but has no Category specified."));
-			}
+			OriginalProperty.MetaData.Add("Category", Struct->GetFName().ToString());
+			Category = OriginalProperty.MetaData.Find("Category");
 		}
-
-		// Validate that pointer properties are not interfaces (which are not GC'd and so will cause runtime errors)
-		if (OriginalProperty.PointerType == EPointerType::Native && OriginalProperty.Struct->IsChildOf(UInterface::StaticClass()))
+		else
 		{
-			// Get the name of the type, removing the asterisk representing the pointer
-			FString TypeName = FString(TypeRange.Count, Input + TypeRange.StartIndex).Trim().TrimTrailing().LeftChop(1).TrimTrailing();
-			FError::Throwf(TEXT("UPROPERTY pointers cannot be interfaces - did you mean TScriptInterface<%s>?"), *TypeName);
+			FError::Throwf(TEXT("Property is exposed to the editor or blueprints but has no Category specified."));
 		}
 	}
 
+	// Validate that pointer properties are not interfaces (which are not GC'd and so will cause runtime errors)
+	if (OriginalProperty.PointerType == EPointerType::Native && OriginalProperty.Struct->IsChildOf(UInterface::StaticClass()))
+	{
+		// Get the name of the type, removing the asterisk representing the pointer
+		FString TypeName = FString(TypeRange.Count, Input + TypeRange.StartIndex).Trim().TrimTrailing().LeftChop(1).TrimTrailing();
+		FError::Throwf(TEXT("UPROPERTY pointers cannot be interfaces - did you mean TScriptInterface<%s>?"), *TypeName);
+	}
+
 	// If the category was specified explicitly, it wins
 	if (Category && !(OriginalProperty.PropertyFlags & (CPF_Edit|CPF_BlueprintVisible|CPF_BlueprintAssignable|CPF_BlueprintCallable)))
 	{
diff --git a/Engine/Source/Programs/UnrealHeaderTool/Private/HeaderParser.h b/Engine/Source/Programs/UnrealHeaderTool/Private/HeaderParser.h
--- a/Engine/Source/Programs/UnrealHeaderTool/Private/HeaderParser.h
+++ b/Engine/Source/Programs/UnrealHeaderTool/Private/HeaderParser.h
@@ -505,7 +505,7 @@ protected:
 	void CompileClassDeclaration(FClasses& AllClasses);
 	void CompileDelegateDeclaration(FUnrealSourceFile& SourceFile, FClasses& AllClasses, const TCHAR* DelegateIdentifier, EDelegateSpecifierAction::Type SpecifierAction = EDelegateSpecifierAction::DontParse);
 	void CompileFunctionDeclaration(FUnrealSourceFile& SourceFile, FClasses& AllClasses);
-	void CompileVariableDeclaration (FClasses& AllClasses, UStruct* Struct, EPropertyDeclarationStyle::Type PropertyDeclarationStyle);
+	void CompileVariableDeclaration (FClasses& AllClasses, UStruct* Struct);
 	void CompileInterfaceDeclaration(FClasses& AllClasses);
 
 	FClass* ParseInterfaceNameDeclaration(FClasses& AllClasses, FString& DeclaredInterfaceName, FString& RequiredAPIMacroIfPresent);
