Deprecating 'operator new' for UObjects.

diff --git a/Engine/Plugins/Editor/SpeedTreeImporter/Source/SpeedTreeImporter/Private/SpeedTreeImportFactory.cpp b/Engine/Plugins/Editor/SpeedTreeImporter/Source/SpeedTreeImporter/Private/SpeedTreeImportFactory.cpp
--- a/Engine/Plugins/Editor/SpeedTreeImporter/Source/SpeedTreeImporter/Private/SpeedTreeImportFactory.cpp
+++ b/Engine/Plugins/Editor/SpeedTreeImporter/Source/SpeedTreeImporter/Private/SpeedTreeImportFactory.cpp
@@ -473,7 +473,7 @@ UTexture* CreateSpeedTreeMaterialTexture(UObject* Parent, FString Filename, bool
 	}
 	else
 	{
-		UTextureFactory* TextureFact = new UTextureFactory(FObjectInitializer());
+		auto TextureFact = NewObject<UTextureFactory>();
 		TextureFact->AddToRoot();
 		TextureFact->SuppressImportOverwriteDialog();
 
@@ -651,7 +651,7 @@ UMaterialInterface* CreateSpeedTreeMaterial(UObject* Parent, FString MaterialFul
 	}
 	
 	// create an unreal material asset
-	UMaterialFactoryNew* MaterialFactory = new UMaterialFactoryNew(FObjectInitializer());
+	auto MaterialFactory = NewObject<UMaterialFactoryNew>();
 	UMaterial* UnrealMaterial = (UMaterial*)MaterialFactory->FactoryCreateNew(UMaterial::StaticClass(), Package, *FixedMaterialName, RF_Standalone|RF_Public, NULL, GWarn);
 	FAssetRegistryModule::AssetCreated(UnrealMaterial);
 	Package->SetDirtyFlag(true);
@@ -1113,7 +1113,7 @@ UObject* USpeedTreeImportFactory::FactoryCreateBinary(UClass* InClass, UObject*
 					ExistingMesh->PreEditChange(NULL);
 				}
 				
-				StaticMesh = new(Package, FName(*MeshName), Flags | RF_Public) UStaticMesh(FObjectInitializer());
+				StaticMesh = NewNamedObject<UStaticMesh>(Package, FName(*MeshName), Flags | RF_Public);
 
 				// @todo AssetImportData make a data class for speed tree assets
 				StaticMesh->AssetImportData = ConstructObject<UAssetImportData>(UAssetImportData::StaticClass(), StaticMesh);
diff --git a/Engine/Plugins/Experimental/LiveEditor/Source/LiveEditor/Private/LiveEditorManager.cpp b/Engine/Plugins/Experimental/LiveEditor/Source/LiveEditor/Private/LiveEditorManager.cpp
--- a/Engine/Plugins/Experimental/LiveEditor/Source/LiveEditor/Private/LiveEditorManager.cpp
+++ b/Engine/Plugins/Experimental/LiveEditor/Source/LiveEditor/Private/LiveEditorManager.cpp
@@ -607,7 +607,8 @@ void FLiveEditorManager::InjectNewBlueprintEditor( const TWeakPtr<FBlueprintEdit
 
 void FLiveEditorManager::CreateLiveEditorWorld()
 {
-	LiveEditorWorld = new UWorld(FObjectInitializer(),FURL(NULL));
+	LiveEditorWorld = NewObject<UWorld>();
+	LiveEditorWorld->SetFlags(RF_Transactional);
 	LiveEditorWorld->WorldType = EWorldType::Preview;
 	
 	bool bTransactional = false;
diff --git a/Engine/Plugins/Runtime/SoundMod/Source/SoundModImporter/Private/SoundModImporterFactory.cpp b/Engine/Plugins/Runtime/SoundMod/Source/SoundModImporter/Private/SoundModImporterFactory.cpp
--- a/Engine/Plugins/Runtime/SoundMod/Source/SoundModImporter/Private/SoundModImporterFactory.cpp
+++ b/Engine/Plugins/Runtime/SoundMod/Source/SoundModImporter/Private/SoundModImporterFactory.cpp
@@ -103,7 +103,7 @@ UObject* USoundModImporterFactory::FactoryCreateBinary
 
 	// Use pre-existing sound if it exists and we want to keep settings,
 	// otherwise create new sound and import raw data.
-	USoundMod* Sound = (bUseExistingSettings && ExistingSound) ? ExistingSound : new(InParent, Name, Flags) USoundMod(FObjectInitializer());
+	USoundMod* Sound = (bUseExistingSettings && ExistingSound) ? ExistingSound : NewNamedObject<USoundMod>(InParent, Name, Flags);
 
 	Sound->Duration = xmpModuleInfo.seq_data->duration / 1000.f;
 
diff --git a/Engine/Source/Developer/MeshUtilities/Private/MeshUtilities.cpp b/Engine/Source/Developer/MeshUtilities/Private/MeshUtilities.cpp
--- a/Engine/Source/Developer/MeshUtilities/Private/MeshUtilities.cpp
+++ b/Engine/Source/Developer/MeshUtilities/Private/MeshUtilities.cpp
@@ -3716,7 +3716,7 @@ void FMeshUtilities::CreateProxyMesh(
 		MeshPackage->Modify();
 	}
 
-	UStaticMesh* StaticMesh = new(MeshPackage, FName(*MeshAssetName), RF_Public|RF_Standalone) UStaticMesh(FObjectInitializer());
+	auto StaticMesh = NewNamedObject<UStaticMesh>(MeshPackage, FName(*MeshAssetName), RF_Public | RF_Standalone);
 	StaticMesh->InitResources();
 	{
 		FString OutputPath = StaticMesh->GetPathName();
@@ -4102,7 +4102,7 @@ void FMeshUtilities::MergeActors(
 		Package->FullyLoad();
 		Package->Modify();
 
-		UStaticMesh* StaticMesh = new(Package, *AssetName, RF_Public|RF_Standalone) UStaticMesh(FObjectInitializer());
+		auto StaticMesh = NewNamedObject<UStaticMesh>(Package, *AssetName, RF_Public | RF_Standalone);
 		StaticMesh->InitResources();
 		
 		FString OutputPath = StaticMesh->GetPathName();
diff --git a/Engine/Source/Editor/FontEditor/Private/SCompositeFontEditor.cpp b/Engine/Source/Editor/FontEditor/Private/SCompositeFontEditor.cpp
--- a/Engine/Source/Editor/FontEditor/Private/SCompositeFontEditor.cpp
+++ b/Engine/Source/Editor/FontEditor/Private/SCompositeFontEditor.cpp
@@ -744,7 +744,8 @@ void STypefaceEntryEditor::OnTypefaceEntryFontPathPicked(const FString& InNewFon
 		CompositeFontEditorPtr->GetFontObject()->Modify();
 
 		// We need to allocate the bulk data with the font as its outer
-		const UFontBulkData* const BulkData = new(CompositeFontEditorPtr->GetFontObject()) UFontBulkData(InNewFontFilename);
+		UFontBulkData* const BulkData = NewObject<UFontBulkData>(CompositeFontEditorPtr->GetFontObject());
+		BulkData->Initialize(InNewFontFilename);
 
 		TypefaceEntryPtr->Font.SetFont(InNewFontFilename, BulkData);
 		CompositeFontEditorPtr->FlushCachedFont();
diff --git a/Engine/Source/Editor/GeometryMode/Private/GeometryModifiers.cpp b/Engine/Source/Editor/GeometryMode/Private/GeometryModifiers.cpp
--- a/Engine/Source/Editor/GeometryMode/Private/GeometryModifiers.cpp
+++ b/Engine/Source/Editor/GeometryMode/Private/GeometryModifiers.cpp
@@ -1841,7 +1841,8 @@ static ABrush* ClipBrushAgainstPlane( const FPlane& InPlane, ABrush* InBrush)
 	// perhaps there were additional brushes were selected. 
 	check( ClippedBrush->GetClass() == InBrush->GetClass() );
 
-	ClippedBrush->Brush = new( InBrush->GetOuter(), NAME_None )UModel( FObjectInitializer(),NULL );
+	ClippedBrush->Brush = NewObject<UModel>(InBrush->GetOuter());
+	ClippedBrush->Brush->Initialize(nullptr);
 	ClippedBrush->GetBrushComponent()->Brush = ClippedBrush->Brush;
 
 	GeometryClipping::BuildGiantAlignedBrush( *ClippedBrush, InPlane );
diff --git a/Engine/Source/Editor/LevelEditor/Private/LevelEditorActions.cpp b/Engine/Source/Editor/LevelEditor/Private/LevelEditorActions.cpp
--- a/Engine/Source/Editor/LevelEditor/Private/LevelEditorActions.cpp
+++ b/Engine/Source/Editor/LevelEditor/Private/LevelEditorActions.cpp
@@ -2562,7 +2562,8 @@ void FLevelEditorActionCallbacks::OnSaveBrushAsCollision()
 
 		// Copy the current builder brush into a temp model.
 		// We keep no reference to this, so it will be GC'd at some point.
-		UModel* TempModel = new UModel(FObjectInitializer(), NULL, 1);
+		UModel* TempModel = NewObject<UModel>();
+		TempModel->Initialize(nullptr, 1);
 		TempModel->Polys->Element.AssignButKeepOwner(BuilderModel->Polys->Element);
 
 		// Now transform each poly into local space for the selected static mesh.
diff --git a/Engine/Source/Editor/Matinee/Classes/MatineeTransBuffer.h b/Engine/Source/Editor/Matinee/Classes/MatineeTransBuffer.h
--- a/Engine/Source/Editor/Matinee/Classes/MatineeTransBuffer.h
+++ b/Engine/Source/Editor/Matinee/Classes/MatineeTransBuffer.h
@@ -12,10 +12,10 @@
 UCLASS(transient)
 class UMatineeTransBuffer : public UTransBuffer
 {
-    GENERATED_UCLASS_BODY()
-
-	UMatineeTransBuffer(const FObjectInitializer& ObjectInitializer, SIZE_T InMaxMemory)
-		:	UTransBuffer( ObjectInitializer, InMaxMemory )
+public:
+    GENERATED_BODY()
+	UMatineeTransBuffer(const FObjectInitializer& ObjectInitializer)
+		:	UTransBuffer(ObjectInitializer)
 	{ }
 
 	/**  
diff --git a/Engine/Source/Editor/Matinee/Private/Matinee.cpp b/Engine/Source/Editor/Matinee/Private/Matinee.cpp
--- a/Engine/Source/Editor/Matinee/Private/Matinee.cpp
+++ b/Engine/Source/Editor/Matinee/Private/Matinee.cpp
@@ -497,7 +497,8 @@ void FMatinee::InitMatinee(const EToolkitMode::Type Mode, const TSharedPtr< clas
 	GEditor->ResetTransaction( NSLOCTEXT("UnrealEd", "OpenMatinee", "Open UnrealMatinee") );
 
 	NormalTransactor = GEditor->Trans;
-	InterpEdTrans = new UMatineeTransBuffer( FObjectInitializer(), 8*1024*1024 );
+	InterpEdTrans = NewObject<UMatineeTransBuffer>();
+	InterpEdTrans->Initialize(8 * 1024 * 1024);
 	InterpEdTrans->OnUndo().AddRaw(this, &FMatinee::OnPostUndoRedo);
 	InterpEdTrans->OnRedo().AddRaw(this, &FMatinee::OnPostUndoRedo);
 	GEditor->Trans = InterpEdTrans;
diff --git a/Engine/Source/Editor/Matinee/Private/MatineeTransBuffer.cpp b/Engine/Source/Editor/Matinee/Private/MatineeTransBuffer.cpp
--- a/Engine/Source/Editor/Matinee/Private/MatineeTransBuffer.cpp
+++ b/Engine/Source/Editor/Matinee/Private/MatineeTransBuffer.cpp
@@ -7,10 +7,6 @@
 /*-----------------------------------------------------------------------------
 	UMatineeTransBuffer / FMatineeTransaction
 -----------------------------------------------------------------------------*/
-UMatineeTransBuffer::UMatineeTransBuffer(const FObjectInitializer& ObjectInitializer)
-	: Super(ObjectInitializer)
-{
-}
 
 void UMatineeTransBuffer::BeginSpecial(const FText& Description)
 {
diff --git a/Engine/Source/Editor/UMGEditor/Private/WidgetBlueprintEditorUtils.cpp b/Engine/Source/Editor/UMGEditor/Private/WidgetBlueprintEditorUtils.cpp
--- a/Engine/Source/Editor/UMGEditor/Private/WidgetBlueprintEditorUtils.cpp
+++ b/Engine/Source/Editor/UMGEditor/Private/WidgetBlueprintEditorUtils.cpp
@@ -479,7 +479,7 @@ void FWidgetBlueprintEditorUtils::ExportWidgetsToText(TSet<UWidget*> WidgetsToEx
 
 		if ( Widget->GetParent() == nullptr || !WidgetsToExport.Contains(Widget->GetParent()) )
 		{
-			UWidgetSlotPair* SlotMetaData = new UWidgetSlotPair();
+			auto SlotMetaData = NewObject<UWidgetSlotPair>();
 			SlotMetaData->SetWidget(Widget);
 
 			UExporter::ExportToOutputDevice(&Context, SlotMetaData, nullptr, Archive, TEXT("copy"), 0, PPF_ExportsNotFullyQualified | PPF_Copy | PPF_Delimited, false, nullptr);
diff --git a/Engine/Source/Editor/UnrealEd/Classes/Editor/TransBuffer.h b/Engine/Source/Editor/UnrealEd/Classes/Editor/TransBuffer.h
--- a/Engine/Source/Editor/UnrealEd/Classes/Editor/TransBuffer.h
+++ b/Engine/Source/Editor/UnrealEd/Classes/Editor/TransBuffer.h
@@ -11,7 +11,8 @@ UCLASS(transient)
 class UNREALED_API UTransBuffer
 	: public UTransactor
 {
-    GENERATED_UCLASS_BODY()
+public:
+    GENERATED_BODY()
 	// Variables.
 	/** The queue of transaction records */
 	TArray<FTransaction> UndoBuffer;
@@ -34,7 +35,7 @@ class UNREALED_API UTransBuffer
 public:
 
 	// Constructor.
-	UTransBuffer( const FObjectInitializer& ObjectInitializer,SIZE_T InMaxMemory );
+	void Initialize(SIZE_T InMaxMemory);
 
 public:
 
diff --git a/Engine/Source/Editor/UnrealEd/Private/BSPOps.cpp b/Engine/Source/Editor/UnrealEd/Private/BSPOps.cpp
--- a/Engine/Source/Editor/UnrealEd/Private/BSPOps.cpp
+++ b/Engine/Source/Editor/UnrealEd/Private/BSPOps.cpp
@@ -549,8 +549,9 @@ void FBSPOps::csgCopyBrush( ABrush* Dest, ABrush* Src, uint32 PolyFlags, EObject
 
 	// Duplicate the brush and its polys.
 	Dest->PolyFlags		= PolyFlags;
-	Dest->Brush			= new( Dest, NAME_None, ResFlags )UModel(FObjectInitializer(), NULL, Src->Brush->RootOutside );
-	Dest->Brush->Polys	= new( Dest->Brush, NAME_None, ResFlags )UPolys(FObjectInitializer());
+	Dest->Brush = NewNamedObject<UModel>(Dest, NAME_None, ResFlags);
+	Dest->Brush->Initialize(nullptr, Src->Brush->RootOutside);
+	Dest->Brush->Polys	= NewNamedObject<UPolys>(Dest->Brush, NAME_None, ResFlags);
 	check(Dest->Brush->Polys->Element.GetOwner()==Dest->Brush->Polys);
 	Dest->Brush->Polys->Element.AssignButKeepOwner(Src->Brush->Polys->Element);
 	check(Dest->Brush->Polys->Element.GetOwner()==Dest->Brush->Polys);
diff --git a/Engine/Source/Editor/UnrealEd/Private/CubemapUnwrapUtils.cpp b/Engine/Source/Editor/UnrealEd/Private/CubemapUnwrapUtils.cpp
--- a/Engine/Source/Editor/UnrealEd/Private/CubemapUnwrapUtils.cpp
+++ b/Engine/Source/Editor/UnrealEd/Private/CubemapUnwrapUtils.cpp
@@ -41,7 +41,7 @@ namespace CubemapHelpers
 		// If the source format is 8 bit per channel or less then select a LDR target format.
 		const EPixelFormat TargetPixelFormat = CalculateImageBytes(1, 1, 0, SourcePixelFormat) <= 4 ? PF_B8G8R8A8 : PF_FloatRGBA;
 
-		UTextureRenderTarget2D* RenderTargetLongLat = new UTextureRenderTarget2D(FObjectInitializer());
+		UTextureRenderTarget2D* RenderTargetLongLat = NewObject<UTextureRenderTarget2D>();
 		check(RenderTargetLongLat);
 		RenderTargetLongLat->AddToRoot();
 		RenderTargetLongLat->ClearColor = FLinearColor(0.0f, 0.0f, 0.0f, 0.0f);
diff --git a/Engine/Source/Editor/UnrealEd/Private/Editor.cpp b/Engine/Source/Editor/UnrealEd/Private/Editor.cpp
--- a/Engine/Source/Editor/UnrealEd/Private/Editor.cpp
+++ b/Engine/Source/Editor/UnrealEd/Private/Editor.cpp
@@ -203,15 +203,15 @@ static void OnObjectSelected(UObject* Object)
 
 static void PrivateInitSelectedSets()
 {
-	PrivateGetSelectedActors() = new( GetTransientPackage(), TEXT("SelectedActors"), RF_Transactional ) USelection(FObjectInitializer());
+	PrivateGetSelectedActors() = NewNamedObject<USelection>(GetTransientPackage(), TEXT("SelectedActors"), RF_Transactional);
 	PrivateGetSelectedActors()->AddToRoot();
 
 	PrivateGetSelectedActors()->SelectObjectEvent.AddStatic(&OnObjectSelected);
 
-	PrivateGetSelectedComponents() = new(GetTransientPackage(), TEXT("SelectedComponents"), RF_Transactional) USelection(FObjectInitializer());
+	PrivateGetSelectedComponents() = NewNamedObject<USelection>(GetTransientPackage(), TEXT("SelectedComponents"), RF_Transactional);
 	PrivateGetSelectedComponents()->AddToRoot();
 
-	PrivateGetSelectedObjects() = new(GetTransientPackage(), TEXT("SelectedObjects"), RF_Transactional) USelection(FObjectInitializer());
+	PrivateGetSelectedObjects() = NewNamedObject<USelection>(GetTransientPackage(), TEXT("SelectedObjects"), RF_Transactional);
 	PrivateGetSelectedObjects()->AddToRoot();
 }
 
@@ -493,8 +493,10 @@ void UEditorEngine::InitEditor(IEngineLoop* InEngineLoop)
 	UNavigationSystem::SetNavigationAutoUpdateEnabled(GetDefault<ULevelEditorMiscSettings>()->bNavigationAutoUpdate, EditorContext.World()->GetNavigationSystem());
 
 	// Allocate temporary model.
-	TempModel = new UModel(FObjectInitializer(), NULL, 1);
-	ConversionTempModel = new UModel(FObjectInitializer(), NULL, 1);
+	TempModel = NewObject<UModel>();
+	TempModel->Initialize(nullptr, 1);
+	ConversionTempModel = NewObject<UModel>();
+	ConversionTempModel->Initialize(nullptr, 1);
 
 	// create the timer manager
 	TimerManager = MakeShareable(new FTimerManager());
@@ -6211,7 +6213,7 @@ TArray<AActor*> UEditorEngine::AddExportTextActors(const FString& ExportText, bo
 	}
 
 	// Use a level factory to spawn all the actors using the ExportText
-	ULevelFactory* Factory = new ULevelFactory(FObjectInitializer());
+	auto Factory = NewObject<ULevelFactory>();
 	FVector Location;
 	{
 		FScopedTransaction Transaction( NSLOCTEXT("UnrealEd", "AddActor", "Add Actor") );
diff --git a/Engine/Source/Editor/UnrealEd/Private/EditorActor.cpp b/Engine/Source/Editor/UnrealEd/Private/EditorActor.cpp
--- a/Engine/Source/Editor/UnrealEd/Private/EditorActor.cpp
+++ b/Engine/Source/Editor/UnrealEd/Private/EditorActor.cpp
@@ -262,7 +262,7 @@ void UUnrealEdEngine::edactPasteSelected(UWorld* InWorld, bool bDuplicate, bool
 	const TCHAR* Paste = *PasteString;
 
 	// Import the actors.
-	ULevelFactory* Factory = new ULevelFactory(FObjectInitializer());
+	auto Factory = NewObject<ULevelFactory>();
 	Factory->FactoryCreateText( ULevel::StaticClass(), InWorld->GetCurrentLevel(), InWorld->GetCurrentLevel()->GetFName(), RF_Transactional, NULL, bDuplicate ? TEXT("move") : TEXT("paste"), Paste, Paste+FCString::Strlen(Paste), GWarn );
 
 	// Fire ULevel::LevelDirtiedEvent when falling out of scope.
diff --git a/Engine/Source/Editor/UnrealEd/Private/EditorExporters.cpp b/Engine/Source/Editor/UnrealEd/Private/EditorExporters.cpp
--- a/Engine/Source/Editor/UnrealEd/Private/EditorExporters.cpp
+++ b/Engine/Source/Editor/UnrealEd/Private/EditorExporters.cpp
@@ -1300,7 +1300,7 @@ bool GenerateExportMaterialPropertyData(
 
 		if(!RenderTarget)
 		{
-			RenderTarget = new UTextureRenderTarget2D(FObjectInitializer());
+			RenderTarget = NewObject<UTextureRenderTarget2D>();
 			check(RenderTarget);
 			RenderTarget->AddToRoot();
 			RenderTarget->ClearColor = FLinearColor(0.0f, 0.0f, 0.0f, 0.0f);
@@ -2630,7 +2630,7 @@ namespace MaterialExportUtils
 			OutFlattenMaterial.DiffuseSize.Y > 0)
 		{
 			// Create temporary render target
-			UTextureRenderTarget2D* RenderTargetDiffuse = new UTextureRenderTarget2D(FObjectInitializer());
+			auto RenderTargetDiffuse = NewObject<UTextureRenderTarget2D>();
 			check(RenderTargetDiffuse);
 			RenderTargetDiffuse->AddToRoot();
 			RenderTargetDiffuse->ClearColor = FLinearColor(0.0f, 0.0f, 0.0f, 0.0f);
@@ -2663,7 +2663,7 @@ namespace MaterialExportUtils
 			InMaterial->GetMaterial()->HasNormalConnected())
 		{
 			// Create temporary render target
-			UTextureRenderTarget2D* RenderTargetNormal = new UTextureRenderTarget2D(FObjectInitializer());
+			auto RenderTargetNormal = NewObject<UTextureRenderTarget2D>();
 			check(RenderTargetNormal);
 			RenderTargetNormal->AddToRoot();
 			RenderTargetNormal->ClearColor = FLinearColor(0.0f, 0.0f, 0.0f, 0.0f);
@@ -2704,7 +2704,7 @@ namespace MaterialExportUtils
 		float TargetGamma,
 		TArray<FColor>& OutSamples)
 	{
-		UTextureRenderTarget2D* RenderTargetTexture = new UTextureRenderTarget2D(FObjectInitializer());
+		auto RenderTargetTexture = NewObject<UTextureRenderTarget2D>();
 		check(RenderTargetTexture);
 		RenderTargetTexture->AddToRoot();
 		RenderTargetTexture->ClearColor = FLinearColor::Transparent;
diff --git a/Engine/Source/Editor/UnrealEd/Private/EditorObject.cpp b/Engine/Source/Editor/UnrealEd/Private/EditorObject.cpp
--- a/Engine/Source/Editor/UnrealEd/Private/EditorObject.cpp
+++ b/Engine/Source/Editor/UnrealEd/Private/EditorObject.cpp
@@ -231,7 +231,7 @@ static const TCHAR* ImportProperties(
 					ExistingBrush->Rename();
 
 				// Create model.
-				UModelFactory* ModelFactory = new UModelFactory(FObjectInitializer());
+				UModelFactory* ModelFactory = NewObject<UModelFactory>();
 				ModelFactory->FactoryCreateText( UModel::StaticClass(), SubobjectRoot, FName(BrushName, FNAME_Add, true), RF_NoFlags, NULL, TEXT("t3d"), SourceText, SourceText+FCString::Strlen(SourceText), Warn );
 				ImportedBrush = 1;
 			}
diff --git a/Engine/Source/Editor/UnrealEd/Private/EditorServer.cpp b/Engine/Source/Editor/UnrealEd/Private/EditorServer.cpp
--- a/Engine/Source/Editor/UnrealEd/Private/EditorServer.cpp
+++ b/Engine/Source/Editor/UnrealEd/Private/EditorServer.cpp
@@ -565,7 +565,7 @@ bool UEditorEngine::Exec_Brush( UWorld* InWorld, const TCHAR* Str, FOutputDevice
 			WorldBrush->SetActorLocation(SnapLocation - PrePivot, false);
 			WorldBrush->SetPrePivot( FVector::ZeroVector );
 			WorldBrush->Brush->Polys->Element.Empty();
-			UPolysFactory* It = new UPolysFactory(FObjectInitializer());
+			UPolysFactory* It = NewObject<UPolysFactory>();
 			It->FactoryCreateText( UPolys::StaticClass(), WorldBrush->Brush->Polys->GetOuter(), *WorldBrush->Brush->Polys->GetName(), RF_NoFlags, WorldBrush->Brush->Polys, TEXT("t3d"), GStream, GStream+FCString::Strlen(GStream), GWarn );
 			// Do NOT merge faces.
 			FBSPOps::bspValidateBrush( WorldBrush->Brush, 0, 1 );
@@ -1111,7 +1111,8 @@ UTransactor* UEditorEngine::CreateTrans()
 		UndoBufferSize = 16;
 	}
 
-	UTransBuffer* TransBuffer = new UTransBuffer(FObjectInitializer(), UndoBufferSize * 1024 * 1024);
+	UTransBuffer* TransBuffer = NewObject<UTransBuffer>();
+	TransBuffer->Initialize(UndoBufferSize * 1024 * 1024);
 	TransBuffer->OnBeforeRedoUndo().AddUObject(this, &UEditorEngine::HandleTransactorBeforeRedoUndo);
 	TransBuffer->OnRedo().AddUObject(this, &UEditorEngine::HandleTransactorRedo);
 	TransBuffer->OnUndo().AddUObject(this, &UEditorEngine::HandleTransactorUndo);
@@ -2827,11 +2828,13 @@ void UEditorEngine::DoMoveSelectedActorsToLevel( ULevel* InDestLevel )
 
 ULevel*  UEditorEngine::CreateTransLevelMoveBuffer( UWorld* InWorld )
 {
-	ULevel* BufferLevel = new(GetTransientPackage(), TEXT("TransLevelMoveBuffer")) ULevel(FObjectInitializer(),FURL(NULL));
+	ULevel* BufferLevel = NewNamedObject<ULevel>(GetTransientPackage(), TEXT("TransLevelMoveBuffer"));
+	BufferLevel->Initialize(FURL(nullptr));
 	check( BufferLevel );
 	BufferLevel->AddToRoot();
 	BufferLevel->OwningWorld = InWorld;
-	BufferLevel->Model = new( BufferLevel ) UModel( FObjectInitializer(), NULL, true );
+	BufferLevel->Model = NewObject<UModel>(BufferLevel);
+	BufferLevel->Model->Initialize(nullptr, true);
 	BufferLevel->bIsVisible = true;
 		
 	BufferLevel->SetFlags( RF_Transactional );
@@ -2850,7 +2853,8 @@ ULevel*  UEditorEngine::CreateTransLevelMoveBuffer( UWorld* InWorld )
 	BufferDefaultBrush->Brush = CastChecked<UModel>(StaticFindObject(UModel::StaticClass(), BufferLevel->OwningWorld->GetOuter(), TEXT("Brush"), true), ECastCheckedType::NullAllowed);
 	if (!BufferDefaultBrush->Brush)
 	{
-		BufferDefaultBrush->Brush = new( InWorld, TEXT("Brush") )UModel( FObjectInitializer(), BufferDefaultBrush, 1 );
+		BufferDefaultBrush->Brush = NewNamedObject<UModel>(InWorld, TEXT("Brush"));
+		BufferDefaultBrush->Brush->Initialize(BufferDefaultBrush, 1);
 	}
 	BufferDefaultBrush->GetBrushComponent()->Brush = BufferDefaultBrush->Brush;
 	BufferDefaultBrush->SetNotForClientOrServer();
diff --git a/Engine/Source/Editor/UnrealEd/Private/EditorTransaction.cpp b/Engine/Source/Editor/UnrealEd/Private/EditorTransaction.cpp
--- a/Engine/Source/Editor/UnrealEd/Private/EditorTransaction.cpp
+++ b/Engine/Source/Editor/UnrealEd/Private/EditorTransaction.cpp
@@ -342,15 +342,9 @@ UTransactor::UTransactor(const FObjectInitializer& ObjectInitializer)
 {
 }
 
-UTransBuffer::UTransBuffer(const FObjectInitializer& ObjectInitializer)
-	: Super(ObjectInitializer)
-{
-}
-
-UTransBuffer::UTransBuffer( const FObjectInitializer& ObjectInitializer, SIZE_T InMaxMemory )
-:	UTransactor(ObjectInitializer)
-,	MaxMemory( InMaxMemory )
+void UTransBuffer::Initialize(SIZE_T InMaxMemory)
 {
+	MaxMemory = InMaxMemory;
 	// Reset.
 	Reset( NSLOCTEXT("UnrealEd", "Startup", "Startup") );
 	CheckState();
diff --git a/Engine/Source/Editor/UnrealEd/Private/Factories/ActorFactory.cpp b/Engine/Source/Editor/UnrealEd/Private/Factories/ActorFactory.cpp
--- a/Engine/Source/Editor/UnrealEd/Private/Factories/ActorFactory.cpp
+++ b/Engine/Source/Editor/UnrealEd/Private/Factories/ActorFactory.cpp
@@ -1773,8 +1773,9 @@ void CreateBrushForVolumeActor( AVolume* NewActor, UBrushBuilder* BrushBuilder )
 		NewActor->PreEditChange(NULL);
 
 		NewActor->PolyFlags = 0;
-		NewActor->Brush = new( NewActor, NAME_None, RF_Transactional )UModel(FObjectInitializer(), NULL, true );
-		NewActor->Brush->Polys = new( NewActor->Brush, NAME_None, RF_Transactional )UPolys(FObjectInitializer());
+		NewActor->Brush = NewNamedObject<UModel>(NewActor, NAME_None, RF_Transactional);
+		NewActor->Brush->Initialize(nullptr, true);
+		NewActor->Brush->Polys = NewNamedObject<UPolys>(NewActor->Brush, NAME_None, RF_Transactional);
 		NewActor->GetBrushComponent()->Brush = NewActor->Brush;
 		if(BrushBuilder != nullptr)
 		{
diff --git a/Engine/Source/Editor/UnrealEd/Private/Factories/EditorFactories.cpp b/Engine/Source/Editor/UnrealEd/Private/Factories/EditorFactories.cpp
--- a/Engine/Source/Editor/UnrealEd/Private/Factories/EditorFactories.cpp
+++ b/Engine/Source/Editor/UnrealEd/Private/Factories/EditorFactories.cpp
@@ -813,7 +813,7 @@ UObject* ULevelFactory::FactoryCreateText
 
 	if (GIsImportingT3D && (MapPackageText.Len() > 0))
 	{
-		UPackageFactory* PackageFactory = new UPackageFactory(FObjectInitializer());
+		UPackageFactory* PackageFactory = NewObject<UPackageFactory>();
 		check(PackageFactory);
 
 		FName NewPackageName(*(RootMapPackage->GetName()));
@@ -1201,8 +1201,8 @@ UObject* UPolysFactory::FactoryCreateText
 {
 	FEditorDelegates::OnAssetPreImport.Broadcast(this, Class, InParent, Name, Type);
 
-	// Create polys.
-	UPolys* Polys = Context ? CastChecked<UPolys>(Context) : new(InParent,Name,Flags)UPolys(FObjectInitializer());
+	// Create polys.	
+	UPolys* Polys = Context ? CastChecked<UPolys>(Context) : NewNamedObject<UPolys>(InParent, Name, Flags);
 
 	// Eat up if present.
 	GetBEGIN( &Buffer, TEXT("POLYLIST") );
@@ -1474,7 +1474,8 @@ UObject* UModelFactory::FactoryCreateText
 	FEditorDelegates::OnAssetPreImport.Broadcast(this, Class, InParent, Name, Type);
 
 	ABrush* TempOwner = (ABrush*)Context;
-	UModel* Model = new( InParent, Name, Flags )UModel( FObjectInitializer(),TempOwner, 1 );
+	UModel* Model = NewNamedObject<UModel>(InParent, Name, Flags);
+	Model->Initialize(TempOwner, true);
 
 	const TCHAR* StrPtr;
 	FString StrLine;
@@ -1492,7 +1493,7 @@ UObject* UModelFactory::FactoryCreateText
 		}
 		else if( GetBEGIN (&StrPtr,TEXT("POLYLIST")) )
 		{
-			UPolysFactory* PolysFactory = new UPolysFactory(FObjectInitializer());
+			UPolysFactory* PolysFactory = NewObject<UPolysFactory>();
 			Model->Polys = (UPolys*)PolysFactory->FactoryCreateText(UPolys::StaticClass(),Model,NAME_None,RF_Transactional,nullptr,Type,Buffer,BufferEnd,Warn);
 			check(Model->Polys);
 		}
@@ -1770,7 +1771,7 @@ UObject* USoundFactory::FactoryCreateBinary
 
 		// Use pre-existing sound if it exists and we want to keep settings,
 		// otherwise create new sound and import raw data.
-		USoundWave* Sound = (bUseExistingSettings && ExistingSound) ? ExistingSound : new( InParent, Name, Flags ) USoundWave(FObjectInitializer());
+		USoundWave* Sound = (bUseExistingSettings && ExistingSound) ? ExistingSound : NewNamedObject<USoundWave>(InParent, Name, Flags);
 		
 		if (bUseExistingSettings && ExistingSound)
 		{
@@ -2077,7 +2078,7 @@ UObject* USoundSurroundFactory::FactoryCreateBinary
 
 			if (Sound == nullptr)
 			{
-				Sound = new( InParent, BaseName, Flags ) USoundWave(FObjectInitializer());
+				Sound = NewNamedObject<USoundWave>(InParent, BaseName, Flags);
 			}
 		}
 
@@ -4509,7 +4510,7 @@ UObject* UTextureFactory::FactoryCreateBinary
 		UPackage* MaterialPackage = CreatePackage(nullptr, *MaterialPackageName);
 
 		// Create the material
-		UMaterialFactoryNew* Factory = new UMaterialFactoryNew(FObjectInitializer());
+		UMaterialFactoryNew* Factory = NewObject<UMaterialFactoryNew>();
 		UMaterial* Material = (UMaterial*)Factory->FactoryCreateNew( UMaterial::StaticClass(), MaterialPackage, *MaterialName, Flags, Context, Warn );
 
 		// Notify the asset registry
@@ -5245,7 +5246,8 @@ UObject* UFontFileImportFactory::FactoryCreateBinary(UClass* InClass, UObject* I
 		Font->FontCacheType = EFontCacheType::Runtime;
 
 		// We need to allocate the bulk data with the font as its outer
-		const UFontBulkData* const BulkData = new(Font) UFontBulkData(InBuffer, InBufferEnd - InBuffer);
+		UFontBulkData* const BulkData = NewObject<UFontBulkData>(Font);
+		BulkData->Initialize(InBuffer, InBufferEnd - InBuffer);
 
 		Font->CompositeFont.DefaultTypeface.Fonts.Add(FTypefaceEntry("Default", GetCurrentFilename(), BulkData, EFontHinting::Auto));
 	}
diff --git a/Engine/Source/Editor/UnrealEd/Private/Factories/TTFontImport.cpp b/Engine/Source/Editor/UnrealEd/Private/Factories/TTFontImport.cpp
--- a/Engine/Source/Editor/UnrealEd/Private/Factories/TTFontImport.cpp
+++ b/Engine/Source/Editor/UnrealEd/Private/Factories/TTFontImport.cpp
@@ -153,7 +153,7 @@ UObject* UTrueTypeFontFactory::FactoryCreateNew(
 	}
 
 	// Create font and its texture.
-	UFont* Font = new( InParent, Name, Flags )UFont(FObjectInitializer());
+	auto Font = NewNamedObject<UFont>(InParent, Name, Flags);
 	
 	if (ImportOptions->Data.bUseDistanceFieldAlpha)
 	{
@@ -777,7 +777,7 @@ UTexture2D* UTrueTypeFontFactory::CreateTextureFromDC( UFont* Font, HDC dc, int3
 	}
 
 	// Create texture for page.
-	UTexture2D* Texture = new(Font, *TextureString)UTexture2D(FObjectInitializer());
+	auto Texture = NewNamedObject<UTexture2D>(Font, *TextureString);
 
 	// note RF_Public because font textures can be referenced directly by material expressions
 	Texture->SetFlags(RF_Public);
diff --git a/Engine/Source/Editor/UnrealEd/Private/Fbx/FbxMaterialImport.cpp b/Engine/Source/Editor/UnrealEd/Private/Fbx/FbxMaterialImport.cpp
--- a/Engine/Source/Editor/UnrealEd/Private/Fbx/FbxMaterialImport.cpp
+++ b/Engine/Source/Editor/UnrealEd/Private/Fbx/FbxMaterialImport.cpp
@@ -111,7 +111,7 @@ UTexture* UnFbx::FFbxImporter::ImportTexture( FbxFileTexture* FbxTexture, bool b
 	{
 		UE_LOG(LogFbxMaterialImport, Verbose, TEXT("Loading texture file %s"),*Filename);
 		const uint8* PtrTexture = DataBinary.GetData();
-		UTextureFactory* TextureFact = new UTextureFactory(FObjectInitializer());
+		auto TextureFact = NewObject<UTextureFactory>();
 		TextureFact->AddToRoot();
 
 		// save texture settings if texture exist
@@ -429,7 +429,7 @@ void UnFbx::FFbxImporter::CreateUnrealMaterial(FbxSurfaceMaterial& FbxMaterial,
 	
 
 	// create an unreal material asset
-	UMaterialFactoryNew* MaterialFactory = new UMaterialFactoryNew(FObjectInitializer());
+	auto MaterialFactory = NewObject<UMaterialFactoryNew>();
 	
 	UMaterial* UnrealMaterial = (UMaterial*)MaterialFactory->FactoryCreateNew(
 		UMaterial::StaticClass(), Package, *MaterialFullName, RF_Standalone|RF_Public, NULL, GWarn );
diff --git a/Engine/Source/Editor/UnrealEd/Private/Fbx/FbxStaticMeshImport.cpp b/Engine/Source/Editor/UnrealEd/Private/Fbx/FbxStaticMeshImport.cpp
--- a/Engine/Source/Editor/UnrealEd/Private/Fbx/FbxStaticMeshImport.cpp
+++ b/Engine/Source/Editor/UnrealEd/Private/Fbx/FbxStaticMeshImport.cpp
@@ -956,7 +956,7 @@ UStaticMesh* UnFbx::FFbxImporter::ImportStaticMeshAsSingle(UObject* InParent, TA
 	}
 	else
 	{
-		StaticMesh = new(Package,FName(*MeshName),Flags|RF_Public) UStaticMesh(FObjectInitializer());
+		StaticMesh = NewNamedObject<UStaticMesh>(Package, FName(*MeshName), Flags | RF_Public);
 	}
 
 	if (StaticMesh->SourceModels.Num() < LODIndex+1)
diff --git a/Engine/Source/Editor/UnrealEd/Private/GeomFitUtils.cpp b/Engine/Source/Editor/UnrealEd/Private/GeomFitUtils.cpp
--- a/Engine/Source/Editor/UnrealEd/Private/GeomFitUtils.cpp
+++ b/Engine/Source/Editor/UnrealEd/Private/GeomFitUtils.cpp
@@ -83,7 +83,8 @@ int32 GenerateKDopAsSimpleCollision(UStaticMesh* StaticMesh, const TArray<FVecto
 		maxDist.Add(-MY_FLTMAX);
 
 	// Construct temporary UModel for kdop creation. We keep no refs to it, so it can be GC'd.
-	UModel* TempModel = new UModel(FObjectInitializer(),NULL,1);
+	auto TempModel = NewObject<UModel>();
+	TempModel->Initialize(nullptr, 1);
 
 	// For each vertex, project along each kdop direction, to find the max in that direction.
 	const FStaticMeshLODResources& RenderData = StaticMesh->RenderData->LODResources[0];
diff --git a/Engine/Source/Editor/UnrealEd/Private/LevelEditorViewport.cpp b/Engine/Source/Editor/UnrealEd/Private/LevelEditorViewport.cpp
--- a/Engine/Source/Editor/UnrealEd/Private/LevelEditorViewport.cpp
+++ b/Engine/Source/Editor/UnrealEd/Private/LevelEditorViewport.cpp
@@ -411,7 +411,7 @@ static UObject* GetOrCreateMaterialFromTexture( UTexture* UnrealTexture )
 	}
 
 	// create an unreal material asset
-	UMaterialFactoryNew* MaterialFactory = new UMaterialFactoryNew( FObjectInitializer() );
+	auto MaterialFactory = NewObject<UMaterialFactoryNew>();
 
 	UMaterial* UnrealMaterial = (UMaterial*)MaterialFactory->FactoryCreateNew(
 		UMaterial::StaticClass(), Package, *MaterialFullName, RF_Standalone | RF_Public, NULL, GWarn );
diff --git a/Engine/Source/Editor/UnrealEd/Private/Lightmass/LightmassRender.cpp b/Engine/Source/Editor/UnrealEd/Private/Lightmass/LightmassRender.cpp
--- a/Engine/Source/Editor/UnrealEd/Private/Lightmass/LightmassRender.cpp
+++ b/Engine/Source/Editor/UnrealEd/Private/Lightmass/LightmassRender.cpp
@@ -1026,7 +1026,7 @@ bool FLightmassMaterialRenderer::CreateRenderTarget(EPixelFormat InFormat, int32
 
 	if (RenderTarget == NULL)
 	{
-		RenderTarget = new UTextureRenderTarget2D(FObjectInitializer());
+		RenderTarget = NewObject<UTextureRenderTarget2D>();
 		check(RenderTarget);
 		RenderTarget->AddToRoot();
 		RenderTarget->ClearColor = FLinearColor(0.0f, 0.0f, 0.0f, 0.0f);
diff --git a/Engine/Source/Editor/UnrealEd/Private/StaticMeshEdit.cpp b/Engine/Source/Editor/UnrealEd/Private/StaticMeshEdit.cpp
--- a/Engine/Source/Editor/UnrealEd/Private/StaticMeshEdit.cpp
+++ b/Engine/Source/Editor/UnrealEd/Private/StaticMeshEdit.cpp
@@ -252,7 +252,8 @@ private:
 void DecomposeUCXMesh( const TArray<FVector>& CollisionVertices, const TArray<int32>& CollisionFaceIdx, UBodySetup* BodySetup )
 {
 	// We keep no ref to this Model, so it will be GC'd at some point after the import.
-	UModel* TempModel = new UModel(FObjectInitializer(),NULL,1);
+	auto TempModel = NewObject<UModel>();
+	TempModel->Initialize(nullptr, 1);
 
 	FMeshConnectivityBuilder ConnectivityBuilder;
 
@@ -509,7 +510,7 @@ UStaticMesh* CreateStaticMesh(struct FRawMesh& RawMesh,TArray<UMaterialInterface
 {
 	// Create the UStaticMesh object.
 	FStaticMeshComponentRecreateRenderStateContext RecreateRenderStateContext(FindObject<UStaticMesh>(InOuter,*InName.ToString()));
-	UStaticMesh* StaticMesh = new(InOuter,InName,RF_Public|RF_Standalone) UStaticMesh(FObjectInitializer());
+	auto StaticMesh = NewNamedObject<UStaticMesh>(InOuter, InName, RF_Public | RF_Standalone);
 
 	// Add one LOD for the base mesh
 	FStaticMeshSourceModel* SrcModel = new(StaticMesh->SourceModels) FStaticMeshSourceModel();
diff --git a/Engine/Source/Editor/WorldBrowser/Private/Tiles/WorldTileCollectionModel.cpp b/Engine/Source/Editor/WorldBrowser/Private/Tiles/WorldTileCollectionModel.cpp
--- a/Engine/Source/Editor/WorldBrowser/Private/Tiles/WorldTileCollectionModel.cpp
+++ b/Engine/Source/Editor/WorldBrowser/Private/Tiles/WorldTileCollectionModel.cpp
@@ -2044,8 +2044,8 @@ bool FWorldTileCollectionModel::GenerateLODLevels(FLevelModelList InLevelList, i
 				MeshOuter->FullyLoad();
 				MeshOuter->Modify();
 			}
-					
-			UStaticMesh* StaticMesh = new(MeshOuter, *LandscapeMeshAssetName, RF_Public|RF_Standalone) UStaticMesh(FObjectInitializer());
+
+			auto StaticMesh = NewNamedObject<UStaticMesh>(MeshOuter, *LandscapeMeshAssetName, RF_Public | RF_Standalone);
 			StaticMesh->InitResources();
 			{
 				FString OutputPath = StaticMesh->GetPathName();
diff --git a/Engine/Source/Programs/UnrealHeaderTool/Private/CodeGenerator.cpp b/Engine/Source/Programs/UnrealHeaderTool/Private/CodeGenerator.cpp
--- a/Engine/Source/Programs/UnrealHeaderTool/Private/CodeGenerator.cpp
+++ b/Engine/Source/Programs/UnrealHeaderTool/Private/CodeGenerator.cpp
@@ -726,7 +726,7 @@ FString FNativeClassHeaderGenerator::PropertyNew(FString& Meta, UProperty* Prop,
 		ExtraArgs = FString::Printf(TEXT(", %s"), *GetSingletonName(TargetFunction));
 	}
 
-	FString Constructor = FString::Printf(TEXT("new(%s, TEXT(\"%s\"), RF_Public|RF_Transient|RF_Native) U%s(%s, 0x%016llx%s);"),
+	FString Constructor = FString::Printf(TEXT("new(EC_InternalUseOnlyConstructor, %s, TEXT(\"%s\"), RF_Public|RF_Transient|RF_Native) U%s(%s, 0x%016llx%s);"),
 		*OuterString,
 		*Prop->GetName(), 
 		*Prop->GetClass()->GetName(), 
@@ -1251,7 +1251,7 @@ void FNativeClassHeaderGenerator::ExportFunction(UFunction* Function, FScope* Sc
 
 	FString UFunctionType = bIsDelegate ? TEXT("UDelegateFunction") : TEXT("UFunction");
 
-	CurrentFunctionText.Logf(TEXT("            ReturnFunction = new(Outer, TEXT(\"%s\"), RF_Public|RF_Transient|RF_Native) %s(FObjectInitializer(), %s, 0x%08X, %d%s);\r\n"),
+	CurrentFunctionText.Logf(TEXT("            ReturnFunction = new(EC_InternalUseOnlyConstructor, Outer, TEXT(\"%s\"), RF_Public|RF_Transient|RF_Native) %s(FObjectInitializer(), %s, 0x%08X, %d%s);\r\n"),
 		*Function->GetName(),
 		*UFunctionType,
 		*SuperFunctionString,
@@ -1892,7 +1892,7 @@ void ExportConstructorDefinition(FStringOutputDevice& Out, FClass* Class, const
 	if (!ClassData->bConstructorDeclared)
 	{
 		Out.Logf(TEXT("%s/** Standard constructor, called after all reflected properties have been initialized */\r\n"), FCString::Spc(4));
-		Out.Logf(TEXT("%s%s_API %s(const class FObjectInitializer& ObjectInitializer) : Super(ObjectInitializer) { };\r\n"), FCString::Spc(4), *API, NameLookupCPP.GetNameCPP(Class));
+		Out.Logf(TEXT("%s%s_API %s(const class FObjectInitializer& ObjectInitializer = FObjectInitializer::Get()) : Super(ObjectInitializer) { };\r\n"), FCString::Spc(4), *API, NameLookupCPP.GetNameCPP(Class));
 
 		ClassData->bConstructorDeclared = true;
 		ClassData->bObjectInitializerConstructorDeclared = true;
@@ -2357,7 +2357,7 @@ void FNativeClassHeaderGenerator::ExportGeneratedStructBodyMacros(FUnrealSourceF
 				ExplicitSizeString = FString::Printf(TEXT(", sizeof(%s), ALIGNOF(%s)"), NameLookupCPP.GetNameCPP(ScriptStruct), NameLookupCPP.GetNameCPP(ScriptStruct));
 			}
 
-			GeneratedStructRegisterFunctionText.Logf(TEXT("            ReturnStruct = new(Outer, TEXT(\"%s\"), RF_Public|RF_Transient|RF_Native) UScriptStruct(FObjectInitializer(), %s, %s, EStructFlags(0x%08X)%s);\r\n"),
+			GeneratedStructRegisterFunctionText.Logf(TEXT("            ReturnStruct = new(EC_InternalUseOnlyConstructor, Outer, TEXT(\"%s\"), RF_Public|RF_Transient|RF_Native) UScriptStruct(FObjectInitializer(), %s, %s, EStructFlags(0x%08X)%s);\r\n"),
 				*ScriptStruct->GetName(),
 				*BaseStructString,
 				*CppStructOpsString,
@@ -2458,7 +2458,7 @@ void FNativeClassHeaderGenerator::ExportGeneratedEnumsInitCode(const TArray<UEnu
 			GeneratedEnumRegisterFunctionText.Logf(TEXT("        if (!ReturnEnum)\r\n"));
 			GeneratedEnumRegisterFunctionText.Logf(TEXT("        {\r\n"));
 
-			GeneratedEnumRegisterFunctionText.Logf(TEXT("            ReturnEnum = new(Outer, TEXT(\"%s\"), RF_Public|RF_Transient|RF_Native) UEnum(FObjectInitializer());\r\n"), *Enum->GetName());
+			GeneratedEnumRegisterFunctionText.Logf(TEXT("            ReturnEnum = new(EC_InternalUseOnlyConstructor, Outer, TEXT(\"%s\"), RF_Public|RF_Transient|RF_Native) UEnum(FObjectInitializer());\r\n"), *Enum->GetName());
 			GeneratedEnumRegisterFunctionText.Logf(TEXT("            TArray<FName> EnumNames;\r\n"));
 			for (int32 Index = 0; Index < Enum->NumEnums(); Index++)
 			{
@@ -5046,7 +5046,7 @@ UClass* ProcessParsedClass(bool bClassIsAnInterface, TArray<FHeaderProvider> &De
 		}
 
 		// Create new class.
-		ResultClass = new(InParent, *ClassNameStripped, Flags) UClass(FObjectInitializer(), nullptr);
+		ResultClass = new(EC_InternalUseOnlyConstructor, InParent, *ClassNameStripped, Flags) UClass(FObjectInitializer(), nullptr);
 		GClassHeaderNameWithNoPathMap.Add(ResultClass, ClassNameStripped);
 
 		// add CLASS_Interface flag if the class is an interface
@@ -5077,7 +5077,7 @@ UClass* ProcessParsedClass(bool bClassIsAnInterface, TArray<FHeaderProvider> &De
 			if (ResultClass->GetSuperStruct() == nullptr)
 			{
 				// don't know its parent class yet
-				ResultClass->SetSuperStruct(new(InParent, *BaseClassNameStripped) UClass(FObjectInitializer(), nullptr));
+				ResultClass->SetSuperStruct(new(EC_InternalUseOnlyConstructor, InParent, *BaseClassNameStripped) UClass(FObjectInitializer(), nullptr));
 			}
 
 			if (ResultClass->GetSuperStruct() != nullptr)
diff --git a/Engine/Source/Programs/UnrealHeaderTool/Private/HeaderParser.cpp b/Engine/Source/Programs/UnrealHeaderTool/Private/HeaderParser.cpp
--- a/Engine/Source/Programs/UnrealHeaderTool/Private/HeaderParser.cpp
+++ b/Engine/Source/Programs/UnrealHeaderTool/Private/HeaderParser.cpp
@@ -645,7 +645,7 @@ UEnum* FHeaderParser::CompileEnum(FUnrealSourceFile& SourceFile)
 	ParseFieldMetaData(EnumToken.MetaData, EnumToken.Identifier);
 
 	// Create enum definition.
-	UEnum* Enum = new(SourceFile.GetPackage(), EnumToken.Identifier, RF_Public) UEnum(FObjectInitializer());
+	UEnum* Enum = new(EC_InternalUseOnlyConstructor, SourceFile.GetPackage(), EnumToken.Identifier, RF_Public) UEnum(FObjectInitializer());
 	Scope->AddType(Enum);
 
 	// Validate the metadata for the enum
@@ -1339,7 +1339,7 @@ UScriptStruct* FHeaderParser::CompileStructDeclaration(FClasses& AllClasses, FUn
 	}
 
 	// Create.
-	UScriptStruct* Struct = new(SourceFile.GetPackage(), *EffectiveStructName, RF_Public) UScriptStruct(FObjectInitializer(), BaseStruct);
+	UScriptStruct* Struct = new(EC_InternalUseOnlyConstructor, SourceFile.GetPackage(), *EffectiveStructName, RF_Public) UScriptStruct(FObjectInitializer(), BaseStruct);
 
 	AddModuleRelativePathToMetadata(Struct, MetaData);
 
@@ -3371,7 +3371,7 @@ UProperty* FHeaderParser::GetVarNameAndDim
 		int32      ArrayDim = 1; // 1 = not a static array, 2 = static array
 		if (VarProperty.ArrayType == EArrayType::Dynamic)
 		{
-			Array       = new(Scope,PropertyName,ObjectFlags)UArrayProperty(FObjectInitializer());
+			Array = new(EC_InternalUseOnlyConstructor, Scope, PropertyName, ObjectFlags)UArrayProperty(FObjectInitializer());
 			NewScope    = Array;
 			ObjectFlags = RF_Public;
 		}
@@ -3382,40 +3382,40 @@ UProperty* FHeaderParser::GetVarNameAndDim
 
 		if (VarProperty.Type == CPT_Byte)
 		{
-			NewProperty = new(NewScope,PropertyName,ObjectFlags)UByteProperty(FObjectInitializer());
+			NewProperty = new(EC_InternalUseOnlyConstructor, NewScope, PropertyName, ObjectFlags)UByteProperty(FObjectInitializer());
 			((UByteProperty*)NewProperty)->Enum = VarProperty.Enum;
 		}
 		else if (VarProperty.Type == CPT_Int8)
 		{
-			NewProperty = new(NewScope,PropertyName,ObjectFlags)UInt8Property(FObjectInitializer());
+			NewProperty = new(EC_InternalUseOnlyConstructor, NewScope, PropertyName, ObjectFlags)UInt8Property(FObjectInitializer());
 		}
 		else if (VarProperty.Type == CPT_Int16)
 		{
-			NewProperty = new(NewScope,PropertyName,ObjectFlags)UInt16Property(FObjectInitializer());
+			NewProperty = new(EC_InternalUseOnlyConstructor, NewScope, PropertyName, ObjectFlags)UInt16Property(FObjectInitializer());
 		}
 		else if (VarProperty.Type == CPT_Int)
 		{
-			NewProperty = new(NewScope,PropertyName,ObjectFlags)UIntProperty(FObjectInitializer());
+			NewProperty = new(EC_InternalUseOnlyConstructor, NewScope, PropertyName, ObjectFlags)UIntProperty(FObjectInitializer());
 		}
 		else if (VarProperty.Type == CPT_Int64)
 		{
-			NewProperty = new(NewScope,PropertyName,ObjectFlags)UInt64Property(FObjectInitializer());
+			NewProperty = new(EC_InternalUseOnlyConstructor, NewScope, PropertyName, ObjectFlags)UInt64Property(FObjectInitializer());
 		}
 		else if (VarProperty.Type == CPT_UInt16)
 		{
-			NewProperty = new(NewScope,PropertyName,ObjectFlags)UUInt16Property(FObjectInitializer());
+			NewProperty = new(EC_InternalUseOnlyConstructor, NewScope, PropertyName, ObjectFlags)UUInt16Property(FObjectInitializer());
 		}
 		else if (VarProperty.Type == CPT_UInt32)
 		{
-			NewProperty = new(NewScope,PropertyName,ObjectFlags)UUInt32Property(FObjectInitializer());
+			NewProperty = new(EC_InternalUseOnlyConstructor, NewScope, PropertyName, ObjectFlags)UUInt32Property(FObjectInitializer());
 		}
 		else if (VarProperty.Type == CPT_UInt64)
 		{
-			NewProperty = new(NewScope,PropertyName,ObjectFlags)UUInt64Property(FObjectInitializer());
+			NewProperty = new(EC_InternalUseOnlyConstructor, NewScope, PropertyName, ObjectFlags)UUInt64Property(FObjectInitializer());
 		}
 		else if (VarProperty.IsBool())
 		{
-			UBoolProperty* NewBoolProperty = new(NewScope,PropertyName,ObjectFlags)UBoolProperty(FObjectInitializer());
+			UBoolProperty* NewBoolProperty = new(EC_InternalUseOnlyConstructor, NewScope, PropertyName, ObjectFlags)UBoolProperty(FObjectInitializer());
 			NewProperty = NewBoolProperty;
 			if (HardcodedName && FCString::Stricmp(HardcodedName, TEXT("ReturnValue")) == 0)
 			{
@@ -3445,18 +3445,18 @@ UProperty* FHeaderParser::GetVarNameAndDim
 		}
 		else if (VarProperty.Type == CPT_Float)
 		{
-			NewProperty = new(NewScope,PropertyName,ObjectFlags)UFloatProperty(FObjectInitializer());
+			NewProperty = new(EC_InternalUseOnlyConstructor, NewScope, PropertyName, ObjectFlags)UFloatProperty(FObjectInitializer());
 		}
 		else if (VarProperty.Type == CPT_Double)
 		{
-			NewProperty = new(NewScope,PropertyName,ObjectFlags)UDoubleProperty(FObjectInitializer());
+			NewProperty = new(EC_InternalUseOnlyConstructor, NewScope, PropertyName, ObjectFlags)UDoubleProperty(FObjectInitializer());
 		}
 		else if (VarProperty.Type == CPT_ObjectReference)
 		{
 			check(VarProperty.PropertyClass);
 			if (VarProperty.PropertyClass->IsChildOf(UClass::StaticClass()))
 			{
-				NewProperty = new(NewScope,PropertyName,ObjectFlags)UClassProperty(FObjectInitializer());
+				NewProperty = new(EC_InternalUseOnlyConstructor, NewScope, PropertyName, ObjectFlags)UClassProperty(FObjectInitializer());
 				((UClassProperty*)NewProperty)->MetaClass = VarProperty.MetaClass;
 			}
 			else
@@ -3466,7 +3466,7 @@ UProperty* FHeaderParser::GetVarNameAndDim
 					VarProperty.PropertyFlags |= CPF_InstancedReference;
 					AddEditInlineMetaData(VarProperty.MetaData);
 				}
-				NewProperty = new(NewScope,PropertyName,ObjectFlags)UObjectProperty(FObjectInitializer());
+				NewProperty = new(EC_InternalUseOnlyConstructor, NewScope, PropertyName, ObjectFlags)UObjectProperty(FObjectInitializer());
 			}
 
 			((UObjectPropertyBase*)NewProperty)->PropertyClass = VarProperty.PropertyClass;
@@ -3474,13 +3474,13 @@ UProperty* FHeaderParser::GetVarNameAndDim
 		else if (VarProperty.Type == CPT_WeakObjectReference)
 		{
 			check(VarProperty.PropertyClass);
-			NewProperty = new(NewScope,PropertyName,ObjectFlags)UWeakObjectProperty(FObjectInitializer());
+			NewProperty = new(EC_InternalUseOnlyConstructor, NewScope, PropertyName, ObjectFlags)UWeakObjectProperty(FObjectInitializer());
 			((UObjectPropertyBase*)NewProperty)->PropertyClass = VarProperty.PropertyClass;
 		}
 		else if( VarProperty.Type == CPT_LazyObjectReference )
 		{
 			check(VarProperty.PropertyClass);
-			NewProperty = new(NewScope,PropertyName,ObjectFlags)ULazyObjectProperty(FObjectInitializer());
+			NewProperty = new(EC_InternalUseOnlyConstructor, NewScope, PropertyName, ObjectFlags)ULazyObjectProperty(FObjectInitializer());
 			((UObjectPropertyBase*)NewProperty)->PropertyClass = VarProperty.PropertyClass;
 		}
 		else if( VarProperty.Type == CPT_AssetObjectReference )
@@ -3488,12 +3488,12 @@ UProperty* FHeaderParser::GetVarNameAndDim
 			check(VarProperty.PropertyClass);
 			if (VarProperty.PropertyClass->IsChildOf(UClass::StaticClass()))
 			{
-				NewProperty = new(NewScope,PropertyName,ObjectFlags)UAssetClassProperty(FObjectInitializer());
+				NewProperty = new(EC_InternalUseOnlyConstructor, NewScope, PropertyName, ObjectFlags)UAssetClassProperty(FObjectInitializer());
 				((UAssetClassProperty*)NewProperty)->MetaClass = VarProperty.MetaClass;
 			}
 			else
 			{
-				NewProperty = new(NewScope,PropertyName,ObjectFlags)UAssetObjectProperty(FObjectInitializer());
+				NewProperty = new(EC_InternalUseOnlyConstructor, NewScope, PropertyName, ObjectFlags)UAssetObjectProperty(FObjectInitializer());
 			}
 			((UObjectPropertyBase*)NewProperty)->PropertyClass = VarProperty.PropertyClass;
 		}
@@ -3502,26 +3502,26 @@ UProperty* FHeaderParser::GetVarNameAndDim
 			check(VarProperty.PropertyClass);
 			check(VarProperty.PropertyClass->HasAnyClassFlags(CLASS_Interface));
 
-			UInterfaceProperty* InterfaceProperty = new(NewScope,PropertyName,ObjectFlags) UInterfaceProperty(FObjectInitializer());
+			UInterfaceProperty* InterfaceProperty = new(EC_InternalUseOnlyConstructor, NewScope, PropertyName, ObjectFlags) UInterfaceProperty(FObjectInitializer());
 			InterfaceProperty->InterfaceClass = VarProperty.PropertyClass;
 
 			NewProperty = InterfaceProperty;
 		}
 		else if (VarProperty.Type == CPT_Name)
 		{
-			NewProperty = new(NewScope,PropertyName,ObjectFlags)UNameProperty(FObjectInitializer());
+			NewProperty = new(EC_InternalUseOnlyConstructor, NewScope, PropertyName, ObjectFlags)UNameProperty(FObjectInitializer());
 		}
 		else if (VarProperty.Type == CPT_String)
 		{
-			NewProperty = new(NewScope,PropertyName,ObjectFlags)UStrProperty(FObjectInitializer());
+			NewProperty = new(EC_InternalUseOnlyConstructor, NewScope, PropertyName, ObjectFlags)UStrProperty(FObjectInitializer());
 		}
 		else if (VarProperty.Type == CPT_Text)
 		{
-			NewProperty = new(NewScope,PropertyName,ObjectFlags)UTextProperty(FObjectInitializer());
+			NewProperty = new(EC_InternalUseOnlyConstructor, NewScope, PropertyName, ObjectFlags)UTextProperty(FObjectInitializer());
 		}
 		else if (VarProperty.Type == CPT_Struct)
 		{
-			NewProperty = new(NewScope,PropertyName,ObjectFlags)UStructProperty(FObjectInitializer());
+			NewProperty = new(EC_InternalUseOnlyConstructor, NewScope, PropertyName, ObjectFlags)UStructProperty(FObjectInitializer());
 			if (VarProperty.Struct->StructFlags & STRUCT_HasInstancedReference)
 			{
 				VarProperty.PropertyFlags |= CPF_ContainsInstancedReference;
@@ -3530,11 +3530,11 @@ UProperty* FHeaderParser::GetVarNameAndDim
 		}
 		else if (VarProperty.Type == CPT_Delegate)
 		{
-			NewProperty = new(NewScope,PropertyName,ObjectFlags)UDelegateProperty(FObjectInitializer());
+			NewProperty = new(EC_InternalUseOnlyConstructor, NewScope, PropertyName, ObjectFlags)UDelegateProperty(FObjectInitializer());
 		}
 		else if (VarProperty.Type == CPT_MulticastDelegate)
 		{
-			NewProperty = new(NewScope,PropertyName,ObjectFlags)UMulticastDelegateProperty(FObjectInitializer());
+			NewProperty = new(EC_InternalUseOnlyConstructor, NewScope, PropertyName, ObjectFlags)UMulticastDelegateProperty(FObjectInitializer());
 		}
 		else
 		{
@@ -7671,7 +7671,7 @@ TFunctionType* CreateFunctionImpl(const FFuncInfo& FuncInfo, UObject* Outer, FSc
 		}
 	}
 
-	TFunctionType* Function = new(Outer, FuncInfo.Function.Identifier, RF_Public) TFunctionType(FObjectInitializer(), nullptr);
+	TFunctionType* Function = new(EC_InternalUseOnlyConstructor, Outer, FuncInfo.Function.Identifier, RF_Public) TFunctionType(FObjectInitializer(), nullptr);
 	Function->RepOffset = MAX_uint16;
 	Function->ReturnValueOffset = MAX_uint16;
 	Function->FirstPropertyToInit = nullptr;
diff --git a/Engine/Source/Runtime/CoreUObject/Private/UObject/LinkerLoad.cpp b/Engine/Source/Runtime/CoreUObject/Private/UObject/LinkerLoad.cpp
--- a/Engine/Source/Runtime/CoreUObject/Private/UObject/LinkerLoad.cpp
+++ b/Engine/Source/Runtime/CoreUObject/Private/UObject/LinkerLoad.cpp
@@ -492,7 +492,7 @@ ULinkerLoad* ULinkerLoad::CreateLinkerAsync( UPackage* Parent, const TCHAR* File
 		{
 			LoadFlags |= LOAD_SeekFree;
 		}
-		Linker = new ULinkerLoad( FObjectInitializer(), Parent, Filename, LoadFlags );
+		Linker = new(EC_InternalUseOnlyConstructor) ULinkerLoad(FObjectInitializer(), Parent, Filename, LoadFlags);
 		// Add to the list of linkers that haven't been finalized yet
 		GObjPendingLoaders.Add(Parent, Linker);
 	}
diff --git a/Engine/Source/Runtime/CoreUObject/Private/UObject/Obj.cpp b/Engine/Source/Runtime/CoreUObject/Private/UObject/Obj.cpp
--- a/Engine/Source/Runtime/CoreUObject/Private/UObject/Obj.cpp
+++ b/Engine/Source/Runtime/CoreUObject/Private/UObject/Obj.cpp
@@ -3567,7 +3567,7 @@ void StaticUObjectInit()
 	UObjectBaseInit();
 
 	// Allocate special packages.
-	GObjTransientPkg = new( NULL, TEXT("/Engine/Transient") )UPackage(FObjectInitializer());
+	GObjTransientPkg = NewNamedObject<UPackage>(nullptr, TEXT("/Engine/Transient"));
 	GObjTransientPkg->AddToRoot();
 
 	if( FParse::Param( FCommandLine::Get(), TEXT("VERIFYGC") ) )
diff --git a/Engine/Source/Runtime/CoreUObject/Private/UObject/SavePackage.cpp b/Engine/Source/Runtime/CoreUObject/Private/UObject/SavePackage.cpp
--- a/Engine/Source/Runtime/CoreUObject/Private/UObject/SavePackage.cpp
+++ b/Engine/Source/Runtime/CoreUObject/Private/UObject/SavePackage.cpp
@@ -2610,12 +2610,12 @@ bool UPackage::SavePackage( UPackage* InOuter, UObject* Base, EObjectFlags TopLe
 				if (bCompressFromMemory || bSaveAsync)
 				{
 					// Allocate the linker with a memory writer, forcing byte swapping if wanted.
-					Linker = new ULinkerSave( FObjectInitializer(), InOuter, bForceByteSwapping, bSaveUnversioned );
+					Linker = new(EC_InternalUseOnlyConstructor) ULinkerSave(FObjectInitializer(), InOuter, bForceByteSwapping, bSaveUnversioned);
 				}
 				else
 				{
 					// Allocate the linker, forcing byte swapping if wanted.
-					Linker = new ULinkerSave( FObjectInitializer(), InOuter, *TempFilename, bForceByteSwapping, bSaveUnversioned );
+					Linker = new(EC_InternalUseOnlyConstructor) ULinkerSave(FObjectInitializer(), InOuter, *TempFilename, bForceByteSwapping, bSaveUnversioned);
 				}
 
 				// Use the custom versions we had previously gleaned from the export tag pass
diff --git a/Engine/Source/Runtime/CoreUObject/Private/UObject/TlsObjectInitializers.cpp b/Engine/Source/Runtime/CoreUObject/Private/UObject/TlsObjectInitializers.cpp
--- a/Engine/Source/Runtime/CoreUObject/Private/UObject/TlsObjectInitializers.cpp
+++ b/Engine/Source/Runtime/CoreUObject/Private/UObject/TlsObjectInitializers.cpp
@@ -2,6 +2,10 @@
 
 #include "CoreUObjectPrivate.h"
 #include "UObject/TlsObjectInitializers.h"
+#include "AssertionMacros.h"
+
+DECLARE_LOG_CATEGORY_EXTERN(LogTlsObjectInitializers, Log, All);
+DEFINE_LOG_CATEGORY(LogTlsObjectInitializers);
 
 TArray<FObjectInitializer*>* FTlsObjectInitializers::GetTlsObjectInitializers()
 {
@@ -21,7 +25,7 @@ TArray<FObjectInitializer*>* FTlsObjectInitializers::GetTlsObjectInitializers()
 FObjectInitializer* FTlsObjectInitializers::Top()
 {
 	auto ObjectInitializers = GetTlsObjectInitializers();
-	return ObjectInitializers->Last();
+	return ObjectInitializers->Num() ? ObjectInitializers->Last() : nullptr;
 }
 
 void FTlsObjectInitializers::Pop()
@@ -34,4 +38,11 @@ void FTlsObjectInitializers::Push(FObjectInitializer* Initializer)
 {
 	auto ObjectInitializers = GetTlsObjectInitializers();
 	ObjectInitializers->Push(Initializer);
-}
\ No newline at end of file
+}
+
+FObjectInitializer& FTlsObjectInitializers::TopChecked()
+{
+	auto* ObjectInitializerPtr = Top();
+	UE_CLOG(!ObjectInitializerPtr, LogTlsObjectInitializers, Fatal, TEXT("Tried to get the current ObjectInitializer, but none is set. Please use NewObject or NewNamedObject to construct new UObject-derived classes."));
+	return *ObjectInitializerPtr;
+}
diff --git a/Engine/Source/Runtime/CoreUObject/Private/UObject/UObjectGlobals.cpp b/Engine/Source/Runtime/CoreUObject/Private/UObject/UObjectGlobals.cpp
--- a/Engine/Source/Runtime/CoreUObject/Private/UObject/UObjectGlobals.cpp
+++ b/Engine/Source/Runtime/CoreUObject/Private/UObject/UObjectGlobals.cpp
@@ -504,7 +504,7 @@ UPackage* CreatePackage( UObject* InOuter, const TCHAR* PackageName )
 			}
 			else
 			{
-				Result = new( InOuter, NewPackageName, RF_Public )UPackage(FObjectInitializer());
+				Result = NewNamedObject<UPackage>(InOuter, NewPackageName, RF_Public);
 			}
 		}
 	}
@@ -1890,6 +1890,16 @@ void UObject::PostInitProperties()
 #endif
 }
 
+UObject::UObject()
+{
+	FObjectInitializer* ObjectInitializerPtr = FTlsObjectInitializers::Top();
+	UE_CLOG(!ObjectInitializerPtr, LogUObjectGlobals, Fatal, TEXT("%s is not being constructed with either NewObject, NewNamedObject or ConstructObject."), *GetName());
+	FObjectInitializer& ObjectInitializer = *ObjectInitializerPtr;
+	check(!ObjectInitializer.Obj || ObjectInitializer.Obj == this);
+	const_cast<FObjectInitializer&>(ObjectInitializer).Obj = this;
+	const_cast<FObjectInitializer&>(ObjectInitializer).FinalizeSubobjectClassInitialization();
+}
+
 UObject::UObject(const FObjectInitializer& ObjectInitializer)
 {
 	check(!ObjectInitializer.Obj || ObjectInitializer.Obj == this);
@@ -2259,6 +2269,12 @@ void FObjectInitializer::AssertIfInConstructor(UObject* Outer, const TCHAR* Erro
 	UE_CLOG(GIsInConstructor && Outer == GConstructedObject, LogUObjectGlobals, Fatal, TEXT("%s"), ErrorMessage);
 }
 
+FObjectInitializer& FObjectInitializer::Get()
+{
+	UE_CLOG(!GIsInConstructor, LogUObjectGlobals, Fatal, TEXT("FObjectInitializer::Get() can only be used inside of UObject-derived class constructor."));
+	return FTlsObjectInitializers::TopChecked();
+}
+
 /**
  * Stores the object flags for all objects in the tracking array.
  */
diff --git a/Engine/Source/Runtime/CoreUObject/Public/UObject/GCObject.h b/Engine/Source/Runtime/CoreUObject/Public/UObject/GCObject.h
--- a/Engine/Source/Runtime/CoreUObject/Public/UObject/GCObject.h
+++ b/Engine/Source/Runtime/CoreUObject/Public/UObject/GCObject.h
@@ -80,7 +80,7 @@ public:
 	{
 		if (GGCObjectReferencer == NULL)
 		{
-			GGCObjectReferencer = new UGCObjectReferencer(FObjectInitializer());
+			GGCObjectReferencer = NewObject<UGCObjectReferencer>();
 			GGCObjectReferencer->AddToRoot();
 		}
 	}
diff --git a/Engine/Source/Runtime/CoreUObject/Public/UObject/ObjectBase.h b/Engine/Source/Runtime/CoreUObject/Public/UObject/ObjectBase.h
--- a/Engine/Source/Runtime/CoreUObject/Public/UObject/ObjectBase.h
+++ b/Engine/Source/Runtime/CoreUObject/Public/UObject/ObjectBase.h
@@ -95,7 +95,7 @@ enum EPackageFlags
 // Internal enums.
 //
 enum EStaticConstructor				{EC_StaticConstructor};
-enum EInternal						{EC_Internal};
+enum EInternal						{EC_InternalUseOnlyConstructor};
 enum ECppProperty					{EC_CppProperty};
 
 /** Empty API definition.  Used as a placeholder parameter when no DLL export/import API is needed for a UObject class */
@@ -1178,11 +1178,16 @@ public: \
 	{ \
 		return TStaticCastFlags; \
 	} \
-	/** For internal use only; use StaticConstructObject() to create new objects. */ \
+	DEPRECATED(4.8, "operator new has been deprecated for UObjects - please use NewObject or NewNamedObject instead") \
 	inline void* operator new( const size_t InSize, UObject* InOuter=(UObject*)GetTransientPackage(), FName InName=NAME_None, EObjectFlags InSetFlags=RF_NoFlags ) \
 	{ \
 		return StaticAllocateObject( StaticClass(), InOuter, InName, InSetFlags ); \
 	} \
+	/** For internal use only; use StaticConstructObject() to create new objects. */ \
+	inline void* operator new(const size_t InSize, EInternal InInternalOnly, UObject* InOuter = (UObject*)GetTransientPackage(), FName InName = NAME_None, EObjectFlags InSetFlags = RF_NoFlags) \
+	{ \
+		return StaticAllocateObject(StaticClass(), InOuter, InName, InSetFlags); \
+} \
 	/** For internal use only; use StaticConstructObject() to create new objects. */ \
 	inline void* operator new( const size_t InSize, EInternal* InMem ) \
 	{ \
@@ -1193,11 +1198,18 @@ public: \
 	static_assert(false, "You have to define " #TClass "::" #TClass "() or " #TClass "::" #TClass "(const FObjectInitializer&). This is required by UObject system to work correctly.");
 
 #define DEFINE_DEFAULT_CONSTRUCTOR_CALL(TClass) \
-	static void __DefaultConstructor(const FObjectInitializer& X) { const_cast<FObjectInitializer&>(X).FinalizeSubobjectClassInitialization(); new((EInternal*)X.GetObj())TClass(); }
+	static void __DefaultConstructor(const FObjectInitializer& X) { new((EInternal*)X.GetObj())TClass(); }
 
 #define DEFINE_DEFAULT_OBJECT_INITIALIZER_CONSTRUCTOR_CALL(TClass) \
 	static void __DefaultConstructor(const FObjectInitializer& X) { new((EInternal*)X.GetObj())TClass(X); }
 
+#define DECLARE_CLASS_INTRINSIC_NO_CTOR(TClass,TSuperClass,TStaticFlags,TPackage) \
+	DECLARE_CLASS(TClass, TSuperClass, TStaticFlags | CLASS_Intrinsic, CASTCLASS_None, TPackage, NO_API) \
+	enum { IsIntrinsic = 1 }; \
+	static void StaticRegisterNatives##TClass() {} \
+	DECLARE_SERIALIZER(TClass) \
+	DEFINE_DEFAULT_OBJECT_INITIALIZER_CONSTRUCTOR_CALL(TClass)
+
 #define DECLARE_CLASS_INTRINSIC(TClass,TSuperClass,TStaticFlags,TPackage) \
 	DECLARE_CLASS(TClass,TSuperClass,TStaticFlags|CLASS_Intrinsic,CASTCLASS_None,TPackage,NO_API ) \
 	RELAY_CONSTRUCTOR(TClass, TSuperClass) \
@@ -1206,6 +1218,13 @@ public: \
 	DECLARE_SERIALIZER(TClass) \
 	DEFINE_DEFAULT_OBJECT_INITIALIZER_CONSTRUCTOR_CALL(TClass)
 
+#define DECLARE_CASTED_CLASS_INTRINSIC_WITH_API_NO_CTOR( TClass, TSuperClass, TStaticFlags, TPackage, TStaticCastFlags, TRequiredAPI ) \
+	DECLARE_CLASS(TClass, TSuperClass, TStaticFlags | CLASS_Intrinsic, TStaticCastFlags, TPackage, TRequiredAPI) \
+	enum { IsIntrinsic = 1 }; \
+	static void StaticRegisterNatives##TClass() {} \
+	DECLARE_SERIALIZER(TClass) \
+	DEFINE_DEFAULT_OBJECT_INITIALIZER_CONSTRUCTOR_CALL(TClass)
+
 #define DECLARE_CASTED_CLASS_INTRINSIC_WITH_API( TClass, TSuperClass, TStaticFlags, TPackage, TStaticCastFlags, TRequiredAPI ) \
 	DECLARE_CLASS(TClass,TSuperClass,TStaticFlags|CLASS_Intrinsic,TStaticCastFlags,TPackage,TRequiredAPI ) \
 	RELAY_CONSTRUCTOR(TClass, TSuperClass) \
diff --git a/Engine/Source/Runtime/CoreUObject/Public/UObject/TlsObjectInitializers.h b/Engine/Source/Runtime/CoreUObject/Public/UObject/TlsObjectInitializers.h
--- a/Engine/Source/Runtime/CoreUObject/Public/UObject/TlsObjectInitializers.h
+++ b/Engine/Source/Runtime/CoreUObject/Public/UObject/TlsObjectInitializers.h
@@ -28,6 +28,12 @@ public:
 	*/
 	static FObjectInitializer* Top();
 
+	/**
+	 * Retrieves current FObjectInitializer for current thread. Will assert of no ObjectInitializer is currently set.
+	 * @return Current FObjectInitializer reference.
+	 */
+	static FObjectInitializer& TopChecked();
+
 private:
 	/**
 	* Retrieve thread local stack of FObjectInitializers.
diff --git a/Engine/Source/Runtime/CoreUObject/Public/UObject/UObject.h b/Engine/Source/Runtime/CoreUObject/Public/UObject/UObject.h
--- a/Engine/Source/Runtime/CoreUObject/Public/UObject/UObject.h
+++ b/Engine/Source/Runtime/CoreUObject/Public/UObject/UObject.h
@@ -43,7 +43,7 @@ class COREUOBJECT_API UObject : public UObjectBaseUtility
 	static const TCHAR* StaticConfigName() {return TEXT("Engine");}
 
 	// Constructors and destructors.
-	UObject() { };
+	UObject();
 	UObject(const FObjectInitializer& ObjectInitializer);
 	UObject( EStaticConstructor, EObjectFlags InFlags );
 
diff --git a/Engine/Source/Runtime/CoreUObject/Public/UObject/UObjectGlobals.h b/Engine/Source/Runtime/CoreUObject/Public/UObject/UObjectGlobals.h
--- a/Engine/Source/Runtime/CoreUObject/Public/UObject/UObjectGlobals.h
+++ b/Engine/Source/Runtime/CoreUObject/Public/UObject/UObjectGlobals.h
@@ -801,6 +801,9 @@ public:
 		bSubobjectClassInitializationAllowed = false;
 	}
 
+	/** Gets ObjectInitializer for the currently constructed object. Can only be used inside of a constructor of UObject-derived class. */
+	static FObjectInitializer& Get();
+
 private:
 
 	friend class UObject; 
diff --git a/Engine/Source/Runtime/Engine/Classes/Components/ActorComponent.h b/Engine/Source/Runtime/Engine/Classes/Components/ActorComponent.h
--- a/Engine/Source/Runtime/Engine/Classes/Components/ActorComponent.h
+++ b/Engine/Source/Runtime/Engine/Classes/Components/ActorComponent.h
@@ -31,18 +31,9 @@ class ENGINE_API UActorComponent : public UObject, public IInterface_AssetUserDa
 public:
 
 	/**
-	 * Default UObject constructor.
+	 * Default UObject constructor that takes an optional ObjectInitializer.
 	 */
-	UActorComponent();
-
-	/**
-	 * UObject constructor that takes an ObjectInitializer
-	 */
-	UActorComponent(const FObjectInitializer& ObjectInitializer);
-
-private:
-	/** Called from the constructor to initialize the object to its default settings */
-	void InitializeDefaults();
+	UActorComponent(const FObjectInitializer& ObjectInitializer = FObjectInitializer::Get());
 
 public:
 	virtual void GetLifetimeReplicatedProps(TArray<FLifetimeProperty>& OutLifetimeProps) const override;
diff --git a/Engine/Source/Runtime/Engine/Classes/Components/ModelComponent.h b/Engine/Source/Runtime/Engine/Classes/Components/ModelComponent.h
--- a/Engine/Source/Runtime/Engine/Classes/Components/ModelComponent.h
+++ b/Engine/Source/Runtime/Engine/Classes/Components/ModelComponent.h
@@ -18,7 +18,7 @@ class FModelElement;
 UCLASS(MinimalAPI)
 class UModelComponent : public UPrimitiveComponent, public IInterface_CollisionDataProvider
 {
-	GENERATED_BODY()
+	GENERATED_UCLASS_BODY()
 
 private:
 	/** The BSP tree. */
@@ -42,14 +42,11 @@ private:
 
 public:
 
-	/** Default constructor */
-	UModelComponent();
-
 #if WITH_EDITOR
 	/**
-	 * Minimal initialization constructor.
+	 * Minimal initialization.
 	 */
-	UModelComponent( UModel* InModel, uint16 InComponentIndex, uint32 MaskedSurfaceFlags, const TArray<uint16>& InNodes );
+	void InitializeModelComponent(UModel* InModel, uint16 InComponentIndex, uint32 MaskedSurfaceFlags, const TArray<uint16>& InNodes);
 #endif // WITH_EDITOR
 
 	/**
diff --git a/Engine/Source/Runtime/Engine/Classes/Components/PrimitiveComponent.h b/Engine/Source/Runtime/Engine/Classes/Components/PrimitiveComponent.h
--- a/Engine/Source/Runtime/Engine/Classes/Components/PrimitiveComponent.h
+++ b/Engine/Source/Runtime/Engine/Classes/Components/PrimitiveComponent.h
@@ -109,15 +109,8 @@ public:
 	/**
 	 * Default UObject constructor.
 	 */
-	UPrimitiveComponent();
+	UPrimitiveComponent(const FObjectInitializer& ObjectInitializer = FObjectInitializer::Get());
 
-	UPrimitiveComponent(const FObjectInitializer& ObjectInitializer);
-
-private:
-	/** Initialize the component to its default settings */
-	void InitializePrimitiveComponentDefaults();
-
-public:
 	// Rendering
 	
 	/**
diff --git a/Engine/Source/Runtime/Engine/Classes/Components/SceneComponent.h b/Engine/Source/Runtime/Engine/Classes/Components/SceneComponent.h
--- a/Engine/Source/Runtime/Engine/Classes/Components/SceneComponent.h
+++ b/Engine/Source/Runtime/Engine/Classes/Components/SceneComponent.h
@@ -82,19 +82,10 @@ class ENGINE_API USceneComponent : public UActorComponent
 	GENERATED_BODY()
 public:
 
-	/**
-	 * Default UObject constructor.
-	 */
-	USceneComponent();
-
 	/**
 	 * UObject constructor that takes an ObjectInitializer
 	 */
-	USceneComponent(const FObjectInitializer& ObjectInitializer);
-
-private:
-	/** Initialize the component to its default settings */
-	void InitializeDefaults();
+	USceneComponent(const FObjectInitializer& ObjectInitializer = FObjectInitializer::Get());
 
 public:
 
diff --git a/Engine/Source/Runtime/Engine/Classes/Engine/Level.h b/Engine/Source/Runtime/Engine/Classes/Engine/Level.h
--- a/Engine/Source/Runtime/Engine/Classes/Engine/Level.h
+++ b/Engine/Source/Runtime/Engine/Classes/Engine/Level.h
@@ -489,8 +489,8 @@ public:
 	ENGINE_API static FSimpleMulticastDelegate LevelDirtiedEvent;
 
 	// Constructor.
-	ENGINE_API ULevel(const FObjectInitializer& ObjectInitializer, const FURL& InURL );
-	ULevel(const FObjectInitializer& ObjectInitializer );
+	ENGINE_API void Initialize(const FURL& InURL);
+	ULevel(const FObjectInitializer& ObjectInitializer);
 	~ULevel();
 
 	// Begin UObject interface.
diff --git a/Engine/Source/Runtime/Engine/Classes/Engine/PackageMapClient.h b/Engine/Source/Runtime/Engine/Classes/Engine/PackageMapClient.h
--- a/Engine/Source/Runtime/Engine/Classes/Engine/PackageMapClient.h
+++ b/Engine/Source/Runtime/Engine/Classes/Engine/PackageMapClient.h
@@ -90,14 +90,16 @@ public:
 UCLASS(transient)
 class UPackageMapClient : public UPackageMap
 {
-    GENERATED_UCLASS_BODY()
+public:
+	GENERATED_BODY()
+
+	UPackageMapClient(const FObjectInitializer & ObjectInitializer);
 
-	UPackageMapClient( const FObjectInitializer & ObjectInitializer, UNetConnection * InConnection, TSharedPtr< FNetGUIDCache > InNetGUIDCache ) : 
-		UPackageMap( ObjectInitializer )
-	,	Connection(InConnection)
+	void Initialize(UNetConnection * InConnection, TSharedPtr<FNetGUIDCache> InNetGUIDCache)
 	{
-		GuidCache				= InNetGUIDCache;
-		ExportNetGUIDCount		= 0;
+		Connection = InConnection;
+		GuidCache = InNetGUIDCache;
+		ExportNetGUIDCount = 0;
 	}
 
 	virtual ~UPackageMapClient()
diff --git a/Engine/Source/Runtime/Engine/Classes/Engine/PendingNetGame.h b/Engine/Source/Runtime/Engine/Classes/Engine/PendingNetGame.h
--- a/Engine/Source/Runtime/Engine/Classes/Engine/PendingNetGame.h
+++ b/Engine/Source/Runtime/Engine/Classes/Engine/PendingNetGame.h
@@ -96,7 +96,7 @@ public:
 	FString					ConnectionError;
 
 	// Constructor.
-	UPendingNetGame( const FObjectInitializer& ObjectInitializer, const FURL& InURL );
+	void Initialize(const FURL& InURL);
 
 	// Constructor.
 	UPendingNetGame( const FObjectInitializer& ObjectInitializer );
diff --git a/Engine/Source/Runtime/Engine/Classes/Engine/World.h b/Engine/Source/Runtime/Engine/Classes/Engine/World.h
--- a/Engine/Source/Runtime/Engine/Classes/Engine/World.h
+++ b/Engine/Source/Runtime/Engine/Classes/Engine/World.h
@@ -1007,14 +1007,6 @@ public:
 	 * UWorld default constructor
 	 */
 	UWorld( const FObjectInitializer& ObjectInitializer );
-
-	/**
-	 * UWorld constructor called at game startup and when creating a new world in the Editor.
-	 * Please note that this constructor does NOT get called when a world is loaded from disk.
-	 *
-	 * @param	InURL	URL associated with this world.
-	 */
-	UWorld( const FObjectInitializer& ObjectInitializer,const FURL& InURL );
 	
 	// LINE TRACE
 
diff --git a/Engine/Source/Runtime/Engine/Private/ActorComponent.cpp b/Engine/Source/Runtime/Engine/Private/ActorComponent.cpp
--- a/Engine/Source/Runtime/Engine/Private/ActorComponent.cpp
+++ b/Engine/Source/Runtime/Engine/Private/ActorComponent.cpp
@@ -98,20 +98,8 @@ FGlobalComponentReregisterContext::~FGlobalComponentReregisterContext()
 	ActiveGlobalReregisterContextCount--;
 }
 
-
-UActorComponent::UActorComponent()
-{
-	InitializeDefaults();
-}
-
-
-UActorComponent::UActorComponent( const FObjectInitializer& ObjectInitializer )
-{
-	// Forward to default constructor (we don't use ObjectInitializer for anything, this is for compatibility with inherited classes that call Super( ObjectInitializer )
-	InitializeDefaults();
-}
-
-void UActorComponent::InitializeDefaults()
+UActorComponent::UActorComponent(const FObjectInitializer& ObjectInitializer /*= FObjectInitializer::Get()*/)
+	: Super(ObjectInitializer)
 {
 	PrimaryComponentTick.TickGroup = TG_DuringPhysics;
 	PrimaryComponentTick.bStartWithTickEnabled = true;
diff --git a/Engine/Source/Runtime/Engine/Private/DemoNetDriver.cpp b/Engine/Source/Runtime/Engine/Private/DemoNetDriver.cpp
--- a/Engine/Source/Runtime/Engine/Private/DemoNetDriver.cpp
+++ b/Engine/Source/Runtime/Engine/Private/DemoNetDriver.cpp
@@ -212,7 +212,7 @@ bool UDemoNetDriver::InitConnect( FNetworkNotify* InNotify, const FURL& ConnectU
 	GetWorld()->DemoNetDriver = NULL;
 	SetWorld( NULL );
 
-	UDemoPendingNetGame * NewPendingNetGame = new UDemoPendingNetGame( FObjectInitializer() );
+	auto NewPendingNetGame = NewObject<UDemoPendingNetGame>();
 
 	NewPendingNetGame->DemoNetDriver = this;
 
diff --git a/Engine/Source/Runtime/Engine/Private/Font.cpp b/Engine/Source/Runtime/Engine/Private/Font.cpp
--- a/Engine/Source/Runtime/Engine/Private/Font.cpp
+++ b/Engine/Source/Runtime/Engine/Private/Font.cpp
@@ -46,7 +46,9 @@ void UFont::Serialize( FArchive& Ar )
 		{
 			if( InFontData.FontData_DEPRECATED.Num() > 0 )
 			{
-				InFontData.BulkDataPtr = new(this) UFontBulkData(InFontData.FontData_DEPRECATED.GetData(), InFontData.FontData_DEPRECATED.Num());
+				UFontBulkData* FontBulkData = NewObject<UFontBulkData>(this);
+				FontBulkData->Initialize(InFontData.FontData_DEPRECATED.GetData(), InFontData.FontData_DEPRECATED.Num());
+				InFontData.BulkDataPtr = FontBulkData;					
 				InFontData.FontData_DEPRECATED.Empty();
 			}
 		};
diff --git a/Engine/Source/Runtime/Engine/Private/Level.cpp b/Engine/Source/Runtime/Engine/Private/Level.cpp
--- a/Engine/Source/Runtime/Engine/Private/Level.cpp
+++ b/Engine/Source/Runtime/Engine/Private/Level.cpp
@@ -190,22 +190,17 @@ bool FLevelSimplificationDetails::operator == (const FLevelSimplificationDetails
 TMap<FName, TWeakObjectPtr<UWorld> > ULevel::StreamedLevelsOwningWorld;
 
 ULevel::ULevel( const FObjectInitializer& ObjectInitializer )
-	:   UObject( ObjectInitializer )
-	,   Actors(this)
-	,   OwningWorld(NULL)
-	,   TickTaskLevel(FTickTaskManagerInterface::Get().AllocateTickTaskLevel())
+	:	UObject( ObjectInitializer )
+	,	Actors(this)
+	,	OwningWorld(NULL)
+	,	TickTaskLevel(FTickTaskManagerInterface::Get().AllocateTickTaskLevel())
+	,	PrecomputedLightVolume(new FPrecomputedLightVolume())
 {
-	PrecomputedLightVolume = new FPrecomputedLightVolume();
 }
 
-ULevel::ULevel( const FObjectInitializer& ObjectInitializer,const FURL& InURL )
-	:   UObject( ObjectInitializer )
-	,   URL(InURL)
-	,   Actors(this)
-	,   OwningWorld(NULL)
-	,   TickTaskLevel(FTickTaskManagerInterface::Get().AllocateTickTaskLevel())
+void ULevel::Initialize(const FURL& InURL)
 {
-	PrecomputedLightVolume = new FPrecomputedLightVolume();
+	URL = InURL;
 }
 
 ULevel::~ULevel()
@@ -823,7 +818,8 @@ void ULevel::CreateModelComponents()
 			Model->Nodes[Nodes[NodeIndex]].ComponentNodeIndex = NodeIndex;
 		}
 
-		UModelComponent* ModelComponent = new(this) UModelComponent(Model,ModelComponents.Num(),Key.MaskedPolyFlags,Nodes);
+		UModelComponent* ModelComponent = NewObject<UModelComponent>(this);
+		ModelComponent->InitializeModelComponent(Model, ModelComponents.Num(), Key.MaskedPolyFlags, Nodes);
 		ModelComponents.Add(ModelComponent);
 
 		for(int32 NodeIndex = 0;NodeIndex < Nodes.Num();NodeIndex++)
diff --git a/Engine/Source/Runtime/Engine/Private/LightMap.cpp b/Engine/Source/Runtime/Engine/Private/LightMap.cpp
--- a/Engine/Source/Runtime/Engine/Private/LightMap.cpp
+++ b/Engine/Source/Runtime/Engine/Private/LightMap.cpp
@@ -815,7 +815,7 @@ void FLightMapPendingTexture::StartEncoding()
 
 	if (bNeedsSkyOcclusionTexture)
 	{
-		ULightMapTexture2D* Texture = new(Outer, GetSkyOcclusionTextureName(GLightmapCounter)) ULightMapTexture2D(FObjectInitializer());
+		auto Texture = NewNamedObject<ULightMapTexture2D>(Outer, GetSkyOcclusionTextureName(GLightmapCounter));
 		SkyOcclusionTexture = Texture;
 
 		Texture->Source.Init2DWithMipChain(GetSizeX(), GetSizeY(), TSF_BGRA8);
@@ -914,7 +914,7 @@ void FLightMapPendingTexture::StartEncoding()
 		}
 
 		// Create the light-map texture for this coefficient.
-		ULightMapTexture2D* Texture = new(Outer, GetLightmapName(GLightmapCounter, CoefficientIndex)) ULightMapTexture2D(FObjectInitializer());
+		auto Texture = NewNamedObject<ULightMapTexture2D>(Outer, GetLightmapName(GLightmapCounter, CoefficientIndex));
 		Textures[CoefficientIndex] = Texture;
 		Texture->Source.Init2DWithMipChain(GetSizeX(), GetSizeY() * 2, TSF_BGRA8);	// Top/bottom atlased
 		Texture->MipGenSettings = TMGS_LeaveExistingMips;
diff --git a/Engine/Source/Runtime/Engine/Private/Model.cpp b/Engine/Source/Runtime/Engine/Private/Model.cpp
--- a/Engine/Source/Runtime/Engine/Private/Model.cpp
+++ b/Engine/Source/Runtime/Engine/Private/Model.cpp
@@ -491,7 +491,7 @@ void UModel::EmptyModel( int32 EmptySurfInfo, int32 EmptyPolys )
 #if WITH_EDITOR
 	if( EmptyPolys )
 	{
-		Polys = new( GetOuter(), NAME_None, RF_Transactional )UPolys(FObjectInitializer());
+		Polys = NewNamedObject<UPolys>(GetOuter(), NAME_None, RF_Transactional);
 	}
 #endif // WITH_EDITOR
 
@@ -502,17 +502,22 @@ void UModel::EmptyModel( int32 EmptySurfInfo, int32 EmptyPolys )
 //
 // Create a new model and allocate all objects needed for it.
 //
-UModel::UModel( const FObjectInitializer& ObjectInitializer,ABrush* Owner, bool InRootOutside )
-:	UObject(ObjectInitializer)
-,	Nodes		( this )
-,	Verts		( this )
-,	Vectors		( this )
-,	Points		( this )
-,	Surfs		( this )
-,	VertexBuffer( this )
-,	LightingGuid( FGuid::NewGuid() )
-,	RootOutside	( InRootOutside )
+UModel::UModel(const FObjectInitializer& ObjectInitializer)
+	: UObject(ObjectInitializer)
+	, Nodes(this)
+	, Verts(this)
+	, Vectors(this)
+	, Points(this)
+	, Surfs(this)
+	, VertexBuffer(this)
 {
+
+}
+
+void UModel::Initialize(ABrush* Owner, bool InRootOutside)
+{
+	LightingGuid = FGuid::NewGuid();
+	RootOutside = InRootOutside;
 	SetFlags( RF_Transactional );
 	EmptyModel( 1, 1 );
 	if( Owner )
@@ -529,6 +534,23 @@ UModel::UModel( const FObjectInitializer& ObjectInitializer,ABrush* Owner, bool
 	}
 }
 
+void UModel::Initialize()
+{
+#if WITH_EDITOR
+	LightingLevel = nullptr;
+#endif // WITH_EDITOR
+	RootOutside = true;
+
+	if (!HasAnyFlags(RF_ClassDefaultObject))
+	{
+		EmptyModel(1, 0);
+		if (GIsEditor && !FApp::IsGame())
+		{
+			UpdateVertices();
+		}
+	}
+}
+
 #if WITH_EDITOR
 void UModel::BuildBound()
 {
diff --git a/Engine/Source/Runtime/Engine/Private/ModelComponent.cpp b/Engine/Source/Runtime/Engine/Private/ModelComponent.cpp
--- a/Engine/Source/Runtime/Engine/Private/ModelComponent.cpp
+++ b/Engine/Source/Runtime/Engine/Private/ModelComponent.cpp
@@ -51,7 +51,8 @@ FArchive& operator<<(FArchive& Ar,FModelElement& Element)
 	return Ar;
 }
 
-UModelComponent::UModelComponent()
+UModelComponent::UModelComponent(const FObjectInitializer& ObjectInitializer)
+	: Super(ObjectInitializer)
 {
 	CastShadow = true;
 	bUseAsOccluder = true;
@@ -62,13 +63,14 @@ UModelComponent::UModelComponent()
 }
 
 #if WITH_EDITOR
-UModelComponent::UModelComponent(UModel* InModel,uint16 InComponentIndex,uint32 MaskedSurfaceFlags,const TArray<uint16>& InNodes):
-	Model(InModel),
-	ComponentIndex(InComponentIndex),
-	Nodes(InNodes)
+void UModelComponent::InitializeModelComponent(UModel* InModel, uint16 InComponentIndex, uint32 MaskedSurfaceFlags, const TArray<uint16>& InNodes)
 {
+	Model = InModel;
+	ComponentIndex = InComponentIndex;
+	Nodes = InNodes;
+
 	// Model components are transacted.
-	SetFlags( RF_Transactional );
+	SetFlags(RF_Transactional);
 
 	GenerateElements(true);
 
diff --git a/Engine/Source/Runtime/Engine/Private/NetConnection.cpp b/Engine/Source/Runtime/Engine/Private/NetConnection.cpp
--- a/Engine/Source/Runtime/Engine/Private/NetConnection.cpp
+++ b/Engine/Source/Runtime/Engine/Private/NetConnection.cpp
@@ -113,7 +113,9 @@ void UNetConnection::InitBase(UNetDriver* InDriver,class FSocket* InSocket, cons
 	}
 
 	// Create package map.
-	PackageMap = new( this )UPackageMapClient( FObjectInitializer(), this, Driver->GuidCache );
+	auto PackageMapClient = NewObject<UPackageMapClient>(this);
+	PackageMapClient->Initialize(this, Driver->GuidCache);
+	PackageMap = PackageMapClient;
 
 	// Create the voice channel
 	CreateChannel(CHTYPE_Voice, true, VOICE_CHANNEL_INDEX);
@@ -160,7 +162,9 @@ void UNetConnection::InitConnection(UNetDriver* InDriver, EConnectionState InSta
 	}
 
 	// Create package map.
-	PackageMap = new( this )UPackageMapClient( FObjectInitializer(), this, Driver->GuidCache );
+	auto PackageMapClient = NewObject<UPackageMapClient>(this);
+	PackageMapClient->Initialize(this, Driver->GuidCache);
+	PackageMap = PackageMapClient;
 }
 
 void UNetConnection::Serialize( FArchive& Ar )
diff --git a/Engine/Source/Runtime/Engine/Private/NetworkDriver.cpp b/Engine/Source/Runtime/Engine/Private/NetworkDriver.cpp
--- a/Engine/Source/Runtime/Engine/Private/NetworkDriver.cpp
+++ b/Engine/Source/Runtime/Engine/Private/NetworkDriver.cpp
@@ -1604,7 +1604,7 @@ void UNetDriver::FlushActorDormancy(AActor* Actor)
 UChildConnection* UNetDriver::CreateChild(UNetConnection* Parent)
 {
 	UE_LOG(LogNet, Log, TEXT("Creating child connection with %s parent"), *Parent->GetName());
-	UChildConnection* Child = new UChildConnection(FObjectInitializer());
+	auto Child = NewObject<UChildConnection>();
 	Child->Driver = this;
 	Child->URL = FURL();
 	Child->State = Parent->State;
diff --git a/Engine/Source/Runtime/Engine/Private/PackageMapClient.cpp b/Engine/Source/Runtime/Engine/Private/PackageMapClient.cpp
--- a/Engine/Source/Runtime/Engine/Private/PackageMapClient.cpp
+++ b/Engine/Source/Runtime/Engine/Private/PackageMapClient.cpp
@@ -30,6 +30,7 @@ static TAutoConsoleVariable<int32> CVarIgnorePackageMismatch( TEXT( "net.IgnoreP
 -----------------------------------------------------------------------------*/
 UPackageMapClient::UPackageMapClient(const FObjectInitializer& ObjectInitializer)
 	: Super(ObjectInitializer)
+  , Connection(nullptr)
 {
 }
 
diff --git a/Engine/Source/Runtime/Engine/Private/PendingNetGame.cpp b/Engine/Source/Runtime/Engine/Private/PendingNetGame.cpp
--- a/Engine/Source/Runtime/Engine/Private/PendingNetGame.cpp
+++ b/Engine/Source/Runtime/Engine/Private/PendingNetGame.cpp
@@ -10,13 +10,12 @@
 #include "Net/DataChannel.h"
 #include "GeneralProjectSettings.h"
 
-UPendingNetGame::UPendingNetGame( const FObjectInitializer& ObjectInitializer, const FURL& InURL )
-	: Super(ObjectInitializer)
-	, NetDriver(NULL)
-	, URL(InURL)
-	, bSuccessfullyConnected(false)
-	, bSentJoinRequest(false)
+void UPendingNetGame::Initialize(const FURL& InURL)
 {
+	NetDriver = NULL;
+	URL = InURL;
+	bSuccessfullyConnected = false;
+	bSentJoinRequest = false;
 }
 
 UPendingNetGame::UPendingNetGame(const FObjectInitializer& ObjectInitializer)
diff --git a/Engine/Source/Runtime/Engine/Private/PreviewScene.cpp b/Engine/Source/Runtime/Engine/Private/PreviewScene.cpp
--- a/Engine/Source/Runtime/Engine/Private/PreviewScene.cpp
+++ b/Engine/Source/Runtime/Engine/Private/PreviewScene.cpp
@@ -14,11 +14,11 @@ FPreviewScene::FPreviewScene(FPreviewScene::ConstructionValues CVS)
 	: PreviewWorld(NULL)
 	, bForceAllUsedMipsResident(CVS.bForceMipsResident)
 {
-	PreviewWorld = new UWorld(FObjectInitializer(),FURL(NULL));
+	PreviewWorld = NewObject<UWorld>();
 	PreviewWorld->WorldType = EWorldType::Preview;
-	if (!CVS.bTransactional)
+	if (CVS.bTransactional)
 	{
-		PreviewWorld->ClearFlags(RF_Transactional);
+		PreviewWorld->SetFlags(RF_Transactional);
 	}
 
 	FWorldContext& WorldContext = GEngine->CreateNewWorldContext(EWorldType::Preview);
diff --git a/Engine/Source/Runtime/Engine/Private/PrimitiveComponent.cpp b/Engine/Source/Runtime/Engine/Private/PrimitiveComponent.cpp
--- a/Engine/Source/Runtime/Engine/Private/PrimitiveComponent.cpp
+++ b/Engine/Source/Runtime/Engine/Private/PrimitiveComponent.cpp
@@ -68,18 +68,8 @@ int32 UPrimitiveComponent::CurrentTag = 2147483647 / 4;
 // 0 is reserved to mean invalid
 uint64 UPrimitiveComponent::NextComponentId = 1;
 
-UPrimitiveComponent::UPrimitiveComponent()
-{
-	InitializePrimitiveComponentDefaults();
-}
-
-UPrimitiveComponent::UPrimitiveComponent(const FObjectInitializer& ObjectInitializer)
+UPrimitiveComponent::UPrimitiveComponent(const FObjectInitializer& ObjectInitializer /*= FObjectInitializer::Get()*/)
 	: Super(ObjectInitializer)
-{
-	InitializePrimitiveComponentDefaults();
-}
-
-void UPrimitiveComponent::InitializePrimitiveComponentDefaults()
 {
 	PostPhysicsComponentTick.bCanEverTick = false;
 	PostPhysicsComponentTick.bStartWithTickEnabled = true;
diff --git a/Engine/Source/Runtime/Engine/Private/SceneComponent.cpp b/Engine/Source/Runtime/Engine/Private/SceneComponent.cpp
--- a/Engine/Source/Runtime/Engine/Private/SceneComponent.cpp
+++ b/Engine/Source/Runtime/Engine/Private/SceneComponent.cpp
@@ -23,25 +23,15 @@ FOverlapInfo::FOverlapInfo(UPrimitiveComponent* InComponent, int32 InBodyIndex)
 	OverlapInfo.Item = InBodyIndex;
 }
 
-USceneComponent::USceneComponent()
-{
-	InitializeDefaults();
-}
-
-USceneComponent::USceneComponent( const FObjectInitializer& ObjectInitializer )
-{
-	// Forward to default constructor (we don't use ObjectInitializer for anything, this is for compatibility with inherited classes that call Super( ObjectInitializer )
-	InitializeDefaults();
-}
-
-void USceneComponent::InitializeDefaults()
+USceneComponent::USceneComponent(const FObjectInitializer& ObjectInitializer /*= FObjectInitializer::Get()*/)
+	: Super(ObjectInitializer)
 {
 	Mobility = EComponentMobility::Movable;
-	RelativeScale3D = FVector(1.0f,1.0f,1.0f);
+	RelativeScale3D = FVector(1.0f, 1.0f, 1.0f);
 	// default behavior is visible
 	bVisible = true;
-	bAutoActivate=false;
-	
+	bAutoActivate = false;
+
 	NetUpdateTransform = false;
 }
 
diff --git a/Engine/Source/Runtime/Engine/Private/ShadowMap.cpp b/Engine/Source/Runtime/Engine/Private/ShadowMap.cpp
--- a/Engine/Source/Runtime/Engine/Private/ShadowMap.cpp
+++ b/Engine/Source/Runtime/Engine/Private/ShadowMap.cpp
@@ -288,7 +288,7 @@ bool FShadowMapPendingTexture::AddElement(FShadowMapAllocationGroup& AllocationG
 void FShadowMapPendingTexture::StartEncoding(UWorld* InWorld)
 {
 	// Create the shadow-map texture.
-	UShadowMapTexture2D* Texture = new(Outer) UShadowMapTexture2D(FObjectInitializer());
+	auto Texture = NewObject<UShadowMapTexture2D>(Outer);
 	Texture->Filter			= GUseBilinearLightmaps ? TF_Default : TF_Nearest;
 	// Signed distance field textures get stored in linear space, since they need more precision near .5.
 	Texture->SRGB			= false;
diff --git a/Engine/Source/Runtime/Engine/Private/Tests/IntSerializationTest.cpp b/Engine/Source/Runtime/Engine/Private/Tests/IntSerializationTest.cpp
--- a/Engine/Source/Runtime/Engine/Private/Tests/IntSerializationTest.cpp
+++ b/Engine/Source/Runtime/Engine/Private/Tests/IntSerializationTest.cpp
@@ -11,7 +11,7 @@ bool FIntSerializationTest::RunTest (const FString& Parameters)
 	//Create object for serialization
 	//UIntSerialization SerializableObject;
 	//SerializableObject->
-	UIntSerialization* SerializableObject = new UIntSerialization(FObjectInitializer());
+	auto SerializableObject = NewObject<UIntSerialization>();
 	SerializableObject->UnsignedInt8Variable = 255U;
 	SerializableObject->UnsignedInt16Variable = 65535U;
 	SerializableObject->UnsignedInt32Variable = 4294967295U;
@@ -30,7 +30,7 @@ bool FIntSerializationTest::RunTest (const FString& Parameters)
 
 	//Deserialize into new object
 	FMemoryReader InAr(SaveData, true);
-	UIntSerialization* DeSerializableObject = new UIntSerialization(FObjectInitializer());
+	auto DeSerializableObject = NewObject<UIntSerialization>();
 	DeSerializableObject->Serialize(InAr);
 
 	//Compare test
diff --git a/Engine/Source/Runtime/Engine/Private/UnrealEngine.cpp b/Engine/Source/Runtime/Engine/Private/UnrealEngine.cpp
--- a/Engine/Source/Runtime/Engine/Private/UnrealEngine.cpp
+++ b/Engine/Source/Runtime/Engine/Private/UnrealEngine.cpp
@@ -8219,7 +8219,8 @@ EBrowseReturnVal::Type UEngine::Browse( FWorldContext& WorldContext, FURL URL, F
 			ShutdownWorldNetDriver(WorldContext.World());
 		}
 
-		WorldContext.PendingNetGame = new UPendingNetGame(FObjectInitializer(), URL);
+		WorldContext.PendingNetGame = NewObject<UPendingNetGame>();
+		WorldContext.PendingNetGame->Initialize(URL);
 		WorldContext.PendingNetGame->InitNetDriver();
 		if( !WorldContext.PendingNetGame->NetDriver )
 		{
diff --git a/Engine/Source/Runtime/Engine/Private/World.cpp b/Engine/Source/Runtime/Engine/Private/World.cpp
--- a/Engine/Source/Runtime/Engine/Private/World.cpp
+++ b/Engine/Source/Runtime/Engine/Private/World.cpp
@@ -107,6 +107,7 @@ UWorld::UWorld( const FObjectInitializer& ObjectInitializer )
 ,	HierarchicalLODBuilder(this)
 #endif
 ,	FeatureLevel(GMaxRHIFeatureLevel)
+, URL(FURL(NULL))
 ,	FXSystem(NULL)
 ,	TickTaskLevel(FTickTaskManagerInterface::Get().AllocateTickTaskLevel())
 ,   bIsBuilt(false)
@@ -121,30 +122,6 @@ UWorld::UWorld( const FObjectInitializer& ObjectInitializer )
 	FWorldDelegates::OnPostWorldCreation.Broadcast(this);
 }
 
-
-UWorld::UWorld( const FObjectInitializer& ObjectInitializer,const FURL& InURL )
-:	UObject(ObjectInitializer)
-#if WITH_EDITOR
-,	HierarchicalLODBuilder(this)
-#endif
-,	FeatureLevel(GMaxRHIFeatureLevel)
-,	URL(InURL)
-,	FXSystem(NULL)
-,	TickTaskLevel(FTickTaskManagerInterface::Get().AllocateTickTaskLevel())
-,   bIsBuilt(false)
-,	FlushLevelStreamingType(EFlushLevelStreamingType::None)
-,	NextTravelType(TRAVEL_Relative)
-{
-	SetFlags( RF_Transactional );
-	TimerManager = new FTimerManager();
-#if WITH_EDITOR
-	bBroadcastSelectionChange = true;
-#endif // WITH_EDITOR
-
-	FWorldDelegates::OnPostWorldCreation.Broadcast(this);
-}
-
-
 void UWorld::Serialize( FArchive& Ar )
 {
 	Super::Serialize( Ar );
@@ -907,7 +884,8 @@ void UWorld::InitWorld(const InitializationValues IVS)
 			// Spawn the default brush.
 			DefaultBrush = SpawnBrush();
 			check(DefaultBrush->GetBrushComponent());
-			DefaultBrush->Brush = new( DefaultBrush->GetOuter(), TEXT("Brush") )UModel( FObjectInitializer(), DefaultBrush, 1 );
+			DefaultBrush->Brush = NewNamedObject<UModel>(DefaultBrush->GetOuter(), TEXT("Brush"));
+			DefaultBrush->Brush->Initialize(DefaultBrush, 1);
 			DefaultBrush->GetBrushComponent()->Brush = DefaultBrush->Brush;
 			DefaultBrush->SetNotForClientOrServer();
 			DefaultBrush->Brush->SetFlags( RF_Transactional );
@@ -977,8 +955,10 @@ void UWorld::InitializeNewWorld(const InitializationValues IVS)
 		ClearFlags(RF_Transactional);
 	}
 
-	PersistentLevel			= new( this, TEXT("PersistentLevel") ) ULevel(FObjectInitializer(),FURL(NULL));
-	PersistentLevel->Model	= new( PersistentLevel				 ) UModel(FObjectInitializer(),NULL, 1 );
+	PersistentLevel = NewNamedObject<ULevel>(this, TEXT("PersistentLevel"));
+	PersistentLevel->Initialize(FURL(nullptr));
+	PersistentLevel->Model = NewObject<UModel>(PersistentLevel);
+	PersistentLevel->Model->Initialize(nullptr, 1);
 	PersistentLevel->OwningWorld = this;
 
 	// Mark objects are transactional for undo/ redo.
@@ -1071,7 +1051,8 @@ UWorld* UWorld::CreateWorld(const EWorldType::Type InWorldType, bool bInformEngi
 
 	// Create new UWorld, ULevel and UModel.
 	const FString WorldNameString = (WorldName != NAME_None) ? WorldName.ToString() : TEXT("NewWorld");
-	UWorld* NewWorld = new(WorldPackage, *WorldNameString) UWorld(FObjectInitializer(),FURL(NULL));
+	UWorld* NewWorld = NewNamedObject<UWorld>(WorldPackage, *WorldNameString);
+	NewWorld->SetFlags(RF_Transactional);
 	NewWorld->WorldType = InWorldType;
 	NewWorld->FeatureLevel = InFeatureLevel;
 	NewWorld->InitializeNewWorld(UWorld::InitializationValues().ShouldSimulatePhysics(false).EnableTraceCollision(true).CreateNavigation(InWorldType == EWorldType::Editor).CreateAISystem(InWorldType == EWorldType::Editor));
diff --git a/Engine/Source/Runtime/Engine/Public/Model.h b/Engine/Source/Runtime/Engine/Public/Model.h
--- a/Engine/Source/Runtime/Engine/Public/Model.h
+++ b/Engine/Source/Runtime/Engine/Public/Model.h
@@ -411,29 +411,9 @@ private:
 public:
 
 	// Constructors.
-	UModel(const FObjectInitializer& ObjectInitializer)
-	: UObject(ObjectInitializer)
-	, Nodes( this )
-	, Verts( this )
-	, Vectors( this )
-	, Points( this )
-	, Surfs( this )
-	, VertexBuffer( this )
-#if WITH_EDITOR
-	, LightingLevel( NULL )
-#endif // WITH_EDITOR
-	, RootOutside( true )
-	{
-		if ( !HasAnyFlags(RF_ClassDefaultObject) )
-		{
-			EmptyModel( 1, 0 );
-			if( GIsEditor && !FApp::IsGame() )
-			{
-				UpdateVertices();
-			}
-		}
-	}
-	ENGINE_API UModel( const FObjectInitializer& ObjectInitializer,ABrush* Owner, bool InRootOutside=1 );
+	UModel(const FObjectInitializer& ObjectInitializer);
+	ENGINE_API void Initialize();
+	ENGINE_API void Initialize(ABrush* Owner, bool InRootOutside = true);
 
 	// UObject interface.
 	virtual void Serialize( FArchive& Ar );	
diff --git a/Engine/Source/Runtime/Landscape/Private/LandscapeGrass.cpp b/Engine/Source/Runtime/Landscape/Private/LandscapeGrass.cpp
--- a/Engine/Source/Runtime/Landscape/Private/LandscapeGrass.cpp
+++ b/Engine/Source/Runtime/Landscape/Private/LandscapeGrass.cpp
@@ -334,7 +334,7 @@ public:
 			0.5f / ZOffset,
 			ZOffset);
 
-		RenderTargetTexture = new UTextureRenderTarget2D(FObjectInitializer());
+		RenderTargetTexture = NewObject<UTextureRenderTarget2D>();
 		check(RenderTargetTexture);
 		RenderTargetTexture->ClearColor = FLinearColor::Transparent;
 		RenderTargetTexture->TargetGamma = 1.f;
diff --git a/Engine/Source/Runtime/SlateCore/Private/Fonts/FontBulkData.cpp b/Engine/Source/Runtime/SlateCore/Private/Fonts/FontBulkData.cpp
--- a/Engine/Source/Runtime/SlateCore/Private/Fonts/FontBulkData.cpp
+++ b/Engine/Source/Runtime/SlateCore/Private/Fonts/FontBulkData.cpp
@@ -8,7 +8,7 @@ UFontBulkData::UFontBulkData()
 	BulkData.SetBulkDataFlags(BULKDATA_SerializeCompressed);
 }
 
-UFontBulkData::UFontBulkData(const FString& InFontFilename)
+void UFontBulkData::Initialize(const FString& InFontFilename)
 {
 	BulkData.SetBulkDataFlags(BULKDATA_SerializeCompressed);
 
@@ -28,7 +28,7 @@ UFontBulkData::UFontBulkData(const FString& InFontFilename)
 	}
 }
 
-UFontBulkData::UFontBulkData(const void* const InFontData, const int32 InFontDataSizeBytes)
+void UFontBulkData::Initialize(const void* const InFontData, const int32 InFontDataSizeBytes)
 {
 	BulkData.SetBulkDataFlags(BULKDATA_SerializeCompressed);
 
diff --git a/Engine/Source/Runtime/SlateCore/Private/Fonts/LegacySlateFontInfoCache.cpp b/Engine/Source/Runtime/SlateCore/Private/Fonts/LegacySlateFontInfoCache.cpp
--- a/Engine/Source/Runtime/SlateCore/Private/Fonts/LegacySlateFontInfoCache.cpp
+++ b/Engine/Source/Runtime/SlateCore/Private/Fonts/LegacySlateFontInfoCache.cpp
@@ -49,7 +49,9 @@ TSharedPtr<const FCompositeFont> FLegacySlateFontInfoCache::GetCompositeFont(con
 		}
 	}
 
-	TSharedRef<const FCompositeFont> NewCompositeFont = MakeShareable(new FStandaloneCompositeFont(NAME_None, LegacyFontPath, new UFontBulkData(LegacyFontPath), InLegacyFontHinting));
+	UFontBulkData* FontBulkData = NewObject<UFontBulkData>();
+	FontBulkData->Initialize(LegacyFontPath);
+	TSharedRef<const FCompositeFont> NewCompositeFont = MakeShareable(new FStandaloneCompositeFont(NAME_None, LegacyFontPath, FontBulkData, InLegacyFontHinting));
 	LegacyFontNameToCompositeFont.Add(LegacyFontKey, NewCompositeFont);
 	return NewCompositeFont;
 }
@@ -59,7 +61,9 @@ TSharedPtr<const FCompositeFont> FLegacySlateFontInfoCache::GetSystemFont()
 	if (!SystemFont.IsValid())
 	{
 		TArray<uint8> FontBytes = FPlatformMisc::GetSystemFontBytes();
-		SystemFont = MakeShareable(new FStandaloneCompositeFont(NAME_None, TEXT("DefaultSystemFont"), new UFontBulkData(FontBytes.GetData(), FontBytes.Num()), EFontHinting::Default));
+		UFontBulkData* FontBulkData = NewObject<UFontBulkData>();
+		FontBulkData->Initialize(FontBytes.GetData(), FontBytes.Num());
+		SystemFont = MakeShareable(new FStandaloneCompositeFont(NAME_None, TEXT("DefaultSystemFont"), FontBulkData, EFontHinting::Default));
 	}
 	return SystemFont;
 }
@@ -80,7 +84,9 @@ const FFontData& FLegacySlateFontInfoCache::GetFallbackFont()
 	}
 
 	const FString FallbackFontPath = FallbackFontName.ToString();
-	TSharedRef<const FFontData> NewFallbackFont = MakeShareable(new FFontData(FallbackFontPath, new UFontBulkData(FallbackFontPath), EFontHinting::Default));
+	UFontBulkData* FontBulkData = NewObject<UFontBulkData>();
+	FontBulkData->Initialize(FallbackFontPath);
+	TSharedRef<const FFontData> NewFallbackFont = MakeShareable(new FFontData(FallbackFontPath, FontBulkData, EFontHinting::Default));
 	FallbackFonts.Add(FallbackFontName, NewFallbackFont);
 	return *NewFallbackFont;
 }
@@ -93,7 +99,9 @@ const FFontData& FLegacySlateFontInfoCache::GetLastResortFont()
 	if (!LastResortFont.IsValid())
 	{
 		const FString LastResortFontPath = FPaths::EngineContentDir() / TEXT("Slate/Fonts/LastResort.ttf");
-		LastResortFont = MakeShareable(new FFontData(LastResortFontPath, new UFontBulkData(LastResortFontPath), EFontHinting::Default));
+		UFontBulkData* FontBulkData = NewObject<UFontBulkData>();
+		FontBulkData->Initialize(LastResortFontPath);
+		LastResortFont = MakeShareable(new FFontData(LastResortFontPath, FontBulkData, EFontHinting::Default));
 	}
 
 	return *LastResortFont;
diff --git a/Engine/Source/Runtime/SlateCore/Public/Fonts/FontBulkData.h b/Engine/Source/Runtime/SlateCore/Public/Fonts/FontBulkData.h
--- a/Engine/Source/Runtime/SlateCore/Public/Fonts/FontBulkData.h
+++ b/Engine/Source/Runtime/SlateCore/Public/Fonts/FontBulkData.h
@@ -14,10 +14,10 @@ public:
 	UFontBulkData();
 
 	/** Construct the bulk font data from the given file */
-	UFontBulkData(const FString& InFontFilename);
+	void Initialize(const FString& InFontFilename);
 
 	/** Construct the bulk font data from the given data */
-	UFontBulkData(const void* const InFontData, const int32 InFontDataSizeBytes);
+	void Initialize(const void* const InFontData, const int32 InFontDataSizeBytes);
 
 	/** Locks the bulk font data and returns a read-only pointer to it */
 	const void* Lock(int32& OutFontDataSizeBytes) const;
