#ue4 - Deprecating UProperties does not cause deprecation warnings in generated code.

diff --git a/Engine/Source/Programs/UnrealHeaderTool/Private/CodeGenerator.cpp b/Engine/Source/Programs/UnrealHeaderTool/Private/CodeGenerator.cpp
--- a/Engine/Source/Programs/UnrealHeaderTool/Private/CodeGenerator.cpp
+++ b/Engine/Source/Programs/UnrealHeaderTool/Private/CodeGenerator.cpp
@@ -1087,7 +1087,13 @@ void FNativeClassHeaderGenerator::ExportNativeGeneratedInitCode(FClass* Class, F
 			{
 				Props.Add(*ItInner);
 			}
-			OutputProperties(Meta, GeneratedClassRegisterFunctionText, OuterString, Props, TEXT("                "));
+
+			if (Props.Num() > 0)
+			{
+				GeneratedClassRegisterFunctionText.Logf(TEXT("PRAGMA_DISABLE_DEPRECATION_WARNINGS\r\n"));
+				OutputProperties(Meta, GeneratedClassRegisterFunctionText, OuterString, Props, TEXT("                "));
+				GeneratedClassRegisterFunctionText.Logf(TEXT("PRAGMA_POP\r\n"));
+			}
 		}
 		// function table
 		{
