UECORE-59: I want to be able to specify server/validation function names in UHT and not rely on predefined ones.

diff --git a/Engine/Source/Programs/UnrealHeaderTool/Private/CodeGenerator.cpp b/Engine/Source/Programs/UnrealHeaderTool/Private/CodeGenerator.cpp
--- a/Engine/Source/Programs/UnrealHeaderTool/Private/CodeGenerator.cpp
+++ b/Engine/Source/Programs/UnrealHeaderTool/Private/CodeGenerator.cpp
@@ -3361,7 +3361,7 @@ void FNativeClassHeaderGenerator::ExportFunctionChecks(const FFuncInfo& Function
 	if (bNeedsValidate)
 	{
 		FunctionName = FunctionData.CppValidationImplName;
-		StaticChecks.Logf(TEXT("\t\tstatic_assert(THasMemberFunction_%s<%s>::Value, \"%s Declare function bool %s(%s) %s in class %s.\");\r\n"), *FunctionName, *ClassName, *AssertMessage, *FunctionReturnType, *FunctionName, *ParameterString, *ConstModifier, *ClassName);
+		StaticChecks.Logf(TEXT("\t\tstatic_assert(THasMemberFunction_%s<%s>::Value, \"%s Declare function bool %s(%s) %s in class %s.\");\r\n"), *FunctionName, *ClassName, *AssertMessage, *FunctionName, *ParameterString, *ConstModifier, *ClassName);
 		CheckClasses.Logf(TEXT("GENERATE_MEMBER_FUNCTION_CHECK(%s, bool, %s, %s)\r\n"), *FunctionName, *ConstModifier, *ParameterString);
 	}
 }
@@ -3620,20 +3620,20 @@ void FNativeClassHeaderGenerator::ExportFunctionThunk(FStringOutputDevice& RPCWr
 
 	auto FunctionName = Function->GetName();
 
-	bool bHasImplementation = ClassDefinition.Contains(FunctionName + TEXT("_Implementation"), ESearchCase::CaseSensitive);
-	bool bHasValidate = ClassDefinition.Contains(FunctionName + TEXT("_Validate"), ESearchCase::CaseSensitive);
+	bool bHasImplementation = ClassDefinition.Contains(FunctionData.CppImplName, ESearchCase::CaseSensitive);
+	bool bHasValidate = ClassDefinition.Contains(FunctionData.CppValidationImplName, ESearchCase::CaseSensitive);
 
 	auto bShouldEnableImplementationDeprecation =
 		// Enable deprecation warnings only if GENERATED_BODY is used inside class or interface (not GENERATED_UCLASS_BODY etc.)
 		ClassRange.bHasGeneratedBody
-		// And _Implementation function is called, but not the one declared by user
-		&& (FunctionData.CppImplName.EndsWith(TEXT("_Implementation")) && !bHasImplementation);
+		// and implementation function is called, but not the one declared by user
+		&& (FunctionData.CppImplName != FunctionName && !bHasImplementation);
 
 	auto bShouldEnableValidateDeprecation =
 		// Enable deprecation warnings only if GENERATED_BODY is used inside class or interface (not GENERATED_UCLASS_BODY etc.)
 		ClassRange.bHasGeneratedBody
-		// or _Validate function is called, but not the one declared by user
-		&& FunctionData.CppImplName.EndsWith(TEXT("_Validate")) && !bHasValidate;
+		// and validation function is called
+		&& (FunctionData.FunctionFlags & FUNC_NetValidate) && !bHasValidate;
 
 	//Emit warning here if necessary
 	FStringOutputDevice FunctionDeclaration;
@@ -3647,7 +3647,7 @@ void FNativeClassHeaderGenerator::ExportFunctionThunk(FStringOutputDevice& RPCWr
 		// NOTE - We don't need to set "Result", since RPC calls don't return any values...
 		if (bShouldEnableValidateDeprecation)
 		{
-			RPCWrappers.Logf(TEXT("EMIT_CUSTOM_WARNING(\"Function %s::%s needs validation by %s due to its properties. Currently %s declaration is autogenerated by UHT. Since next release you'll have to provide declaration on your own. Please update your code before upgrading to the next release, otherwise your project will no longer compile.\")\r\n"), *ClassName, *FunctionName, *FunctionDeclaration, *FunctionData.CppValidationImplName);
+			RPCWrappers.Logf(TEXT("EMIT_CUSTOM_WARNING(\"Function %s::%s needs validation by %s due to its properties. Currently %s declaration is autogenerated by UHT. Since next release you'll have to provide declaration on your own. Please update your code before upgrading to the next release, otherwise your project will no longer compile.\")\r\n"), *ClassName, *FunctionName, *FunctionData.CppValidationImplName, *FunctionData.CppValidationImplName);
 		}
 		RPCWrappers.Logf(TEXT("%sif (!this->%s(%s))%s"), FCString::Spc(8), *FunctionData.CppValidationImplName, *ParameterList, LINE_TERMINATOR);
 		RPCWrappers.Logf(TEXT("%s{%s"), FCString::Spc(8), LINE_TERMINATOR);
@@ -3658,7 +3658,7 @@ void FNativeClassHeaderGenerator::ExportFunctionThunk(FStringOutputDevice& RPCWr
 
 	if (bShouldEnableImplementationDeprecation)
 	{
-		RPCWrappers.Logf(TEXT("EMIT_CUSTOM_WARNING(\"Function %s::%s needs native implementation by %s due to its properties. Currently %s declaration is autogenerated by UHT. Since next release you'll have to provide declaration on your own. Please update your code before upgrading to the next release, otherwise your project will no longer compile.\")\r\n"), *ClassName, *FunctionName, *FunctionDeclaration, *FunctionData.CppImplName);
+		RPCWrappers.Logf(TEXT("EMIT_CUSTOM_WARNING(\"Function %s::%s needs native implementation by %s due to its properties. Currently %s declaration is autogenerated by UHT. Since next release you'll have to provide declaration on your own. Please update your code before upgrading to the next release, otherwise your project will no longer compile.\")\r\n"), *ClassName, *FunctionName, *FunctionData.CppImplName, *FunctionData.CppImplName);
 	}
 	// write out the return value
 	RPCWrappers.Log(FCString::Spc(8));
@@ -3764,8 +3764,8 @@ void FNativeClassHeaderGenerator::ExportNativeFunctions(FUnrealSourceFile& Sourc
 
 			auto FunctionName = Function->GetName();
 
-			bool bHasImplementation = ClassDefinition.Contains(FunctionName + TEXT("_Implementation"), ESearchCase::CaseSensitive);
-			bool bHasValidate = ClassDefinition.Contains(FunctionName + TEXT("_Validate"), ESearchCase::CaseSensitive);
+			bool bHasImplementation = ClassDefinition.Contains(FunctionData.CppImplName, ESearchCase::CaseSensitive);
+			bool bHasValidate = ClassDefinition.Contains(FunctionData.CppValidationImplName, ESearchCase::CaseSensitive);
 
 			//Emit warning here if necessary
 			FStringOutputDevice FunctionDeclaration;
@@ -3786,7 +3786,7 @@ void FNativeClassHeaderGenerator::ExportNativeFunctions(FUnrealSourceFile& Sourc
 			}
 
 			AutogeneratedBlueprintFunctionDeclarations.Log(*FunctionDeclaration);
-			if ((FunctionData.bCppImplNameEndsWith_Implementation && !bHasImplementation) || (FunctionData.bCppValidationImplNameEndsWith_Validate && !bHasValidate))
+			if ((FunctionData.CppImplName != FunctionName && !bHasImplementation))
 			{
 				AutogeneratedBlueprintFunctionDeclarationsOnlyNotDeclared.Log(*FunctionDeclaration);
 			}
diff --git a/Engine/Source/Programs/UnrealHeaderTool/Private/HeaderParser.cpp b/Engine/Source/Programs/UnrealHeaderTool/Private/HeaderParser.cpp
--- a/Engine/Source/Programs/UnrealHeaderTool/Private/HeaderParser.cpp
+++ b/Engine/Source/Programs/UnrealHeaderTool/Private/HeaderParser.cpp
@@ -233,6 +233,12 @@ namespace
 
 				FuncInfo.FunctionFlags |= FUNC_Net;
 				FuncInfo.FunctionFlags |= FUNC_NetServer;
+
+				if (Specifier.Values.Num())
+				{
+					FuncInfo.CppImplName = Specifier.Values[0];
+				}
+
 				if( FuncInfo.FunctionFlags & FUNC_Exec )
 				{
 					FError::Throwf(TEXT("Exec functions cannot be replicated!") );
@@ -247,6 +253,11 @@ namespace
 
 				FuncInfo.FunctionFlags |= FUNC_Net;
 				FuncInfo.FunctionFlags |= FUNC_NetClient;
+
+				if (Specifier.Values.Num())
+				{
+					FuncInfo.CppImplName = Specifier.Values[0];
+				}
 			}
 			else if (Specifier.Key == TEXT("NetMulticast"))
 			{
@@ -318,6 +329,11 @@ namespace
 			else if (Specifier.Key == TEXT("WithValidation"))
 			{
 				FuncInfo.FunctionFlags |= FUNC_NetValidate;
+
+				if (Specifier.Values.Num())
+				{
+					FuncInfo.CppValidationImplName = Specifier.Values[0];
+				}
 			}
 			else
 			{
diff --git a/Engine/Source/Programs/UnrealHeaderTool/Private/ParserHelper.h b/Engine/Source/Programs/UnrealHeaderTool/Private/ParserHelper.h
--- a/Engine/Source/Programs/UnrealHeaderTool/Private/ParserHelper.h
+++ b/Engine/Source/Programs/UnrealHeaderTool/Private/ParserHelper.h
@@ -869,12 +869,8 @@ struct FFuncInfo
 	FString		MarshallAndCallName;
 	/** Name of the actual implementation **/
 	FString		CppImplName;
-	/** Cached info whether CppImplName ends with _Implementation suffix to avoid string comparisons. **/
-	bool		bCppImplNameEndsWith_Implementation;
 	/** Name of the actual validation implementation **/
 	FString		CppValidationImplName;
-	/** Cached info whether CppValidationImplName ends with _Validate suffix to avoid string comparisons. **/
-	bool		bCppValidationImplNameEndsWith_Validate;
 	/** Name for callback-style names **/
 	FString		UnMarshallAndCallName;
 	/** Identifier for an RPC call to a platform service */
@@ -889,17 +885,17 @@ struct FFuncInfo
 
 	/** Constructor. */
 	FFuncInfo()
-	:	Function            ()
-	,	FunctionFlags       (0)
-	,	FunctionExportFlags (0)
-	,	ExpectParms		    (0)
-	,	FunctionReference(NULL)
-	,	bCppImplNameEndsWith_Implementation(false)
-	,	bCppValidationImplNameEndsWith_Validate(false)
-	,	RPCId               (0)
-	,	RPCResponseId       (0)
-	,	bSealedEvent        (false)
-	,	DelegateMacroLine(-1)
+	:	Function				()
+	,	FunctionFlags			(0)
+	,	FunctionExportFlags		(0)
+	,	ExpectParms				(0)
+	,	FunctionReference		(NULL)
+	,	CppImplName				(TEXT(""))
+	,	CppValidationImplName	(TEXT(""))
+	,	RPCId					(0)
+	,	RPCResponseId			(0)
+	,	bSealedEvent			(false)
+	,	DelegateMacroLine		(-1)
 	{}
 
 	FFuncInfo( const FFuncInfo& Other )
@@ -907,8 +903,8 @@ struct FFuncInfo
 	,	FunctionExportFlags(Other.FunctionExportFlags)
 	,	ExpectParms		(Other.ExpectParms)
 	,	FunctionReference(Other.FunctionReference)
-	,	bCppImplNameEndsWith_Implementation(Other.bCppImplNameEndsWith_Implementation)
-	,	bCppValidationImplNameEndsWith_Validate(Other.bCppValidationImplNameEndsWith_Validate)
+	,	CppImplName(Other.CppImplName)
+	,	CppValidationImplName(Other.CppValidationImplName)
 	,	RPCId(Other.RPCId)
 	,	RPCResponseId(Other.RPCResponseId)
 	,	DelegateMacroLine(Other.DelegateMacroLine)
@@ -929,8 +925,7 @@ struct FFuncInfo
 			FunctionName = FunctionName.LeftChop( FString( TEXT( "__DelegateSignature" ) ).Len() );
 		}
 		UnMarshallAndCallName = FString(TEXT("exec")) + FunctionName;
-		CppImplName = FunctionName;
-
+		
 		if (FunctionReference->HasAnyFunctionFlags(FUNC_BlueprintEvent))
 		{
 			MarshallAndCallName = FunctionName;
@@ -940,8 +935,6 @@ struct FFuncInfo
 			MarshallAndCallName = FString(TEXT("event")) + FunctionName;
 		}
 
-		CppValidationImplName = TEXT("");
-
 		if (FunctionReference->HasAllFunctionFlags(FUNC_Native | FUNC_Net))
 		{
 			MarshallAndCallName = FunctionName;
@@ -952,12 +945,22 @@ struct FFuncInfo
 			}
 			else
 			{
-				CppImplName = FunctionReference->GetName() + TEXT("_Implementation");
-				bCppImplNameEndsWith_Implementation = true;
-				if (FunctionReference->HasAllFunctionFlags(FUNC_NetValidate))
+				if (CppImplName.IsEmpty())
+				{
+					CppImplName = FunctionReference->GetName() + TEXT("_Implementation");
+				}
+				else if (CppImplName == FunctionName)
+				{
+					FError::Throwf(TEXT("Native implementation function must be different than original function name."));
+				}
+
+				if (CppValidationImplName.IsEmpty() && FunctionReference->HasAllFunctionFlags(FUNC_NetValidate))
 				{
 					CppValidationImplName = FunctionReference->GetName() + TEXT("_Validate");
-					bCppValidationImplNameEndsWith_Validate = true;
+				}
+				else if(CppValidationImplName == FunctionName)
+				{
+					FError::Throwf(TEXT("Validation function must be different than original function name."));
 				}
 			}
 		}
@@ -971,7 +974,11 @@ struct FFuncInfo
 		{
 			MarshallAndCallName = FunctionName;
 			CppImplName = FunctionReference->GetName() + TEXT("_Implementation");
-			bCppImplNameEndsWith_Implementation = true;
+		}
+
+		if (CppImplName.IsEmpty())
+		{
+			CppImplName = FunctionName;
 		}
 	}
 };
