Making sure UHT does not try to automatically declare constructor with OI if the base class does not declare one.

diff --git a/Engine/Source/Programs/UnrealHeaderTool/Private/CodeGenerator.cpp b/Engine/Source/Programs/UnrealHeaderTool/Private/CodeGenerator.cpp
--- a/Engine/Source/Programs/UnrealHeaderTool/Private/CodeGenerator.cpp
+++ b/Engine/Source/Programs/UnrealHeaderTool/Private/CodeGenerator.cpp
@@ -1898,10 +1898,31 @@ void ExportConstructorDefinition(FStringOutputDevice& Out, FClass* Class, const
 	if (!ClassData->bConstructorDeclared)
 	{
 		Out.Logf(TEXT("%s/** Standard constructor, called after all reflected properties have been initialized */\r\n"), FCString::Spc(4));
-		Out.Logf(TEXT("%s%s_API %s(const class FObjectInitializer& ObjectInitializer = FObjectInitializer::Get()) : Super(ObjectInitializer) { };\r\n"), FCString::Spc(4), *API, NameLookupCPP.GetNameCPP(Class));
+		
+		// Assume super class has OI constructor, this may not always be true but we should always be able to check this.
+		// In any case, it will default to old behaviour before we even checked this.
+		bool bSuperClassObjectInitializerConstructorDeclared = true;
+		auto* SuperClass = Class->GetSuperClass();		
+		if (SuperClass != nullptr)
+		{
+			auto* SuperClassData = GScriptHelper.FindClassData(SuperClass);
+			if (SuperClassData)
+			{
+				bSuperClassObjectInitializerConstructorDeclared = SuperClassData->bObjectInitializerConstructorDeclared;
+			}
+		}
+		if (bSuperClassObjectInitializerConstructorDeclared)
+		{
+			Out.Logf(TEXT("%s%s_API %s(const class FObjectInitializer& ObjectInitializer = FObjectInitializer::Get()) : Super(ObjectInitializer) { };\r\n"), FCString::Spc(4), *API, NameLookupCPP.GetNameCPP(Class));
+			ClassData->bObjectInitializerConstructorDeclared = true;
+		}
+		else
+		{
+			Out.Logf(TEXT("%s%s_API %s() { };\r\n"), FCString::Spc(4), *API, NameLookupCPP.GetNameCPP(Class));
+			ClassData->bDefaultConstructorDeclared = true;
+		}
 
 		ClassData->bConstructorDeclared = true;
-		ClassData->bObjectInitializerConstructorDeclared = true;
 	}
 	ExportCopyConstructorDefinition(Out, Class, API);
 }
