UHT: Missing adding function to the current scope which resulted in lack of error when there were two functions with the same name in class.

diff --git a/Engine/Source/Programs/UnrealHeaderTool/Private/HeaderParser.cpp b/Engine/Source/Programs/UnrealHeaderTool/Private/HeaderParser.cpp
--- a/Engine/Source/Programs/UnrealHeaderTool/Private/HeaderParser.cpp
+++ b/Engine/Source/Programs/UnrealHeaderTool/Private/HeaderParser.cpp
@@ -5359,6 +5359,8 @@ void FHeaderParser::CompileFunctionDeclaration(FUnrealSourceFile& SourceFile, FC
 	FuncInfo.FunctionReference = TopFunction;
 	FuncInfo.SetFunctionNames();
 
+	GetCurrentScope()->AddType(TopFunction);
+
 	auto* StoredFuncData = FFunctionData::Add(FuncInfo);
 	if (FuncInfo.FunctionReference->HasAnyFunctionFlags(FUNC_Delegate))
 	{
