Change _Implementation and _Validation functions UHT warning message.

diff --git a/Engine/Source/Programs/UnrealHeaderTool/Private/CodeGenerator.cpp b/Engine/Source/Programs/UnrealHeaderTool/Private/CodeGenerator.cpp
--- a/Engine/Source/Programs/UnrealHeaderTool/Private/CodeGenerator.cpp
+++ b/Engine/Source/Programs/UnrealHeaderTool/Private/CodeGenerator.cpp
@@ -3668,7 +3668,7 @@ void FNativeClassHeaderGenerator::ExportFunctionThunk(FStringOutputDevice& RPCWr
 		// NOTE - We don't need to set "Result", since RPC calls don't return any values...
 		if (bShouldEnableValidateDeprecation)
 		{
-			RPCWrappers.Logf(TEXT("EMIT_CUSTOM_WARNING(\"Function %s::%s needs validation by %s due to its properties. Currently %s declaration is autogenerated by UHT. Since next release you'll have to provide declaration on your own. Please update your code before upgrading to the next release, otherwise your project will no longer compile.\")\r\n"), *ClassName, *FunctionName, *FunctionData.CppValidationImplName, *FunctionData.CppValidationImplName);
+			RPCWrappers.Logf(TEXT("EMIT_CUSTOM_WARNING(\"Function %s::%s needs additional native validation function declaration %s due to its properties. Currently %s declaration is autogenerated by UHT. Since next release you'll have to provide declaration on your own. Please update your code before upgrading to the next release, otherwise your project will no longer compile.\")\r\n"), *ClassName, *FunctionName, *FunctionData.CppValidationImplName, *FunctionData.CppValidationImplName);
 		}
 		RPCWrappers.Logf(TEXT("%sif (!this->%s(%s))%s"), FCString::Spc(8), *FunctionData.CppValidationImplName, *ParameterList, LINE_TERMINATOR);
 		RPCWrappers.Logf(TEXT("%s{%s"), FCString::Spc(8), LINE_TERMINATOR);
@@ -3679,7 +3679,7 @@ void FNativeClassHeaderGenerator::ExportFunctionThunk(FStringOutputDevice& RPCWr
 
 	if (bShouldEnableImplementationDeprecation)
 	{
-		RPCWrappers.Logf(TEXT("EMIT_CUSTOM_WARNING(\"Function %s::%s needs native implementation by %s due to its properties. Currently %s declaration is autogenerated by UHT. Since next release you'll have to provide declaration on your own. Please update your code before upgrading to the next release, otherwise your project will no longer compile.\")\r\n"), *ClassName, *FunctionName, *FunctionData.CppImplName, *FunctionData.CppImplName);
+		RPCWrappers.Logf(TEXT("EMIT_CUSTOM_WARNING(\"Function %s::%s needs additional native declaration %s due to its properties. Currently %s declaration is autogenerated by UHT. Since next release you'll have to provide declaration on your own. Please update your code before upgrading to the next release, otherwise your project will no longer compile.\")\r\n"), *ClassName, *FunctionName, *FunctionData.CppImplName, *FunctionData.CppImplName);
 	}
 	// write out the return value
 	RPCWrappers.Log(FCString::Spc(8));
