Fixed UHT issue where enums were not being registered as types in a data structure that is used to track what header file a type is declared in, in turn used to populate metadata.

diff --git a/Engine/Source/Programs/UnrealHeaderTool/Private/ClassMaps.cpp b/Engine/Source/Programs/UnrealHeaderTool/Private/ClassMaps.cpp
--- a/Engine/Source/Programs/UnrealHeaderTool/Private/ClassMaps.cpp
+++ b/Engine/Source/Programs/UnrealHeaderTool/Private/ClassMaps.cpp
@@ -16,11 +16,11 @@ TMap<UPackage*,  const FManifestModule*>							GPackageToManifestModuleMap;
 TMap<UField*, uint32>												GGeneratedCodeCRCs;
 TMap<UEnum*,  EPropertyType>										GEnumUnderlyingTypes;
 
-TSharedRef<FUnrealTypeDefinitionInfo> AddTypeDefinition(FUnrealSourceFile& SourceFile, UStruct* Struct, int32 Line)
+TSharedRef<FUnrealTypeDefinitionInfo> AddTypeDefinition(FUnrealSourceFile& SourceFile, UField* Field, int32 Line)
 {
 	TSharedRef<FUnrealTypeDefinitionInfo> DefinitionInfo = MakeShareable(new FUnrealTypeDefinitionInfo(SourceFile, Line));
 
-	GTypeDefinitionInfoMap.Add(Struct, DefinitionInfo);
+	GTypeDefinitionInfoMap.Add(Field, DefinitionInfo);
 
 	return DefinitionInfo;
 }
\ No newline at end of file
diff --git a/Engine/Source/Programs/UnrealHeaderTool/Private/ClassMaps.h b/Engine/Source/Programs/UnrealHeaderTool/Private/ClassMaps.h
--- a/Engine/Source/Programs/UnrealHeaderTool/Private/ClassMaps.h
+++ b/Engine/Source/Programs/UnrealHeaderTool/Private/ClassMaps.h
@@ -37,9 +37,9 @@ enum EAccessSpecifier
  * Add type definition info to global map.
  *
  * @param SourceFile SourceFile in which type was defined.
- * @param Struct Defined type.
+ * @param Field Defined type.
  * @param Line Line on which the type was defined.
  *
  * @returns Type definition info.
  */
-TSharedRef<FUnrealTypeDefinitionInfo> AddTypeDefinition(FUnrealSourceFile& SourceFile, UStruct* Struct, int32 Line);
\ No newline at end of file
+TSharedRef<FUnrealTypeDefinitionInfo> AddTypeDefinition(FUnrealSourceFile& SourceFile, UField* Field, int32 Line);
\ No newline at end of file
diff --git a/Engine/Source/Programs/UnrealHeaderTool/Private/HeaderParser.cpp b/Engine/Source/Programs/UnrealHeaderTool/Private/HeaderParser.cpp
--- a/Engine/Source/Programs/UnrealHeaderTool/Private/HeaderParser.cpp
+++ b/Engine/Source/Programs/UnrealHeaderTool/Private/HeaderParser.cpp
@@ -664,6 +664,8 @@ UEnum* FHeaderParser::CompileEnum(FUnrealSourceFile& SourceFile)
 	UEnum* Enum = new(EC_InternalUseOnlyConstructor, SourceFile.GetPackage(), EnumToken.Identifier, RF_Public) UEnum(FObjectInitializer());
 	Scope->AddType(Enum);
 
+	AddTypeDefinition(SourceFile, Enum, InputLine);
+
 	// Validate the metadata for the enum
 	ValidateMetaDataFormat(Enum, EnumToken.MetaData);
 
