Better error message for when a TArray property has an allocator.

diff --git a/Engine/Source/Programs/UnrealHeaderTool/Private/HeaderParser.cpp b/Engine/Source/Programs/UnrealHeaderTool/Private/HeaderParser.cpp
--- a/Engine/Source/Programs/UnrealHeaderTool/Private/HeaderParser.cpp
+++ b/Engine/Source/Programs/UnrealHeaderTool/Private/HeaderParser.cpp
@@ -2652,7 +2652,30 @@ bool FHeaderParser::GetVarType
 		OriginalVarTypeFlags |= VarProperty.PropertyFlags & (CPF_ContainsInstancedReference | CPF_InstancedReference); // propagate these to the array, we will fix them later
 		VarType.PropertyFlags = OriginalVarTypeFlags;
 		VarProperty.ArrayType = EArrayType::Dynamic;
-		RequireSymbol( TEXT(">"), TEXT("'tarray'"), ESymbolParseOption::CloseTemplateBracket );
+
+		FToken CloseTemplateToken;
+		if (!GetToken(CloseTemplateToken, /*bNoConsts=*/ true, ESymbolParseOption::CloseTemplateBracket))
+		{
+			FError::Throwf(TEXT("Missing token while parsing TArray."));
+		}
+
+		if (CloseTemplateToken.TokenType != TOKEN_Symbol || FCString::Stricmp(CloseTemplateToken.Identifier, TEXT(">")))
+		{
+			// If we didn't find a comma, report it
+			if (FCString::Stricmp(CloseTemplateToken.Identifier, TEXT(",")))
+			{
+				FError::Throwf(TEXT("Expected '>' but found '%s'"), CloseTemplateToken.Identifier);
+			}
+
+			// If we found a comma, read the next thing, assume it's an allocator, and report that
+			FToken AllocatorToken;
+			if (!GetToken(AllocatorToken, /*bNoConsts=*/ true, ESymbolParseOption::CloseTemplateBracket))
+			{
+				FError::Throwf(TEXT("Expected '>' but found '%s'"), CloseTemplateToken.Identifier);
+			}
+
+			FError::Throwf(TEXT("Found '%s' - explicit allocators are not supported in TArray properties."), AllocatorToken.Identifier);
+		}
 	}
 	else if ( VarType.Matches(TEXT("FString")) )
 	{
