Fixing UnrealHeaderTool not respecting class flags on unparsed classes used as property types in other classes.

- Fixed some properties requiring 'instanced' keyword even though their type was marked as 'defaulttoinstanced"
- Removed 'Placeable' keyword from ALandscape (it has long been deprecated)

diff --git a/Engine/Source/Programs/UnrealHeaderTool/Private/ClassMaps.cpp b/Engine/Source/Programs/UnrealHeaderTool/Private/ClassMaps.cpp
--- a/Engine/Source/Programs/UnrealHeaderTool/Private/ClassMaps.cpp
+++ b/Engine/Source/Programs/UnrealHeaderTool/Private/ClassMaps.cpp
@@ -4,17 +4,18 @@
 
 #include "ClassMaps.h"
 
-TMap<FString, TSharedRef<FUnrealSourceFile> >						GUnrealSourceFilesMap;
-TMap<UField*, TSharedRef<FUnrealTypeDefinitionInfo> >				GTypeDefinitionInfoMap;
-TMap<UClass*, FString>												GClassStrippedHeaderTextMap;
-TMap<UClass*, FString>												GClassHeaderNameWithNoPathMap;
-TSet<UClass*>														GPublicClassSet;
-TSet<FUnrealSourceFile*>											GPublicSourceFileSet;
-TSet<FUnrealSourceFile*>											GExportedSourceFiles;
-TMap<UProperty*, FString>											GArrayDimensions;
-TMap<UPackage*,  const FManifestModule*>							GPackageToManifestModuleMap;
-TMap<UField*, uint32>												GGeneratedCodeCRCs;
-TMap<UEnum*,  EPropertyType>										GEnumUnderlyingTypes;
+TMap<FString, TSharedRef<FUnrealSourceFile> > GUnrealSourceFilesMap;
+TMap<UField*, TSharedRef<FUnrealTypeDefinitionInfo> > GTypeDefinitionInfoMap;
+TMap<UClass*, FString> GClassStrippedHeaderTextMap;
+TMap<UClass*, FString> GClassHeaderNameWithNoPathMap;
+TSet<UClass*> GPublicClassSet;
+TSet<FUnrealSourceFile*> GPublicSourceFileSet;
+TSet<FUnrealSourceFile*> GExportedSourceFiles;
+TMap<UProperty*, FString> GArrayDimensions;
+TMap<UPackage*,  const FManifestModule*> GPackageToManifestModuleMap;
+TMap<UField*, uint32> GGeneratedCodeCRCs;
+TMap<UEnum*,  EPropertyType> GEnumUnderlyingTypes;
+TMap<FName, TSharedRef<FClassDeclarationMetaData> > GClassDeclarations;
 
 TSharedRef<FUnrealTypeDefinitionInfo> AddTypeDefinition(FUnrealSourceFile& SourceFile, UField* Field, int32 Line)
 {
diff --git a/Engine/Source/Programs/UnrealHeaderTool/Private/ClassMaps.h b/Engine/Source/Programs/UnrealHeaderTool/Private/ClassMaps.h
--- a/Engine/Source/Programs/UnrealHeaderTool/Private/ClassMaps.h
+++ b/Engine/Source/Programs/UnrealHeaderTool/Private/ClassMaps.h
@@ -11,17 +11,35 @@
 
 struct FManifestModule;
 
-extern TMap<FString, TSharedRef<FUnrealSourceFile> >					GUnrealSourceFilesMap;
-extern TMap<UField*, TSharedRef<FUnrealTypeDefinitionInfo> >			GTypeDefinitionInfoMap;
-extern TMap<UClass*, FString>											GClassStrippedHeaderTextMap;
-extern TMap<UClass*, FString>											GClassHeaderNameWithNoPathMap;
-extern TSet<FUnrealSourceFile*>											GExportedSourceFiles;
-extern TSet<UClass*>													GPublicClassSet;
-extern TSet<FUnrealSourceFile*>											GPublicSourceFileSet;
-extern TMap<UProperty*, FString>										GArrayDimensions;
-extern TMap<UPackage*,  const FManifestModule*>							GPackageToManifestModuleMap;
-extern TMap<UField*, uint32>											GGeneratedCodeCRCs;
-extern TMap<UEnum*,  EPropertyType>										GEnumUnderlyingTypes;
+/** Structure that holds class meta data generated from its UCLASS declaration */
+struct FClassDeclarationMetaData
+{
+	uint64 ClassFlags;
+	TMap<FName, FString> MetaData;
+	FString ClassWithin;
+	FString ConfigName;
+	TArray<FString> HideCategories;
+	TArray<FString> ShowSubCatgories;
+	TArray<FString> HideFunctions;
+	TArray<FString> AutoExpandCategories;
+	TArray<FString> AutoCollapseCategories;
+	TArray<FString> DependsOn;
+	TArray<FString> ClassGroupNames;
+};
+
+
+extern TMap<FString, TSharedRef<FUnrealSourceFile> > GUnrealSourceFilesMap;
+extern TMap<UField*, TSharedRef<FUnrealTypeDefinitionInfo> > GTypeDefinitionInfoMap;
+extern TMap<UClass*, FString> GClassStrippedHeaderTextMap;
+extern TMap<UClass*, FString> GClassHeaderNameWithNoPathMap;
+extern TSet<FUnrealSourceFile*> GExportedSourceFiles;
+extern TSet<UClass*> GPublicClassSet;
+extern TSet<FUnrealSourceFile*> GPublicSourceFileSet;
+extern TMap<UProperty*, FString> GArrayDimensions;
+extern TMap<UPackage*,  const FManifestModule*> GPackageToManifestModuleMap;
+extern TMap<UField*, uint32> GGeneratedCodeCRCs;
+extern TMap<UEnum*,  EPropertyType> GEnumUnderlyingTypes;
+extern TMap<FName, TSharedRef<FClassDeclarationMetaData> > GClassDeclarations;
 
 /** Types access specifiers. */
 enum EAccessSpecifier
diff --git a/Engine/Source/Programs/UnrealHeaderTool/Private/CodeGenerator.cpp b/Engine/Source/Programs/UnrealHeaderTool/Private/CodeGenerator.cpp
--- a/Engine/Source/Programs/UnrealHeaderTool/Private/CodeGenerator.cpp
+++ b/Engine/Source/Programs/UnrealHeaderTool/Private/CodeGenerator.cpp
@@ -4920,7 +4920,9 @@ ECompilationResult::Type UnrealHeaderTool_Main(const FString& ModuleInfoFilename
 				(CurrentlyProcessing == PublicHeaders       ) ? Module.PublicUObjectHeaders        :
 				                                                Module.PrivateUObjectHeaders;
 			if (!UObjectHeaders.Num())
+			{
 				continue;
+			}
 
 			for (const FString& Filename : UObjectHeaders)
 			{
diff --git a/Engine/Source/Programs/UnrealHeaderTool/Private/HeaderParser.cpp b/Engine/Source/Programs/UnrealHeaderTool/Private/HeaderParser.cpp
--- a/Engine/Source/Programs/UnrealHeaderTool/Private/HeaderParser.cpp
+++ b/Engine/Source/Programs/UnrealHeaderTool/Private/HeaderParser.cpp
@@ -4115,262 +4115,135 @@ void PostParsingClassSetup(UClass* Class)
 	}
 }
 
-/**
- * Compiles a class declaration.
- */
-void FHeaderParser::CompileClassDeclaration(FClasses& AllClasses)
+void FHeaderParser::ParseClassProperties(const TArray<FPropertySpecifier>& InClassSpecifiers, const FString& InRequiredAPIMacroIfPresent, FClassDeclarationMetaData& OutClassData)
 {
-	// Start of a class block.
-	CheckAllow(TEXT("'class'"), ALLOW_Class);
-
-	// New-style UCLASS() syntax
-	TMap<FName, FString> MetaData;
-
-	TArray<FPropertySpecifier> SpecifiersFound;
-	ReadSpecifierSetInsideMacro(SpecifiersFound, TEXT("Class"), MetaData);
-
-	auto PrologFinishLine = InputLine;
-
-	AddFormattedPrevCommentAsTooltipMetaData(MetaData);
-
-	// New style files have the class name / extends afterwards
-	RequireIdentifier(TEXT("class"), TEXT("Class declaration"));
-
-	SkipDeprecatedMacroIfNecessary();
-
-	FString DeclaredClassName;
-	FString RequiredAPIMacroIfPresent;
-	
-	FClass* Class = ParseClassNameDeclaration(AllClasses, /*out*/ DeclaredClassName, /*out*/ RequiredAPIMacroIfPresent);
-	ClassDefinitionRanges.Add(Class, ClassDefinitionRange(&Input[InputPos], nullptr));
-
-	check(Class);
-
-	Class->ClassFlags |= CLASS_Parsed;
-
-	PushNest(ENestType::NEST_Class, Class);
-
-	uint32 PreviousClassFlags = Class->ClassFlags;
-
-	ResetClassData();
-
-	// Verify class variables haven't been filled in
-	check(Class->Children == NULL);
-	check(Class->Next == NULL);
-	check(Class->NetFields.Num() == 0);
-
-	// Make sure our parent classes is parsed.
-	for (UClass* Temp = Class->GetSuperClass(); Temp; Temp = Temp->GetSuperClass())
-	{
-		if (!(Temp->ClassFlags & (CLASS_Parsed | CLASS_Intrinsic)))
-		{
-			FError::Throwf(TEXT("'%s' can't be compiled: Parent class '%s' has errors"), *Class->GetName(), *Temp->GetName());
-		}
-	}
-
-	// Get categories inherited from the parent.
-	TArray<FString> HideCategories;
-	TArray<FString> ShowSubCatgories;
-	TArray<FString> HideFunctions;
-	TArray<FString> AutoExpandCategories;
-	TArray<FString> AutoCollapseCategories;
-	Class->GetHideCategories(HideCategories);
-	Class->GetShowCategories(ShowSubCatgories);
-	Class->GetHideFunctions(HideFunctions);
-	Class->GetAutoExpandCategories(AutoExpandCategories);
-	Class->GetAutoCollapseCategories(AutoCollapseCategories);
-
-	// Class attributes.
-	FClassMetaData* ClassData = GScriptHelper.FindClassData(Class);
-	check(ClassData);
-
-	ClassData->SetPrologLine(PrologFinishLine);
-
+	OutClassData.ClassFlags = 0;
 	// Record that this class is RequiredAPI if the CORE_API style macro was present
-	if (!RequiredAPIMacroIfPresent.IsEmpty())
-	{
-		Class->ClassFlags |= CLASS_RequiredAPI;
-	}
-
-	// All classes that are parsed are expected to be native
-	UClass* Super = Class->GetSuperClass();
-	if (Super && !Super->HasAnyClassFlags(CLASS_Native))
+	if (!InRequiredAPIMacroIfPresent.IsEmpty())
 	{
-		FError::Throwf(TEXT("Native classes cannot extend non-native classes") );
+		OutClassData.ClassFlags |= CLASS_RequiredAPI;
 	}
-
-	Class->SetFlags(RF_Native);
-	Class->ClassFlags |= CLASS_Native;
+	OutClassData.ClassFlags |= CLASS_Native;
 
 	// Process all of the class specifiers
-	TArray<FString> ClassGroupNames;
+	
 	bool bWithinSpecified = false;
 	bool bDeclaresConfigFile = false;
-	for (const FPropertySpecifier& PropSpecifier : SpecifiersFound)
+	for (const FPropertySpecifier& PropSpecifier : InClassSpecifiers)
 	{
 		const FString& Specifier = PropSpecifier.Key;
 
 		if (Specifier == TEXT("noexport"))
 		{
 			// Don't export to C++ header.
-			Class->ClassFlags |= CLASS_NoExport;
+			OutClassData.ClassFlags |= CLASS_NoExport;
 		}
 		else if (Specifier == TEXT("intrinsic"))
 		{
-			Class->ClassFlags |= CLASS_Intrinsic;
+			OutClassData.ClassFlags |= CLASS_Intrinsic;
 		}
 		else if (Specifier == TEXT("ComponentWrapperClass"))
 		{
-			MetaData.Add(TEXT("IgnoreCategoryKeywordsInSubclasses"), TEXT("true"));
+			OutClassData.MetaData.Add(TEXT("IgnoreCategoryKeywordsInSubclasses"), TEXT("true"));
 		}
 		else if (Specifier == TEXT("within"))
 		{
-			FString WithinNameStr = RequireExactlyOneSpecifierValue(PropSpecifier);
-
-			UClass* RequiredWithinClass = AllClasses.FindClass(*WithinNameStr);
-
-			if (!RequiredWithinClass)
-			{
-				FError::Throwf(TEXT("Within class '%s' not found."), *WithinNameStr);
-			}
-
-			if (RequiredWithinClass->IsChildOf(UInterface::StaticClass()))
-			{
-				FError::Throwf(TEXT("Classes cannot be 'within' interfaces"));
-			}
-			else if (Class->ClassWithin == NULL || Class->ClassWithin == UObject::StaticClass() || RequiredWithinClass->IsChildOf(Class->ClassWithin))
-			{
-				Class->ClassWithin = RequiredWithinClass;
-			}
-			else if (Class->ClassWithin != RequiredWithinClass)
-			{
-				FError::Throwf(TEXT("%s must be within %s, not %s"), *Class->GetPathName(), *Class->ClassWithin->GetPathName(), *RequiredWithinClass->GetPathName());
-			}
-
-			bWithinSpecified = true;
+			OutClassData.ClassWithin = RequireExactlyOneSpecifierValue(PropSpecifier);
 		}
 		else if (Specifier == TEXT("editinlinenew"))
 		{
-			// don't allow actor classes to be declared editinlinenew
-			if (IsActorClass(Class))
-			{
-				FError::Throwf(TEXT("Invalid class attribute: Creating actor instances via the property window is not allowed"));
-			}
-
 			// Class can be constructed from the New button in editinline
-			Class->ClassFlags |= CLASS_EditInlineNew;
+			OutClassData.ClassFlags |= CLASS_EditInlineNew;
 		}
 		else if (Specifier == TEXT("noteditinlinenew"))
 		{
 			// Class cannot be constructed from the New button in editinline
-			Class->ClassFlags &= ~CLASS_EditInlineNew;
+			OutClassData.ClassFlags &= ~CLASS_EditInlineNew;
 		}
 		else if (Specifier == TEXT("placeable"))
 		{
-			if (!(Class->ClassFlags & CLASS_NotPlaceable))
+			if (!(OutClassData.ClassFlags & CLASS_NotPlaceable))
 			{
 				FError::Throwf(TEXT("The placeable specifier is deprecated. Classes are assumed to be placeable by default."));
 			}
-			Class->ClassFlags &= ~CLASS_NotPlaceable;
+			OutClassData.ClassFlags &= ~CLASS_NotPlaceable;
 		}
 		else if (Specifier == TEXT("defaulttoinstanced"))
 		{
 			// these classed default to instanced.
-			Class->ClassFlags |= CLASS_DefaultToInstanced;
+			OutClassData.ClassFlags |= CLASS_DefaultToInstanced;
 		}
 		else if (Specifier == TEXT("notplaceable"))
 		{
 			// Don't allow the class to be placed in the editor.
-			Class->ClassFlags |= CLASS_NotPlaceable;
+			OutClassData.ClassFlags |= CLASS_NotPlaceable;
 		}
 		else if (Specifier == TEXT("hidedropdown"))
 		{
 			// Prevents class from appearing in class comboboxes in the property window
-			Class->ClassFlags |= CLASS_HideDropDown;
+			OutClassData.ClassFlags |= CLASS_HideDropDown;
 		}
 		else if (Specifier == TEXT("dependsOn"))
 		{
 			FError::Throwf(TEXT("The dependsOn specifier is deprecated. Please use proper #include instead."));
-
-			// Make sure the syntax matches but don't do anything with it;
-			RequireSpecifierValue(PropSpecifier);
 		}
 		else if (Specifier == TEXT("MinimalAPI"))
 		{
-			Class->ClassFlags |= CLASS_MinimalAPI;
+			OutClassData.ClassFlags |= CLASS_MinimalAPI;
 		}
 		else if (Specifier == TEXT("const"))
 		{
-			Class->ClassFlags |= CLASS_Const;
+			OutClassData.ClassFlags |= CLASS_Const;
 		}
 		else if (Specifier == TEXT("perObjectConfig"))
 		{
-			Class->ClassFlags |= CLASS_PerObjectConfig;
+			OutClassData.ClassFlags |= CLASS_PerObjectConfig;
 		}
 		else if (Specifier == TEXT("configdonotcheckdefaults"))
 		{
-			Class->ClassFlags |= CLASS_ConfigDoNotCheckDefaults;
+			OutClassData.ClassFlags |= CLASS_ConfigDoNotCheckDefaults;
 		}
 		else if (Specifier == TEXT("abstract"))
 		{
 			// Hide all editable properties.
-			Class->ClassFlags |= CLASS_Abstract;
+			OutClassData.ClassFlags |= CLASS_Abstract;
 		}
 		else if (Specifier == TEXT("deprecated"))
 		{
-			Class->ClassFlags |= CLASS_Deprecated;
+			OutClassData.ClassFlags |= CLASS_Deprecated;
 
 			// Don't allow the class to be placed in the editor.
-			Class->ClassFlags |= CLASS_NotPlaceable;
+			OutClassData.ClassFlags |= CLASS_NotPlaceable;
 		}
 		else if (Specifier == TEXT("transient"))
 		{
 			// Transient class.
-			Class->ClassFlags |= CLASS_Transient;
+			OutClassData.ClassFlags |= CLASS_Transient;
 		}
 		else if (Specifier == TEXT("nonTransient"))
 		{
 			// this child of a transient class is not transient - remove the transient flag
-			Class->ClassFlags &= ~CLASS_Transient;
+			OutClassData.ClassFlags &= ~CLASS_Transient;
 		}
 		else if (Specifier == TEXT("customConstructor"))
 		{
 			// we will not export a constructor for this class, assuming it is in the CPP block
-			Class->ClassFlags |= CLASS_CustomConstructor;
+			OutClassData.ClassFlags |= CLASS_CustomConstructor;
 		}
 		else if (Specifier == TEXT("config"))
 		{
 			// Class containing config properties - parse the name of the config file to use
-			FString ConfigNameStr = RequireExactlyOneSpecifierValue(PropSpecifier);
-
-			// if the user specified "inherit", we're just going to use the parent class's config filename
-			// this is not actually necessary but it can be useful for explicitly communicating config-ness
-			if (ConfigNameStr == TEXT("inherit"))
-			{
-				UClass* SuperClass = Class->GetSuperClass();
-				if (!SuperClass)
-					FError::Throwf(TEXT("Cannot inherit config filename: %s has no super class"), *Class->GetName());
-
-				if (SuperClass->ClassConfigName == NAME_None)
-					FError::Throwf(TEXT("Cannot inherit config filename: parent class %s is not marked config."), *SuperClass->GetPathName());
-			}
-			else
-			{
-				// otherwise, set the config name to the parsed identifier
-				Class->ClassConfigName = FName(*ConfigNameStr);
-			}
-
-			bDeclaresConfigFile = true;
+			OutClassData.ConfigName = RequireExactlyOneSpecifierValue(PropSpecifier);
 		}
 		else if (Specifier == TEXT("defaultconfig"))
 		{
 			// Save object config only to Default INIs, never to local INIs.
-			Class->ClassFlags |= CLASS_DefaultConfig;
+			OutClassData.ClassFlags |= CLASS_DefaultConfig;
 		}
 		else if (Specifier == TEXT("globaluserconfig"))
 		{
 			// Save object config only to global user overrides, never to local INIs
-			Class->ClassFlags |= CLASS_GlobalUserConfig;
+			OutClassData.ClassFlags |= CLASS_GlobalUserConfig;
 		}
 		else if (Specifier == TEXT("showCategories"))
 		{
@@ -4379,20 +4252,20 @@ void FHeaderParser::CompileClassDeclaration(FClasses& AllClasses)
 			for (const FString& Value : PropSpecifier.Values)
 			{
 				// if we didn't find this specific category path in the HideCategories metadata
-				if (HideCategories.Remove(Value) == 0)
+				if (OutClassData.HideCategories.Remove(Value) == 0)
 				{
 					TArray<FString> SubCategoryList;
 					Value.ParseIntoArray(SubCategoryList, TEXT("|"), true);
 
 					FString SubCategoryPath;
 					// look to see if any of the parent paths are excluded in the HideCategories list
-					for (int32 CategoryPathIndex = 0; CategoryPathIndex < SubCategoryList.Num()-1; ++CategoryPathIndex)
+					for (int32 CategoryPathIndex = 0; CategoryPathIndex < SubCategoryList.Num() - 1; ++CategoryPathIndex)
 					{
 						SubCategoryPath += SubCategoryList[CategoryPathIndex];
 						// if we're hiding a parent category, then we need to flag this sub category for show
-						if (HideCategories.Contains(SubCategoryPath))
+						if (OutClassData.HideCategories.Contains(SubCategoryPath))
 						{
-							ShowSubCatgories.AddUnique(Value);
+							OutClassData.ShowSubCatgories.AddUnique(Value);
 							break;
 						}
 						SubCategoryPath += "|";
@@ -4406,7 +4279,7 @@ void FHeaderParser::CompileClassDeclaration(FClasses& AllClasses)
 
 			for (const FString& Value : PropSpecifier.Values)
 			{
-				HideCategories.AddUnique(Value);
+				OutClassData.HideCategories.AddUnique(Value);
 			}
 		}
 		else if (Specifier == TEXT("showFunctions"))
@@ -4415,7 +4288,7 @@ void FHeaderParser::CompileClassDeclaration(FClasses& AllClasses)
 
 			for (const FString& Value : PropSpecifier.Values)
 			{
-				HideFunctions.Remove(Value);
+				OutClassData.HideFunctions.Remove(Value);
 			}
 		}
 		else if (Specifier == TEXT("hideFunctions"))
@@ -4424,7 +4297,7 @@ void FHeaderParser::CompileClassDeclaration(FClasses& AllClasses)
 
 			for (const FString& Value : PropSpecifier.Values)
 			{
-				HideFunctions.AddUnique(Value);
+				OutClassData.HideFunctions.AddUnique(Value);
 			}
 		}
 		else if (Specifier == TEXT("classGroup"))
@@ -4433,7 +4306,7 @@ void FHeaderParser::CompileClassDeclaration(FClasses& AllClasses)
 
 			for (const FString& Value : PropSpecifier.Values)
 			{
-				ClassGroupNames.Add(Value);
+				OutClassData.ClassGroupNames.Add(Value);
 			}
 		}
 		else if (Specifier == TEXT("autoExpandCategories"))
@@ -4442,8 +4315,8 @@ void FHeaderParser::CompileClassDeclaration(FClasses& AllClasses)
 
 			for (const FString& Value : PropSpecifier.Values)
 			{
-				AutoCollapseCategories.Remove   (Value);
-				AutoExpandCategories  .AddUnique(Value);
+				OutClassData.AutoCollapseCategories.Remove(Value);
+				OutClassData.AutoExpandCategories.AddUnique(Value);
 			}
 		}
 		else if (Specifier == TEXT("autoCollapseCategories"))
@@ -4452,8 +4325,8 @@ void FHeaderParser::CompileClassDeclaration(FClasses& AllClasses)
 
 			for (const FString& Value : PropSpecifier.Values)
 			{
-				AutoExpandCategories  .Remove   (Value);
-				AutoCollapseCategories.AddUnique(Value);
+				OutClassData.AutoExpandCategories.Remove(Value);
+				OutClassData.AutoCollapseCategories.AddUnique(Value);
 			}
 		}
 		else if (Specifier == TEXT("dontAutoCollapseCategories"))
@@ -4462,40 +4335,178 @@ void FHeaderParser::CompileClassDeclaration(FClasses& AllClasses)
 
 			for (const FString& Value : PropSpecifier.Values)
 			{
-				AutoCollapseCategories.Remove(Value);
+				OutClassData.AutoCollapseCategories.Remove(Value);
 			}
 		}
 		else if (Specifier == TEXT("collapseCategories"))
 		{
 			// Class' properties should not be shown categorized in the editor.
-			Class->ClassFlags |= CLASS_CollapseCategories;
+			OutClassData.ClassFlags |= CLASS_CollapseCategories;
 		}
 		else if (Specifier == TEXT("dontCollapseCategories"))
 		{
 			// Class' properties should be shown categorized in the editor.
-			Class->ClassFlags &= ~CLASS_CollapseCategories;
+			OutClassData.ClassFlags &= ~CLASS_CollapseCategories;
 		}
 		else if (Specifier == TEXT("AdvancedClassDisplay"))
 		{
 			// By default the class properties are shown in advanced sections in UI
-			Class->ClassFlags |= CLASS_AdvancedDisplay;
+			OutClassData.ClassFlags |= CLASS_AdvancedDisplay;
 		}
 		else if (Specifier == TEXT("ConversionRoot"))
 		{
-			MetaData.Add(FName(TEXT("IsConversionRoot")), "true");
+			OutClassData.MetaData.Add(FName(TEXT("IsConversionRoot")), "true");
 		}
 		else
 		{
 			FError::Throwf(TEXT("Unknown class specifier '%s'"), *Specifier);
 		}
 	}
+}
+
+/**
+ * Compiles a class declaration.
+ */
+void FHeaderParser::CompileClassDeclaration(FClasses& AllClasses)
+{
+	// Start of a class block.
+	CheckAllow(TEXT("'class'"), ALLOW_Class);
+
+	// New-style UCLASS() syntax
+	TMap<FName, FString> MetaData;
+
+	TArray<FPropertySpecifier> SpecifiersFound;
+	ReadSpecifierSetInsideMacro(SpecifiersFound, TEXT("Class"), MetaData);
+
+	auto PrologFinishLine = InputLine;
+
+	AddFormattedPrevCommentAsTooltipMetaData(MetaData);
+
+	// New style files have the class name / extends afterwards
+	RequireIdentifier(TEXT("class"), TEXT("Class declaration"));
+
+	SkipDeprecatedMacroIfNecessary();
+
+	FString DeclaredClassName;
+	FString RequiredAPIMacroIfPresent;
+	
+	FClass* Class = ParseClassNameDeclaration(AllClasses, /*out*/ DeclaredClassName, /*out*/ RequiredAPIMacroIfPresent);
+	check(Class);
+	auto ClassDeclarationData = GClassDeclarations.FindChecked(Class->GetFName());
+
+	ClassDefinitionRanges.Add(Class, ClassDefinitionRange(&Input[InputPos], nullptr));
+
+	check(Class->ClassFlags == 0 || (Class->ClassFlags & ClassDeclarationData->ClassFlags) != 0);
+
+	Class->ClassFlags |= CLASS_Parsed;
+
+	PushNest(ENestType::NEST_Class, Class);
+	
+	uint32 PreviousClassFlags = Class->ClassFlags;	
+
+	ResetClassData();
+
+	// Verify class variables haven't been filled in
+	check(Class->Children == NULL);
+	check(Class->Next == NULL);
+	check(Class->NetFields.Num() == 0);
+
+	// Make sure our parent classes is parsed.
+	for (UClass* Temp = Class->GetSuperClass(); Temp; Temp = Temp->GetSuperClass())
+	{
+		if (!(Temp->ClassFlags & (CLASS_Parsed | CLASS_Intrinsic)))
+		{
+			FError::Throwf(TEXT("'%s' can't be compiled: Parent class '%s' has errors"), *Class->GetName(), *Temp->GetName());
+		}
+	}
 
-	if (ClassGroupNames       .Num()) { MetaData.Add("ClassGroupNames",        FString::Join(ClassGroupNames,        TEXT(" "))); }
-	if (AutoCollapseCategories.Num()) { MetaData.Add("AutoCollapseCategories", FString::Join(AutoCollapseCategories, TEXT(" "))); }
-	if (HideCategories        .Num()) { MetaData.Add("HideCategories",         FString::Join(HideCategories,         TEXT(" "))); }
-	if (ShowSubCatgories      .Num()) { MetaData.Add("ShowCategories",         FString::Join(ShowSubCatgories,       TEXT(" "))); }
-	if (HideFunctions         .Num()) { MetaData.Add("HideFunctions",          FString::Join(HideFunctions,          TEXT(" "))); }
-	if (AutoExpandCategories  .Num()) { MetaData.Add("AutoExpandCategories",   FString::Join(AutoExpandCategories,   TEXT(" "))); }
+	// Get categories inherited from the parent.
+	Class->GetHideCategories(ClassDeclarationData->HideCategories);
+	Class->GetShowCategories(ClassDeclarationData->ShowSubCatgories);
+	Class->GetHideFunctions(ClassDeclarationData->HideFunctions);
+	Class->GetAutoExpandCategories(ClassDeclarationData->AutoExpandCategories);
+	Class->GetAutoCollapseCategories(ClassDeclarationData->AutoCollapseCategories);
+
+	// Class attributes.
+	FClassMetaData* ClassData = GScriptHelper.FindClassData(Class);
+	check(ClassData);
+
+	ClassData->SetPrologLine(PrologFinishLine);
+
+	Class->ClassFlags |= ClassDeclarationData->ClassFlags;
+
+	// All classes that are parsed are expected to be native
+	UClass* Super = Class->GetSuperClass();
+	if (Super && !Super->HasAnyClassFlags(CLASS_Native))
+	{
+		FError::Throwf(TEXT("Native classes cannot extend non-native classes") );
+	}
+
+	Class->SetFlags(RF_Native);
+
+	// Process all of the class specifiers
+	bool bWithinSpecified = false;
+	bool bDeclaresConfigFile = false;
+	if (ClassDeclarationData->ClassWithin.IsEmpty() == false)
+	{
+		UClass* RequiredWithinClass = AllClasses.FindClass(*ClassDeclarationData->ClassWithin);
+		if (!RequiredWithinClass)
+		{
+			FError::Throwf(TEXT("Within class '%s' not found."), *ClassDeclarationData->ClassWithin);
+		}
+		if (RequiredWithinClass->IsChildOf(UInterface::StaticClass()))
+		{
+			FError::Throwf(TEXT("Classes cannot be 'within' interfaces"));
+		}
+		else if (Class->ClassWithin == NULL || Class->ClassWithin == UObject::StaticClass() || RequiredWithinClass->IsChildOf(Class->ClassWithin))
+		{
+			Class->ClassWithin = RequiredWithinClass;
+		}
+		else if (Class->ClassWithin != RequiredWithinClass)
+		{
+			FError::Throwf(TEXT("%s must be within %s, not %s"), *Class->GetPathName(), *Class->ClassWithin->GetPathName(), *RequiredWithinClass->GetPathName());
+		}
+		bWithinSpecified = true;
+	}
+	if (!!(Class->ClassFlags & CLASS_EditInlineNew))
+	{
+		// don't allow actor classes to be declared editinlinenew
+		if (IsActorClass(Class))
+		{
+			FError::Throwf(TEXT("Invalid class attribute: Creating actor instances via the property window is not allowed"));
+		}
+	}
+	if (ClassDeclarationData->ConfigName.IsEmpty() == false)
+	{
+		// if the user specified "inherit", we're just going to use the parent class's config filename
+		// this is not actually necessary but it can be useful for explicitly communicating config-ness
+		if (ClassDeclarationData->ConfigName == TEXT("inherit"))
+		{
+			UClass* SuperClass = Class->GetSuperClass();
+			if (!SuperClass)
+			{
+				FError::Throwf(TEXT("Cannot inherit config filename: %s has no super class"), *Class->GetName());
+			}
+
+			if (SuperClass->ClassConfigName == NAME_None)
+			{
+				FError::Throwf(TEXT("Cannot inherit config filename: parent class %s is not marked config."), *SuperClass->GetPathName());
+			}
+		}
+		else
+		{
+			// otherwise, set the config name to the parsed identifier
+			Class->ClassConfigName = FName(*ClassDeclarationData->ConfigName);
+		}
+		bDeclaresConfigFile = true;
+	}
+
+	if (ClassDeclarationData->ClassGroupNames.Num()) { MetaData.Add("ClassGroupNames", FString::Join(ClassDeclarationData->ClassGroupNames, TEXT(" "))); }
+	if (ClassDeclarationData->AutoCollapseCategories.Num()) { MetaData.Add("AutoCollapseCategories", FString::Join(ClassDeclarationData->AutoCollapseCategories, TEXT(" "))); }
+	if (ClassDeclarationData->HideCategories.Num()) { MetaData.Add("HideCategories", FString::Join(ClassDeclarationData->HideCategories, TEXT(" "))); }
+	if (ClassDeclarationData->ShowSubCatgories.Num()) { MetaData.Add("ShowCategories", FString::Join(ClassDeclarationData->ShowSubCatgories, TEXT(" "))); }
+	if (ClassDeclarationData->HideFunctions.Num()) { MetaData.Add("HideFunctions", FString::Join(ClassDeclarationData->HideFunctions, TEXT(" "))); }
+	if (ClassDeclarationData->AutoExpandCategories.Num()) { MetaData.Add("AutoExpandCategories", FString::Join(ClassDeclarationData->AutoExpandCategories, TEXT(" "))); }
 
 	// Make sure both RequiredAPI and MinimalAPI aren't specified
 	if (Class->HasAllClassFlags(CLASS_MinimalAPI | CLASS_RequiredAPI))
@@ -7209,6 +7220,17 @@ void FHeaderPreParser::ParseClassDeclaration(const TCHAR* InputText, int32 InLin
 	FString RequiredAPIMacroIfPresent;
 	ParseNameWithPotentialAPIMacroPrefix(/*out*/ out_ClassName, /*out*/ RequiredAPIMacroIfPresent, StartingMatchID);
 
+	FName ClassNameWithoutPrefix(*GetClassNameWithPrefixRemoved(out_ClassName));
+	auto DeclarationDataPtr = GClassDeclarations.Find(ClassNameWithoutPrefix);
+	if (!DeclarationDataPtr)
+	{
+		TSharedRef<FClassDeclarationMetaData> DeclarationData = MakeShareable(new FClassDeclarationMetaData());
+		DeclarationData->MetaData = MetaData;
+		DeclarationData->ClassFlags = 0;
+		FHeaderParser::ParseClassProperties(SpecifiersFound, RequiredAPIMacroIfPresent, *DeclarationData);
+		GClassDeclarations.Add(ClassNameWithoutPrefix, DeclarationData);
+	}
+
 	// Handle inheritance
 	if (MatchSymbol(TEXT(":")))
 	{
@@ -7246,6 +7268,7 @@ void FHeaderPreParser::ParseClassDeclaration(const TCHAR* InputText, int32 InLin
 	{
 		if (Specifier.Key == TEXT("DependsOn"))
 		{
+			check(false);
 			for (auto It2 = Specifier.Values.CreateConstIterator(); It2; ++It2)
 			{
 				const FString& Value = *It2;
@@ -7808,4 +7831,27 @@ UFunction* FHeaderParser::CreateFunction(const FFuncInfo &FuncInfo) const
 UDelegateFunction* FHeaderParser::CreateDelegateFunction(const FFuncInfo &FuncInfo) const
 {
 	return CreateFunctionImpl<UDelegateFunction>(FuncInfo, IsInAClass() ? (UObject*)GetCurrentClass() : (UObject*)GetCurrentFileScope()->GetSourceFile()->GetPackage(), GetCurrentScope());
+}
+
+bool FHeaderParser::DoesAnythingInHierarchyHaveDefaultToInstanced(UClass* TestClass)
+{
+	bool bDefaultToInstanced = false;
+
+	UClass* Search = TestClass;
+	while (!bDefaultToInstanced && (Search != NULL))
+	{
+		bDefaultToInstanced = Search->HasAnyClassFlags(CLASS_DefaultToInstanced);		
+		if (!bDefaultToInstanced && !Search->HasAnyClassFlags(CLASS_Intrinsic | CLASS_Parsed))
+		{			
+			// The class might not have been parsed yet, look for declaration data.
+			auto ClassDeclarationDataPtr = GClassDeclarations.Find(Search->GetFName());
+			if (ClassDeclarationDataPtr)
+			{
+				bDefaultToInstanced = !!((*ClassDeclarationDataPtr)->ClassFlags & CLASS_DefaultToInstanced);
+			}
+		}
+		Search = Search->GetSuperClass();
+	}
+
+	return bDefaultToInstanced;
 }
\ No newline at end of file
diff --git a/Engine/Source/Programs/UnrealHeaderTool/Private/HeaderParser.h b/Engine/Source/Programs/UnrealHeaderTool/Private/HeaderParser.h
--- a/Engine/Source/Programs/UnrealHeaderTool/Private/HeaderParser.h
+++ b/Engine/Source/Programs/UnrealHeaderTool/Private/HeaderParser.h
@@ -233,6 +233,15 @@ public:
 	 */
 	static ECompilationResult::Type ParseHeaders(FClasses& AllClasses, FHeaderParser& HeaderParser, FUnrealSourceFile& SourceFile, bool bParseSubclasses);
 
+	/**
+	 * Parse Class's properties to generate its declaration data.
+	 *
+	 * @param	InClassSpecifiers Class properties collected from its UCLASS macro
+	 * @param	InRequiredAPIMacroIfPresent *_API macro if present (empty otherwise)
+	 * @param	OutClassData Parsed class meta data
+	 */
+	static void ParseClassProperties(const TArray<FPropertySpecifier>& InClassSpecifiers, const FString& InRequiredAPIMacroIfPresent, FClassDeclarationMetaData& OutClassData);
+
 protected:
 	friend struct FScriptLocation;
 
@@ -428,8 +437,8 @@ protected:
 	void ParseParameterList(FClasses& AllClasses, UFunction* Function, bool bExpectCommaBeforeName = false, TMap<FName, FString>* MetaData = NULL);
 
 	// Throws if a specifier value wasn't provided
-	void RequireSpecifierValue(const FPropertySpecifier& Specifier, bool bRequireExactlyOne = false);
-	FString RequireExactlyOneSpecifierValue(const FPropertySpecifier& Specifier);
+	static void RequireSpecifierValue(const FPropertySpecifier& Specifier, bool bRequireExactlyOne = false);
+	static FString RequireExactlyOneSpecifierValue(const FPropertySpecifier& Specifier);
 
 	/**
 	 * Parse rest of the module's source files.
@@ -516,7 +525,7 @@ protected:
 	/**
 	 * Create new delegate function object based on given info structure.
 	 */
-	UDelegateFunction* CreateDelegateFunction(const FFuncInfo &FuncInfo) const;
+	UDelegateFunction* CreateDelegateFunction(const FFuncInfo &FuncInfo) const;	
 
 	void CompileClassDeclaration(FClasses& AllClasses);
 	void CompileDelegateDeclaration(FUnrealSourceFile& SourceFile, FClasses& AllClasses, const TCHAR* DelegateIdentifier, EDelegateSpecifierAction::Type SpecifierAction = EDelegateSpecifierAction::DontParse);
@@ -748,19 +757,7 @@ protected:
 	EFindName GetFindFlagForPropertyName(const TCHAR* PropertyName);
 
 	// Check to see if anything in the class hierarchy passed in has CLASS_DefaultToInstanced
-	static bool DoesAnythingInHierarchyHaveDefaultToInstanced(UClass* TestClass)
-	{
-		bool bDefaultToInstanced = false;
-
-		UClass* Search = TestClass;
-		while (!bDefaultToInstanced && (Search != NULL))
-		{
-			bDefaultToInstanced = Search->HasAnyClassFlags(CLASS_DefaultToInstanced);
-			Search = Search->GetSuperClass();
-		}
-
-		return bDefaultToInstanced;
-	}
+	static bool DoesAnythingInHierarchyHaveDefaultToInstanced(UClass* TestClass);
 
 	static void ValidatePropertyIsDeprecatedIfNecessary(FPropertyBase& VarProperty, FToken* OuterPropertyType);
 
diff --git a/Engine/Source/Runtime/Landscape/Classes/Landscape.h b/Engine/Source/Runtime/Landscape/Classes/Landscape.h
--- a/Engine/Source/Runtime/Landscape/Classes/Landscape.h
+++ b/Engine/Source/Runtime/Landscape/Classes/Landscape.h
@@ -26,7 +26,7 @@ enum ELandscapeSetupErrors
 	LSE_MAX,
 };
 
-UCLASS(Placeable, hidecategories=LandscapeProxy, showcategories=(Display, Movement, Collision, Lighting, LOD, Input), MinimalAPI)
+UCLASS(hidecategories=LandscapeProxy, showcategories=(Display, Movement, Collision, Lighting, LOD, Input), MinimalAPI)
 class ALandscape : public ALandscapeProxy
 {
 	GENERATED_UCLASS_BODY()
