Fully deprecating 'DependsOn' specifier in UHT

diff --git a/Engine/Source/Programs/UnrealHeaderTool/Private/HeaderParser.cpp b/Engine/Source/Programs/UnrealHeaderTool/Private/HeaderParser.cpp
--- a/Engine/Source/Programs/UnrealHeaderTool/Private/HeaderParser.cpp
+++ b/Engine/Source/Programs/UnrealHeaderTool/Private/HeaderParser.cpp
@@ -4185,7 +4185,7 @@ void FHeaderParser::ParseClassProperties(const TArray<FPropertySpecifier>& InCla
 		}
 		else if (Specifier == TEXT("dependsOn"))
 		{
-			FError::Throwf(TEXT("The dependsOn specifier is deprecated. Please use proper #include instead."));
+			FError::Throwf(TEXT("The dependsOn specifier is deprecated. Please use #include \"ClassHeaderFilename.h\" instead."));
 		}
 		else if (Specifier == TEXT("MinimalAPI"))
 		{
@@ -4751,8 +4751,7 @@ void FHeaderParser::CompileInterfaceDeclaration(FClasses& AllClasses)
 
 		if (Specifier == TEXT("DependsOn"))
 		{
-			// Make sure the syntax matches but don't do anything with it
-			RequireSpecifierValue(*SpecifierIt);
+			FError::Throwf(TEXT("The dependsOn specifier is deprecated. Please use #include \"ClassHeaderFilename.h\" instead."));
 		}
 		else if (Specifier == TEXT("MinimalAPI"))
 		{
@@ -6370,53 +6369,6 @@ ECompilationResult::Type FHeaderParser::ParseRestOfModulesSourceFiles(FClasses&
 	return ECompilationResult::Succeeded;
 }
 
-/**
- * Parses source file of subclasses of classes defined in given file.
- *
- * @param HeaderParser Current parser.
- * @param AllClasses Current classes tree.
- * @param SourceFile Given source file.
- *
- * @returns Compilation result enum.
- */
-/*
-ECompilationResult::Type ParseSourceFileSubclasses(FHeaderParser& HeaderParser, FClasses& AllClasses, FUnrealSourceFile& SourceFile)
-{
-	ECompilationResult::Type Result = ECompilationResult::Succeeded;
-
-	for (auto* Class : SourceFile.GetDefinedClasses())
-	{
-		for (auto SubClass : AllClasses.GetDerivedClasses((FClass*)Class))
-		{
-			auto* ClassDefinitionInfo = GTypeDefinitionInfoMap.Find(SubClass);
-
-			if (ClassDefinitionInfo == nullptr)
-			{
-				continue;
-			}
-
-			FUnrealSourceFile& SubclassSourceFile = (*ClassDefinitionInfo)->GetUnrealSourceFile();
-
-			//note: you must always pass in the root tree node here, since we may add new classes to the tree
-			// if an manual dependency (through dependson()) is encountered
-			ECompilationResult::Type ParseResult = FHeaderParser::ParseHeaders(AllClasses, HeaderParser, SubclassSourceFile, true);
-
-			if (ParseResult == ECompilationResult::FailedDueToHeaderChange)
-			{
-				return ParseResult;
-			}
-
-			if (ParseResult == ECompilationResult::OtherCompilationError)
-			{
-				Result = ECompilationResult::OtherCompilationError;
-			}
-		}
-	}
-
-	return Result;
-}
-*/
-
 // Parse Class's annotated headers and optionally its child classes.
 ECompilationResult::Type FHeaderParser::ParseHeaders(FClasses& AllClasses, FHeaderParser& HeaderParser, FUnrealSourceFile& SourceFile, bool bParseSubclasses)
 {
@@ -7224,6 +7176,7 @@ void FHeaderPreParser::ParseClassDeclaration(const TCHAR* InputText, int32 InLin
 	auto DeclarationDataPtr = GClassDeclarations.Find(ClassNameWithoutPrefix);
 	if (!DeclarationDataPtr)
 	{
+		// Add class declaration meta data so that we can access class flags before the class is fully parsed
 		TSharedRef<FClassDeclarationMetaData> DeclarationData = MakeShareable(new FClassDeclarationMetaData());
 		DeclarationData->MetaData = MetaData;
 		DeclarationData->ClassFlags = 0;
@@ -7262,22 +7215,6 @@ void FHeaderPreParser::ParseClassDeclaration(const TCHAR* InputText, int32 InLin
 			out_ClassNames.Add(FHeaderProvider(EHeaderProviderSourceType::ClassName, InterfaceClassName.ToString().Mid(1)));
 		}
 	}
-
-	// Run thru the specifier list looking for dependency links
-	for (const auto& Specifier : SpecifiersFound)
-	{
-		if (Specifier.Key == TEXT("DependsOn"))
-		{
-			check(false);
-			for (auto It2 = Specifier.Values.CreateConstIterator(); It2; ++It2)
-			{
-				const FString& Value = *It2;
-
-				FName DependentOnClassName = FName(*Value);
-				out_ClassNames.Add(FHeaderProvider(EHeaderProviderSourceType::ClassName, DependentOnClassName.ToString()));
-			}
-		}
-	}
 }
 
 //////////////////////////////////////////////////////////////////////////////////////////////////////////////
