Disable deprecation warnings in code generated by UHT.

diff --git a/Engine/Source/Programs/UnrealHeaderTool/Private/CodeGenerator.cpp b/Engine/Source/Programs/UnrealHeaderTool/Private/CodeGenerator.cpp
--- a/Engine/Source/Programs/UnrealHeaderTool/Private/CodeGenerator.cpp
+++ b/Engine/Source/Programs/UnrealHeaderTool/Private/CodeGenerator.cpp
@@ -2174,6 +2174,8 @@ void FNativeClassHeaderGenerator::ExportClassesFromSourceFileWrapper(FUnrealSour
 		LINE_TERMINATOR
 		TEXT("#include \"ObjectBase.h\"") LINE_TERMINATOR
 		LINE_TERMINATOR);
+
+	GeneratedHeaderTextWithCopyright.Logf(TEXT("PRAGMA_DISABLE_DEPRECATION_WARNINGS") LINE_TERMINATOR);
 	GeneratedHeaderTextWithCopyright.Log(*GeneratedHeaderTextBeforeForwardDeclarations);
 
 	TSet<FString> ForwardDeclarationStrings;
@@ -2193,6 +2195,7 @@ void FNativeClassHeaderGenerator::ExportClassesFromSourceFileWrapper(FUnrealSour
 
 	GeneratedHeaderTextWithCopyright.Log(*GeneratedForwardDeclarations);
 	GeneratedHeaderTextWithCopyright.Log(*GeneratedHeaderText);
+	GeneratedHeaderTextWithCopyright.Logf(TEXT("PRAGMA_ENABLE_DEPRECATION_WARNINGS") LINE_TERMINATOR);
 
 	SourceFile.SetGeneratedFilename(ClassHeaderPath);
 	SourceFile.SetHasChanged(SaveHeaderIfChanged(*ClassHeaderPath, *GeneratedHeaderTextWithCopyright));
@@ -4728,6 +4731,9 @@ void FNativeClassHeaderGenerator::ExportGeneratedCPP()
 	GeneratedCPPEpilogue.Logf(
 		LINE_TERMINATOR
 		);
+
+	FString EnableDeprecationWarnings = FString(TEXT("PRAGMA_ENABLE_DEPRECATION_WARNINGS") LINE_TERMINATOR);
+	FString DisableDeprecationWarnings = FString(TEXT("PRAGMA_DISABLE_DEPRECATION_WARNINGS") LINE_TERMINATOR);
 	
 	FString PkgName = FPackageName::GetShortName(Package);
 	FString PkgDir;
@@ -4763,7 +4769,7 @@ void FNativeClassHeaderGenerator::ExportGeneratedCPP()
 		}
 
 		FString CppPath = ModuleInfo->GeneratedCPPFilenameBase + (GeneratedFunctionBodyTextSplit.Num() > 1 ? *FString::Printf(TEXT(".%d.cpp"), FileIdx + 1) : TEXT(".cpp"));
-		SaveHeaderIfChanged(*CppPath, *(GeneratedCPPPreamble + ModulePCHInclude + GeneratedCPPClassesIncludes + ((FileIdx > 0) ? FString() : GeneratedLinkerFixupFunction) + FileText + GeneratedCPPEpilogue));
+		SaveHeaderIfChanged(*CppPath, *(GeneratedCPPPreamble + ModulePCHInclude + GeneratedCPPClassesIncludes + DisableDeprecationWarnings + ((FileIdx > 0) ? FString() : GeneratedLinkerFixupFunction) + FileText + GeneratedCPPEpilogue + EnableDeprecationWarnings));
 
 		if (GeneratedFunctionBodyTextSplit.Num() > 1)
 		{
diff --git a/Engine/Source/Runtime/AIModule/Classes/EnvironmentQuery/EnvQueryTypes.h b/Engine/Source/Runtime/AIModule/Classes/EnvironmentQuery/EnvQueryTypes.h
--- a/Engine/Source/Runtime/AIModule/Classes/EnvironmentQuery/EnvQueryTypes.h
+++ b/Engine/Source/Runtime/AIModule/Classes/EnvironmentQuery/EnvQueryTypes.h
@@ -934,26 +934,20 @@ struct AIMODULE_API FEnvBoolParam_DEPRECATED
 	void Convert(UObject* Owner, FAIDataProviderBoolValue& ValueProvider);
 };
 
-/* DISABLED FOR NOW, DEPRECATED MACRO DOESN'T WORK WITH USTRUCTS
-struct DEPRECATED(4.8, "FEnvFloatParam is deprecated in 4.8 and was replaced with FAIDataProviderFloatValue. Please use that type instead.") FEnvFloatParam;
-struct DEPRECATED(4.8, "FEnvIntParam is deprecated in 4.8 and was replaced with FAIDataProviderIntValue. Please use that type instead.") FEnvIntParam;
-struct DEPRECATED(4.8, "FEnvBoolParam is deprecated in 4.8 and was replaced with FAIDataProviderBoolValue. Please use that type instead.") FEnvBoolParam;
-*/
-
 USTRUCT()
-struct AIMODULE_API FEnvFloatParam : public FEnvFloatParam_DEPRECATED
+struct DEPRECATED(4.8, "FEnvFloatParam is deprecated in 4.8 and was replaced with FAIDataProviderFloatValue. Please use that type instead.") AIMODULE_API FEnvFloatParam : public FEnvFloatParam_DEPRECATED
 {
 	GENERATED_USTRUCT_BODY();
 };
 
 USTRUCT()
-struct AIMODULE_API FEnvIntParam : public FEnvIntParam_DEPRECATED
+struct DEPRECATED(4.8, "FEnvIntParam is deprecated in 4.8 and was replaced with FAIDataProviderIntValue. Please use that type instead.") AIMODULE_API FEnvIntParam : public FEnvIntParam_DEPRECATED
 {
 	GENERATED_USTRUCT_BODY();
 };
 
 USTRUCT()
-struct AIMODULE_API FEnvBoolParam : public FEnvBoolParam_DEPRECATED
+struct DEPRECATED(4.8, "FEnvBoolParam is deprecated in 4.8 and was replaced with FAIDataProviderBoolValue. Please use that type instead.") AIMODULE_API FEnvBoolParam : public FEnvBoolParam_DEPRECATED
 {
 	GENERATED_USTRUCT_BODY();
 };
diff --git a/Engine/Source/Runtime/Core/Public/Android/AndroidPlatformCompilerPreSetup.h b/Engine/Source/Runtime/Core/Public/Android/AndroidPlatformCompilerPreSetup.h
--- a/Engine/Source/Runtime/Core/Public/Android/AndroidPlatformCompilerPreSetup.h
+++ b/Engine/Source/Runtime/Core/Public/Android/AndroidPlatformCompilerPreSetup.h
@@ -19,8 +19,14 @@
 	 *		DEPRECATED(4.6, "Message")
 	 *		void Function();
 	 *
-	 *		struct DEPRECATED(4.6, "Message") MyStruct;
-	 *		class DEPRECATED(4.6, "Message") MyClass;
+	 *		struct DEPRECATED(4.6, "Message") MODULE_API MyStruct
+	 *		{
+	 *			// StructImplementation
+	 *		};
+	 *		class DEPRECATED(4.6, "Message") MODULE_API MyClass
+	 *		{
+	 *			// ClassImplementation
+	 *		};
 	 *
 	 * @param VERSION The release number in which the feature was marked deprecated.
 	 * @param MESSAGE A message containing upgrade notes.
diff --git a/Engine/Source/Runtime/Core/Public/GenericPlatform/GenericPlatformCompilerPreSetup.h b/Engine/Source/Runtime/Core/Public/GenericPlatform/GenericPlatformCompilerPreSetup.h
--- a/Engine/Source/Runtime/Core/Public/GenericPlatform/GenericPlatformCompilerPreSetup.h
+++ b/Engine/Source/Runtime/Core/Public/GenericPlatform/GenericPlatformCompilerPreSetup.h
@@ -17,9 +17,50 @@
 	 *		DEPRECATED(4.6, "Message")
 	 *		void Function();
 	 *
-	 *		struct DEPRECATED(4.6, "Message") MyStruct;
-	 *		class DEPRECATED(4.6, "Message") MyClass;
+	 *		struct DEPRECATED(4.6, "Message") MODULE_API MyStruct
+	 *		{
+	 *			// StructImplementation
+	 *		};
+	 *		class DEPRECATED(4.6, "Message") MODULE_API MyClass
+	 *		{
+	 *			// ClassImplementation
+	 *		};
 	 *
+	 *		Unfortunately, VC++ will complain about using member functions and fields from deprecated
+	 *		class/structs even for class/struct implementation e.g.:
+	 *		class DEPRECATED(4.8, "") DeprecatedClass
+	 *		{
+	 *		public:
+	 *			DeprecatedClass() {}
+	 *
+	 *			float GetMyFloat()
+	 *			{
+	 *				return MyFloat; // This line will cause warning that deprecated field is used.
+	 *			}
+	 *		private:
+	 *			float MyFloat;
+	 *		};
+	 *
+	 *		To get rid of this warning, place all code not called in class implementation in non-deprecated
+	 *		base class and deprecate only derived one. This may force you to change some access specifiers
+	 *		from private to protected, e.g.:
+	 *
+	 *		class DeprecatedClass_Base_DEPRECATED
+	 *		{
+	 *		protected: // MyFloat is protected now, so DeprecatedClass has access to it.
+	 *			float MyFloat;
+	 *		};
+	 *
+	 *		class DEPRECATED(4.8, "") DeprecatedClass : DeprecatedClass_Base_DEPRECATED
+	 *		{
+	 *		public:
+	 *			DeprecatedClass() {}
+	 *
+	 *			float GetMyFloat()
+	 *			{
+	 *				return MyFloat;
+	 *			}
+	 *		};
 	 * @param VERSION The release number in which the feature was marked deprecated.
 	 * @param MESSAGE A message text containing additional upgrade notes.
 	 */
diff --git a/Engine/Source/Runtime/Core/Public/HTML5/HTML5PlatformCompilerPreSetup.h b/Engine/Source/Runtime/Core/Public/HTML5/HTML5PlatformCompilerPreSetup.h
--- a/Engine/Source/Runtime/Core/Public/HTML5/HTML5PlatformCompilerPreSetup.h
+++ b/Engine/Source/Runtime/Core/Public/HTML5/HTML5PlatformCompilerPreSetup.h
@@ -21,9 +21,50 @@
 	 *		DEPRECATED(4.6, "Message")
 	 *		void Function();
 	 *
-	 *		struct DEPRECATED(4.6, "Message") MyStruct;
-	 *		class DEPRECATED(4.6, "Message") MyClass;
+	 *		struct DEPRECATED(4.6, "Message") MODULE_API MyStruct
+	 *		{
+	 *			// StructImplementation
+	 *		};
+	 *		class DEPRECATED(4.6, "Message") MODULE_API MyClass
+	 *		{
+	 *			// ClassImplementation
+	 *		};
 	 *
+	 *		Unfortunately, VC++ will complain about using member functions and fields from deprecated
+	 *		class/structs even for class/struct implementation e.g.:
+	 *		class DEPRECATED(4.8, "") DeprecatedClass
+	 *		{
+	 *		public:
+	 *			DeprecatedClass() {}
+	 *
+	 *			float GetMyFloat()
+	 *			{
+	 *				return MyFloat; // This line will cause warning that deprecated field is used.
+	 *			}
+	 *		private:
+	 *			float MyFloat;
+	 *		};
+	 *
+	 *		To get rid of this warning, place all code not called in class implementation in non-deprecated
+	 *		base class and deprecate only derived one. This may force you to change some access specifiers
+	 *		from private to protected, e.g.:
+	 *
+	 *		class DeprecatedClass_Base_DEPRECATED
+	 *		{
+	 *		protected: // MyFloat is protected now, so DeprecatedClass has access to it.
+	 *			float MyFloat;
+	 *		};
+	 *
+	 *		class DEPRECATED(4.8, "") DeprecatedClass : DeprecatedClass_Base_DEPRECATED
+	 *		{
+	 *		public:
+	 *			DeprecatedClass() {}
+	 *
+	 *			float GetMyFloat()
+	 *			{
+	 *				return MyFloat;
+	 *			}
+	 *		};
 	 * @param VERSION The release number in which the feature was marked deprecated.
 	 * @param MESSAGE A message containing upgrade notes.
 	 */
diff --git a/Engine/Source/Runtime/Core/Public/IOS/IOSPlatformCompilerPreSetup.h b/Engine/Source/Runtime/Core/Public/IOS/IOSPlatformCompilerPreSetup.h
--- a/Engine/Source/Runtime/Core/Public/IOS/IOSPlatformCompilerPreSetup.h
+++ b/Engine/Source/Runtime/Core/Public/IOS/IOSPlatformCompilerPreSetup.h
@@ -19,9 +19,50 @@
 	 *		DEPRECATED(4.6, "Message")
 	 *		void Function();
 	 *
-	 *		struct DEPRECATED(4.6, "Message") MyStruct;
-	 *		class DEPRECATED(4.6, "Message") MyClass;
+	 *		struct DEPRECATED(4.6, "Message") MODULE_API MyStruct
+	 *		{
+	 *			// StructImplementation
+	 *		};
+	 *		class DEPRECATED(4.6, "Message") MODULE_API MyClass
+	 *		{
+	 *			// ClassImplementation
+	 *		};
 	 *
+	 *		Unfortunately, VC++ will complain about using member functions and fields from deprecated
+	 *		class/structs even for class/struct implementation e.g.:
+	 *		class DEPRECATED(4.8, "") DeprecatedClass
+	 *		{
+	 *		public:
+	 *			DeprecatedClass() {}
+	 *
+	 *			float GetMyFloat()
+	 *			{
+	 *				return MyFloat; // This line will cause warning that deprecated field is used.
+	 *			}
+	 *		private:
+	 *			float MyFloat;
+	 *		};
+	 *
+	 *		To get rid of this warning, place all code not called in class implementation in non-deprecated
+	 *		base class and deprecate only derived one. This may force you to change some access specifiers
+	 *		from private to protected, e.g.:
+	 *
+	 *		class DeprecatedClass_Base_DEPRECATED
+	 *		{
+	 *		protected: // MyFloat is protected now, so DeprecatedClass has access to it.
+	 *			float MyFloat;
+	 *		};
+	 *
+	 *		class DEPRECATED(4.8, "") DeprecatedClass : DeprecatedClass_Base_DEPRECATED
+	 *		{
+	 *		public:
+	 *			DeprecatedClass() {}
+	 *
+	 *			float GetMyFloat()
+	 *			{
+	 *				return MyFloat;
+	 *			}
+	 *		};
 	 * @param VERSION The release number in which the feature was marked deprecated.
 	 * @param MESSAGE A message containing upgrade notes.
 	 */
diff --git a/Engine/Source/Runtime/Core/Public/Linux/LinuxPlatformCompilerPreSetup.h b/Engine/Source/Runtime/Core/Public/Linux/LinuxPlatformCompilerPreSetup.h
--- a/Engine/Source/Runtime/Core/Public/Linux/LinuxPlatformCompilerPreSetup.h
+++ b/Engine/Source/Runtime/Core/Public/Linux/LinuxPlatformCompilerPreSetup.h
@@ -19,8 +19,14 @@
 	 *		DEPRECATED(4.6, "Message")
 	 *		void Function();
 	 *
-	 *		struct DEPRECATED(4.6, "Message") MyStruct;
-	 *		class DEPRECATED(4.6, "Message") MyClass;
+	 *		struct DEPRECATED(4.6, "Message") MODULE_API MyStruct
+	 *		{
+	 *			// StructImplementation
+	 *		};
+	 *		class DEPRECATED(4.6, "Message") MODULE_API MyClass
+	 *		{
+	 *			// ClassImplementation
+	 *		};
 	 *
 	 * @param VERSION The release number in which the feature was marked deprecated.
 	 * @param MESSAGE A message containing upgrade notes.
diff --git a/Engine/Source/Runtime/Core/Public/Mac/MacPlatformCompilerPreSetup.h b/Engine/Source/Runtime/Core/Public/Mac/MacPlatformCompilerPreSetup.h
--- a/Engine/Source/Runtime/Core/Public/Mac/MacPlatformCompilerPreSetup.h
+++ b/Engine/Source/Runtime/Core/Public/Mac/MacPlatformCompilerPreSetup.h
@@ -19,9 +19,50 @@
 	 *		DEPRECATED(4.6, "Message")
 	 *		void Function();
 	 *
-	 *		struct DEPRECATED(4.6, "Message") MyStruct;
-	 *		class DEPRECATED(4.6, "Message") MyClass;
+	 *		struct DEPRECATED(4.6, "Message") MODULE_API MyStruct
+	 *		{
+	 *			// StructImplementation
+	 *		};
+	 *		class DEPRECATED(4.6, "Message") MODULE_API MyClass
+	 *		{
+	 *			// ClassImplementation
+	 *		};
 	 *
+	 *		Unfortunately, VC++ will complain about using member functions and fields from deprecated
+	 *		class/structs even for class/struct implementation e.g.:
+	 *		class DEPRECATED(4.8, "") DeprecatedClass
+	 *		{
+	 *		public:
+	 *			DeprecatedClass() {}
+	 *
+	 *			float GetMyFloat()
+	 *			{
+	 *				return MyFloat; // This line will cause warning that deprecated field is used.
+	 *			}
+	 *		private:
+	 *			float MyFloat;
+	 *		};
+	 *
+	 *		To get rid of this warning, place all code not called in class implementation in non-deprecated
+	 *		base class and deprecate only derived one. This may force you to change some access specifiers
+	 *		from private to protected, e.g.:
+	 *
+	 *		class DeprecatedClass_Base_DEPRECATED
+	 *		{
+	 *		protected: // MyFloat is protected now, so DeprecatedClass has access to it.
+	 *			float MyFloat;
+	 *		};
+	 *
+	 *		class DEPRECATED(4.8, "") DeprecatedClass : DeprecatedClass_Base_DEPRECATED
+	 *		{
+	 *		public:
+	 *			DeprecatedClass() {}
+	 *
+	 *			float GetMyFloat()
+	 *			{
+	 *				return MyFloat;
+	 *			}
+	 *		};
 	 * @param VERSION The release number in which the feature was marked deprecated.
 	 * @param MESSAGE A message containing upgrade notes.
 	 */
diff --git a/Engine/Source/Runtime/Core/Public/WinRT/WinRTPlatformCompilerPreSetup.h b/Engine/Source/Runtime/Core/Public/WinRT/WinRTPlatformCompilerPreSetup.h
--- a/Engine/Source/Runtime/Core/Public/WinRT/WinRTPlatformCompilerPreSetup.h
+++ b/Engine/Source/Runtime/Core/Public/WinRT/WinRTPlatformCompilerPreSetup.h
@@ -17,9 +17,50 @@
 	 *		DEPRECATED(4.6, "Message")
 	 *		void Function();
 	 *
-	 *		struct DEPRECATED(4.6, "Message") MyStruct;
-	 *		class DEPRECATED(4.6, "Message") MyClass;
+	 *		struct DEPRECATED(4.6, "Message") MODULE_API MyStruct
+	 *		{
+	 *			// StructImplementation
+	 *		};
+	 *		class DEPRECATED(4.6, "Message") MODULE_API MyClass
+	 *		{
+	 *			// ClassImplementation
+	 *		};
 	 *
+	 *		Unfortunately, VC++ will complain about using member functions and fields from deprecated
+	 *		class/structs even for class/struct implementation e.g.:
+	 *		class DEPRECATED(4.8, "") DeprecatedClass
+	 *		{
+	 *		public:
+	 *			DeprecatedClass() {}
+	 *
+	 *			float GetMyFloat()
+	 *			{
+	 *				return MyFloat; // This line will cause warning that deprecated field is used.
+	 *			}
+	 *		private:
+	 *			float MyFloat;
+	 *		};
+	 *
+	 *		To get rid of this warning, place all code not called in class implementation in non-deprecated
+	 *		base class and deprecate only derived one. This may force you to change some access specifiers
+	 *		from private to protected, e.g.:
+	 *
+	 *		class DeprecatedClass_Base_DEPRECATED
+	 *		{
+	 *		protected: // MyFloat is protected now, so DeprecatedClass has access to it.
+	 *			float MyFloat;
+	 *		};
+	 *
+	 *		class DEPRECATED(4.8, "") DeprecatedClass : DeprecatedClass_Base_DEPRECATED
+	 *		{
+	 *		public:
+	 *			DeprecatedClass() {}
+	 *
+	 *			float GetMyFloat()
+	 *			{
+	 *				return MyFloat;
+	 *			}
+	 *		};
 	 * @param VERSION The release number in which the feature was marked deprecated.
 	 * @param MESSAGE A message containing upgrade notes.
 	 */
diff --git a/Engine/Source/Runtime/Core/Public/Windows/WindowsPlatformCompilerPreSetup.h b/Engine/Source/Runtime/Core/Public/Windows/WindowsPlatformCompilerPreSetup.h
--- a/Engine/Source/Runtime/Core/Public/Windows/WindowsPlatformCompilerPreSetup.h
+++ b/Engine/Source/Runtime/Core/Public/Windows/WindowsPlatformCompilerPreSetup.h
@@ -17,8 +17,14 @@
 	 *		DEPRECATED(4.6, "Message")
 	 *		void Function();
 	 *
-	 *		struct DEPRECATED(4.6, "Message") MyStruct;
-	 *		class DEPRECATED(4.6, "Message") MyClass;
+	 *		struct DEPRECATED(4.6, "Message") MODULE_API MyStruct
+	 *		{
+	 *			// StructImplementation
+	 *		};
+	 *		class DEPRECATED(4.6, "Message") MODULE_API MyClass
+	 *		{
+	 *			// ClassImplementation
+	 *		};
 	 *
 	 * @param VERSION The release number in which the feature was marked deprecated.
 	 * @param MESSAGE A message containing upgrade notes.
