UHT was reporting circular dependency if UClassA inherited UClassB defined in the same header. Fixed.

diff --git a/Engine/Source/Programs/UnrealHeaderTool/Private/CodeGenerator.cpp b/Engine/Source/Programs/UnrealHeaderTool/Private/CodeGenerator.cpp
--- a/Engine/Source/Programs/UnrealHeaderTool/Private/CodeGenerator.cpp
+++ b/Engine/Source/Programs/UnrealHeaderTool/Private/CodeGenerator.cpp
@@ -2256,21 +2256,6 @@ void FNativeClassHeaderGenerator::ExportSourceFileHeaderRecursive(FClasses& AllC
 	VisitedSet.Add(SourceFile);
 
 	TArray<FUnrealSourceFile*> DependOnList;
-
-	// Export the super class first.
-	for (auto* Class : SourceFile->GetDefinedClasses())
-	{
-		if (FClass* SuperClass = ((FClass*)Class)->GetSuperClass())
-		{
-			if (SuperClass->ClassFlags & (CLASS_NoExport | CLASS_Intrinsic))
-			{
-				continue;
-			}
-
-			DependOnList.Add(&GTypeDefinitionInfoMap[SuperClass]->GetUnrealSourceFile());
-		}
-	}
-
 	for (auto& Include : SourceFile->GetIncludes())
 	{
 		auto* SourceFile = Include.Resolve();
@@ -5253,11 +5238,6 @@ UClass* ProcessParsedClass(bool bClassIsAnInterface, TArray<FHeaderProvider> &De
 				ResultClass->ClassCastFlags |= ResultClass->GetSuperClass()->ClassCastFlags;
 			}
 		}
-
-		if (BaseClassNameStripped != TEXT("Object") && ResultClass->GetSuperClass())
-		{
-			DependentOn.AddUnique(FHeaderProvider(EHeaderProviderSourceType::ClassName, *BaseClassNameStripped));
-		}
 	}
 
 	if (bVerboseOutput)
diff --git a/Engine/Source/Programs/UnrealHeaderTool/Private/HeaderParser.cpp b/Engine/Source/Programs/UnrealHeaderTool/Private/HeaderParser.cpp
--- a/Engine/Source/Programs/UnrealHeaderTool/Private/HeaderParser.cpp
+++ b/Engine/Source/Programs/UnrealHeaderTool/Private/HeaderParser.cpp
@@ -7133,13 +7133,13 @@ void FHeaderParser::SimplifiedClassParse(const TCHAR* InBuffer, TArray<FSimplifi
 			// Get class or interface name
 			if (const TCHAR* UInterfaceMacroDecl = FCString::Strfind(Str, TEXT("UINTERFACE(")))
 			{
-				Parser.ParseClassDeclaration(StartOfLine + (UInterfaceMacroDecl - Str), CurrentLine, TEXT("UINTERFACE"), /*out*/ ClassName, /*out*/ BaseClassName, /*out*/ DependentOn);
+				Parser.ParseClassDeclaration(StartOfLine + (UInterfaceMacroDecl - Str), CurrentLine, TEXT("UINTERFACE"), /*out*/ ClassName, /*out*/ BaseClassName, /*out*/ DependentOn, OutParsedClassArray);
 				OutParsedClassArray.Add(FSimplifiedParsingClassInfo(MoveTemp(ClassName), MoveTemp(BaseClassName), CurrentLine, true));
 			}
 
 			if (const TCHAR* UClassMacroDecl = FCString::Strfind(Str, TEXT("UCLASS(")))
 			{
-				Parser.ParseClassDeclaration(StartOfLine + (UClassMacroDecl - Str), CurrentLine, TEXT("UCLASS"), /*out*/ ClassName, /*out*/ BaseClassName, /*out*/ DependentOn);
+				Parser.ParseClassDeclaration(StartOfLine + (UClassMacroDecl - Str), CurrentLine, TEXT("UCLASS"), /*out*/ ClassName, /*out*/ BaseClassName, /*out*/ DependentOn, OutParsedClassArray);
 				OutParsedClassArray.Add(FSimplifiedParsingClassInfo(MoveTemp(ClassName), MoveTemp(BaseClassName), CurrentLine, false));
 			}
 		}
@@ -7151,7 +7151,7 @@ void FHeaderParser::SimplifiedClassParse(const TCHAR* InBuffer, TArray<FSimplifi
 /////////////////////////////////////////////////////
 // FHeaderPreParser
 
-void FHeaderPreParser::ParseClassDeclaration(const TCHAR* InputText, int32 InLineNumber, const TCHAR* StartingMatchID, FString& out_ClassName, FString& out_BaseClassName, TArray<FHeaderProvider>& out_ClassNames)
+void FHeaderPreParser::ParseClassDeclaration(const TCHAR* InputText, int32 InLineNumber, const TCHAR* StartingMatchID, FString& out_ClassName, FString& out_BaseClassName, TArray<FHeaderProvider>& out_RequiredIncludes, const TArray<FSimplifiedParsingClassInfo>& ParsedClassArray)
 {
 	FString ErrorMsg = TEXT("Class declaration");
 
@@ -7199,6 +7199,8 @@ void FHeaderPreParser::ParseClassDeclaration(const TCHAR* InputText, int32 InLin
 
 		out_BaseClassName = BaseClassNameToken.Identifier;
 
+		AddDependencyIfNeeded(ParsedClassArray, out_BaseClassName, out_RequiredIncludes);
+
 		// Get additional inheritance links and rack them up as dependencies if they're UObject derived
 		while (MatchSymbol(TEXT(",")))
 		{
@@ -7211,12 +7213,22 @@ void FHeaderPreParser::ParseClassDeclaration(const TCHAR* InputText, int32 InLin
 				FError::Throwf(TEXT("Expected an interface class name"));
 			}
 
-			FName InterfaceClassName(InterfaceClassNameToken.Identifier);
-			out_ClassNames.Add(FHeaderProvider(EHeaderProviderSourceType::ClassName, InterfaceClassName.ToString().Mid(1)));
+			AddDependencyIfNeeded(ParsedClassArray, FString(InterfaceClassNameToken.Identifier), out_RequiredIncludes);
 		}
 	}
 }
 
+void FHeaderPreParser::AddDependencyIfNeeded(const TArray<FSimplifiedParsingClassInfo>& ParsedClassArray, const FString& DependencyClassName, TArray<FHeaderProvider>& RequiredIncludes) const
+{
+	if (ParsedClassArray.FindByPredicate([&DependencyClassName](const FSimplifiedParsingClassInfo& Info)
+		{
+			return Info.GetClassName() == DependencyClassName;
+		}) == nullptr)
+	{
+		RequiredIncludes.Add(FHeaderProvider(EHeaderProviderSourceType::ClassName, DependencyClassName.Mid(1)));
+	}
+}
+
 //////////////////////////////////////////////////////////////////////////////////////////////////////////////
 
 bool FHeaderParser::DefaultValueStringCppFormatToInnerFormat(const UProperty* Property, const FString& CppForm, FString &OutForm)
diff --git a/Engine/Source/Programs/UnrealHeaderTool/Private/HeaderParser.h b/Engine/Source/Programs/UnrealHeaderTool/Private/HeaderParser.h
--- a/Engine/Source/Programs/UnrealHeaderTool/Private/HeaderParser.h
+++ b/Engine/Source/Programs/UnrealHeaderTool/Private/HeaderParser.h
@@ -787,5 +787,8 @@ public:
 	{
 	}
 
-	void ParseClassDeclaration(const TCHAR* InputText, int32 InLineNumber, const TCHAR* StartingMatchID, FString& out_ClassName, FString& out_BaseClassName, TArray<FHeaderProvider>& out_ClassNames);
+	void ParseClassDeclaration(const TCHAR* InputText, int32 InLineNumber, const TCHAR* StartingMatchID, FString& out_ClassName, FString& out_BaseClassName, TArray<FHeaderProvider>& out_ClassNames, const TArray<FSimplifiedParsingClassInfo>& ParsedClassArray);
+
+private:
+	void AddDependencyIfNeeded(const TArray<FSimplifiedParsingClassInfo>& ParsedClassArray, const FString& DependencyClassName, TArray<FHeaderProvider>& RequiredIncludes) const;
 };
diff --git a/Engine/Source/Programs/UnrealHeaderTool/Private/HeaderProvider.cpp b/Engine/Source/Programs/UnrealHeaderTool/Private/HeaderProvider.cpp
--- a/Engine/Source/Programs/UnrealHeaderTool/Private/HeaderProvider.cpp
+++ b/Engine/Source/Programs/UnrealHeaderTool/Private/HeaderProvider.cpp
@@ -5,7 +5,7 @@
 #include "UnrealSourceFile.h"
 #include "ClassMaps.h"
 
-FHeaderProvider::FHeaderProvider(EHeaderProviderSourceType::Type Type, const FString& Id)
+FHeaderProvider::FHeaderProvider(EHeaderProviderSourceType Type, const FString& Id)
 	: Type(Type), Id(Id), Cache(nullptr)
 {
 
@@ -81,6 +81,11 @@ const FString& FHeaderProvider::GetId() const
 	return Id;
 }
 
+EHeaderProviderSourceType FHeaderProvider::GetType() const
+{
+	return Type;
+}
+
 bool operator==(const FHeaderProvider& A, const FHeaderProvider& B)
 {
 	if (A.Cache != B.Cache)
diff --git a/Engine/Source/Programs/UnrealHeaderTool/Private/HeaderProvider.h b/Engine/Source/Programs/UnrealHeaderTool/Private/HeaderProvider.h
--- a/Engine/Source/Programs/UnrealHeaderTool/Private/HeaderProvider.h
+++ b/Engine/Source/Programs/UnrealHeaderTool/Private/HeaderProvider.h
@@ -6,21 +6,18 @@
 
 class FUnrealSourceFile;
 
-namespace EHeaderProviderSourceType
+enum class EHeaderProviderSourceType
 {
-	enum Type
-	{
-		ClassName,
-		FileName,
-		Resolved
-	};
-}
+	ClassName,
+	FileName,
+	Resolved
+};
 
 class FHeaderProvider
 {
 	friend bool operator==(const FHeaderProvider& A, const FHeaderProvider& B);
 public:
-	FHeaderProvider(EHeaderProviderSourceType::Type Type, const FString& Id);
+	FHeaderProvider(EHeaderProviderSourceType Type, const FString& Id);
 
 	FUnrealSourceFile* Resolve();
 
@@ -28,8 +25,10 @@ public:
 
 	const FString& GetId() const;
 
+	EHeaderProviderSourceType GetType() const;
+
 private:
-	EHeaderProviderSourceType::Type Type;
+	EHeaderProviderSourceType Type;
 	FString Id;
 	FUnrealSourceFile* Cache;
 };
