Inner DELEGATEs as children. [By JaroslawP]

UE-11693 Compiler error when using Bind Event to OnTextChanged in UMG
UE-11674 BindEvent node gives "Wrong Event Dispatcher" error for BlueprintAssignable MC delegate

diff --git a/Engine/Source/Programs/UnrealHeaderTool/Private/CodeGenerator.cpp b/Engine/Source/Programs/UnrealHeaderTool/Private/CodeGenerator.cpp
--- a/Engine/Source/Programs/UnrealHeaderTool/Private/CodeGenerator.cpp
+++ b/Engine/Source/Programs/UnrealHeaderTool/Private/CodeGenerator.cpp
@@ -1029,7 +1029,12 @@ void FNativeClassHeaderGenerator::ExportNativeGeneratedInitCode(FClass* Class, F
 	for (int32 FuncIndex = 0; FuncIndex < FunctionsToExport.Num(); FuncIndex++)
 	{
 		UFunction* Function = FunctionsToExport[FuncIndex];
-		ExportFunction(Function, &FScope::GetTypeScope(Class).Get(), bIsNoExport);
+
+		if (!Function->IsA<UDelegateFunction>())
+		{
+			ExportFunction(Function, &FScope::GetTypeScope(Class).Get(), bIsNoExport);
+		}
+
 		CallSingletons.Logf(TEXT("                OuterClass->LinkChild(%s);\r\n"), *GetSingletonName(Function));
 	}
 
diff --git a/Engine/Source/Programs/UnrealHeaderTool/Private/HeaderParser.cpp b/Engine/Source/Programs/UnrealHeaderTool/Private/HeaderParser.cpp
--- a/Engine/Source/Programs/UnrealHeaderTool/Private/HeaderParser.cpp
+++ b/Engine/Source/Programs/UnrealHeaderTool/Private/HeaderParser.cpp
@@ -7764,17 +7764,20 @@ TFunctionType* CreateFunctionImpl(const FFuncInfo& FuncInfo, UObject* Outer, FSc
 	Function->ReturnValueOffset = MAX_uint16;
 	Function->FirstPropertyToInit = nullptr;
 
+	if (!CurrentScope->IsFileScope())
+	{
+		auto* Struct = ((FStructScope*)CurrentScope)->GetStruct();
+
+		Function->Next = Struct->Children;
+		Struct->Children = Function;
+	}
+
 	return Function;
 }
 
 UFunction* FHeaderParser::CreateFunction(const FFuncInfo &FuncInfo) const
 {
-	auto* Function = CreateFunctionImpl<UFunction>(FuncInfo, GetCurrentClass(), GetCurrentScope());
-
-	Function->Next = GetCurrentClass()->Children;
-	GetCurrentClass()->Children = Function;
-
-	return Function;
+	return CreateFunctionImpl<UFunction>(FuncInfo, GetCurrentClass(), GetCurrentScope());
 }
 
 UDelegateFunction* FHeaderParser::CreateDelegateFunction(const FFuncInfo &FuncInfo) const
