https://github.com/EpicGames/UnrealEngine/pull/910

Fixes UHT crashes with malformed UINTERFACE headers.

diff --git a/Engine/Source/Programs/UnrealHeaderTool/Private/HeaderParser.cpp b/Engine/Source/Programs/UnrealHeaderTool/Private/HeaderParser.cpp
--- a/Engine/Source/Programs/UnrealHeaderTool/Private/HeaderParser.cpp
+++ b/Engine/Source/Programs/UnrealHeaderTool/Private/HeaderParser.cpp
@@ -3717,6 +3717,8 @@ bool FHeaderParser::CompileDeclaration(FClasses& AllClasses, FUnrealSourceFile&
 
 		GetCurrentClassData()->SetInterfaceGeneratedBodyLine(InputLine);
 
+		bClassHasGeneratedIInterfaceBody = true;
+
 		// The body implementation macro always ends with a 'public:'
 		CurrentAccessSpecifier = ACCESS_Public;
 	}
@@ -3731,6 +3733,8 @@ bool FHeaderParser::CompileDeclaration(FClasses& AllClasses, FUnrealSourceFile&
 
 		GetCurrentClassData()->SetGeneratedBodyLine(InputLine);
 
+		bClassHasGeneratedUInterfaceBody = true;
+
 		// The body implementation macro always ends with a 'public:'
 		CurrentAccessSpecifier = ACCESS_Public;
 	}
@@ -3800,6 +3804,18 @@ bool FHeaderParser::CompileDeclaration(FClasses& AllClasses, FUnrealSourceFile&
 		{
 			PopNest(NEST_Interface, TEXT("'Interface'"));
 			PostPopNestInterface(AllClasses, CurrentClass);
+
+			// Ensure the UINTERFACE classes have a GENERATED_UINTERFACE_BODY declaration
+			if (bHaveSeenUClass && !bClassHasGeneratedUInterfaceBody)
+			{
+				FError::Throwf(TEXT("Expected a GENERATED_UINTERFACE_BODY() at the start of class"));
+			}
+
+			// Ensure the non-UINTERFACE interface classes have a GENERATED_IINTERFACE_BODY declaration
+			if (!bHaveSeenUClass && !bClassHasGeneratedIInterfaceBody)
+			{
+				FError::Throwf(TEXT("Expected a GENERATED_IINTERFACE_BODY() at the start of class"));
+			}
 		}
 		else
 		{
@@ -3811,10 +3827,13 @@ bool FHeaderParser::CompileDeclaration(FClasses& AllClasses, FUnrealSourceFile&
 			{
 				FError::Throwf(TEXT("Expected a GENERATED_UCLASS_BODY() at the start of class"));
 			}
-
-			bClassHasGeneratedBody = false;
 		}
 
+		bHaveSeenUClass                  = false;
+		bClassHasGeneratedBody           = false;
+		bClassHasGeneratedUInterfaceBody = false;
+		bClassHasGeneratedIInterfaceBody = false;
+
 		GetCurrentScope()->AddType(CurrentClass);
 	}
 	else if (Token.Matches(TEXT(";")))
@@ -6203,6 +6222,8 @@ ECompilationResult::Type FHeaderParser::ParseHeader(FClasses& AllClasses, FUnrea
 	bSpottedAutogeneratedHeaderInclude          = false;
 	bHaveSeenUClass                             = false;
 	bClassHasGeneratedBody                      = false;
+	bClassHasGeneratedUInterfaceBody            = false;
+	bClassHasGeneratedIInterfaceBody            = false;
 
 	ECompilationResult::Type Result = ECompilationResult::OtherCompilationError;
 
diff --git a/Engine/Source/Programs/UnrealHeaderTool/Private/HeaderParser.h b/Engine/Source/Programs/UnrealHeaderTool/Private/HeaderParser.h
--- a/Engine/Source/Programs/UnrealHeaderTool/Private/HeaderParser.h
+++ b/Engine/Source/Programs/UnrealHeaderTool/Private/HeaderParser.h
@@ -389,6 +389,12 @@ protected:
 	// Indicates that a GENERATED_UCLASS_BODY or GENERATED_BODY has been found in the UClass.
 	bool bClassHasGeneratedBody;
 
+	// Indicates that a GENERATED_UINTERFACE_BODY has been found in the UClass.
+	bool bClassHasGeneratedUInterfaceBody;
+
+	// Indicates that a GENERATED_IINTERFACE_BODY has been found in the UClass.
+	bool bClassHasGeneratedIInterfaceBody;
+
 	// public, private, etc at the current parse spot
 	EAccessSpecifier CurrentAccessSpecifier;
 
