Some unused stuff removed from FHeaderParser: bIsMulticastDelegate, LegalVariableSpecifiers and IsValidVariableSpecifier.

diff --git a/Engine/Source/Programs/UnrealHeaderTool/Private/HeaderParser.cpp b/Engine/Source/Programs/UnrealHeaderTool/Private/HeaderParser.cpp
--- a/Engine/Source/Programs/UnrealHeaderTool/Private/HeaderParser.cpp
+++ b/Engine/Source/Programs/UnrealHeaderTool/Private/HeaderParser.cpp
@@ -2143,12 +2143,11 @@ bool FHeaderParser::GetVarType
 )
 {
 	UStruct* OwnerStruct = Scope->IsFileScope() ? nullptr : ((FStructScope*)Scope)->GetStruct();
-	FName RepCallbackName      = FName(NAME_None);
-	bool  bIsMulticastDelegate = false;
-	bool  bIsWeak              = false;
-	bool  bIsLazy              = false;
-	bool  bIsAsset             = false;
-	bool  bWeakIsAuto          = false;
+	FName RepCallbackName = FName(NAME_None);
+	bool  bIsWeak         = false;
+	bool  bIsLazy         = false;
+	bool  bIsAsset        = false;
+	bool  bWeakIsAuto     = false;
 
 	// Get flags.
 	ObjectFlags = RF_Public;
@@ -2175,29 +2174,14 @@ bool FHeaderParser::GetVarType
 		{
 			ReadSpecifierSetInsideMacro(SpecifiersFound, TEXT("Variable"), MetaDataFromNewStyle);
 		}
-		else
-		{
-			// Legacy variable support
-			for ( ; ; )
-			{
-				FToken Specifier;
-				GetToken(Specifier);
-
-				if (IsValidVariableSpecifier(Specifier))
-				{
-					// Record the specifier
-					FPropertySpecifier* NewPair = new (SpecifiersFound) FPropertySpecifier();
-					NewPair->Key = Specifier.Identifier;
+	}
 
-					// Look for a value for this specifier
-					ReadOptionalCommaSeparatedListInParens(NewPair->Values, TEXT("'Variable declaration specifier'"));
-				}
-				else
-				{
-					UngetToken(Specifier);
-					break;
-				}
-			}
+	if (VariableCategory != EVariableCategory::Member)
+	{
+		// const before the variable type support (only for params)
+		if (MatchIdentifier(TEXT("const")))
+		{
+			Flags |= CPF_ConstParm;
 		}
 	}
 
@@ -2644,7 +2628,7 @@ bool FHeaderParser::GetVarType
 			FError::Throwf(TEXT("Arrays within arrays not supported.") );
 		}
 
-		if (bIsLazy || bIsWeak || bIsMulticastDelegate)
+		if (bIsLazy || bIsWeak)
 		{
 			FError::Throwf(TEXT("Object reference and delegate flags must be specified inside array <>."));
 		}
@@ -6606,31 +6590,6 @@ FHeaderParser::FHeaderParser(FFeedbackContext* InWarn)
 	StructsWithTPrefix.Add("MultiMap");
 	StructsWithTPrefix.Add("SharedPtr");
 
-	// List of legal variable specifiers
-	LegalVariableSpecifiers.Add(TEXT("Const"));
-	LegalVariableSpecifiers.Add(TEXT("BlueprintReadOnly"));
-	LegalVariableSpecifiers.Add(TEXT("Config"));
-	LegalVariableSpecifiers.Add(TEXT("GlobalConfig"));
-	LegalVariableSpecifiers.Add(TEXT("Localized"));
-	LegalVariableSpecifiers.Add(TEXT("EditBlueprint"));
-	LegalVariableSpecifiers.Add(TEXT("Transient"));
-	LegalVariableSpecifiers.Add(TEXT("Native"));
-	LegalVariableSpecifiers.Add(TEXT("DuplicateTransient"));
-	LegalVariableSpecifiers.Add(TEXT("Ref"));
-	LegalVariableSpecifiers.Add(TEXT("Export"));
-	LegalVariableSpecifiers.Add(TEXT("NoClear"));
-	LegalVariableSpecifiers.Add(TEXT("EditFixedSize"));
-	LegalVariableSpecifiers.Add(TEXT("Replicated"));
-	LegalVariableSpecifiers.Add(TEXT("ReplicatedUsing"));
-	LegalVariableSpecifiers.Add(TEXT("RepRetry"));
-	LegalVariableSpecifiers.Add(TEXT("Interp"));
-	LegalVariableSpecifiers.Add(TEXT("NonTransactional"));
-	LegalVariableSpecifiers.Add(TEXT("Deprecated"));
-	LegalVariableSpecifiers.Add(TEXT("Instanced"));
-	LegalVariableSpecifiers.Add(TEXT("BlueprintAssignable"));
-	LegalVariableSpecifiers.Add(TEXT("Category"));
-	LegalVariableSpecifiers.Add(TEXT("AssetRegistrySearchable"));
-
 	// List of legal delegate parameter counts
 	DelegateParameterCountStrings.Add(TEXT("_OneParam"));
 	DelegateParameterCountStrings.Add(TEXT("_TwoParams"));
@@ -6642,12 +6601,6 @@ FHeaderParser::FHeaderParser(FFeedbackContext* InWarn)
 	DelegateParameterCountStrings.Add(TEXT("_EightParams"));
 }
 
-// Returns true if the token is a variable specifier
-bool FHeaderParser::IsValidVariableSpecifier(const FToken& Token) const
-{
-	return (Token.TokenType == TOKEN_Identifier) && LegalVariableSpecifiers.Contains(Token.Identifier);
-}
-
 // Throws if a specifier value wasn't provided
 void FHeaderParser::RequireSpecifierValue(const FPropertySpecifier& Specifier, bool bRequireExactlyOne)
 {
diff --git a/Engine/Source/Programs/UnrealHeaderTool/Private/HeaderParser.h b/Engine/Source/Programs/UnrealHeaderTool/Private/HeaderParser.h
--- a/Engine/Source/Programs/UnrealHeaderTool/Private/HeaderParser.h
+++ b/Engine/Source/Programs/UnrealHeaderTool/Private/HeaderParser.h
@@ -406,9 +406,6 @@ protected:
 	// Special parsed struct names that have a 'T' prefix
 	TArray<FString> StructsWithTPrefix;
 
-	// List of all legal variable specifier tokens
-	TSet<FString> LegalVariableSpecifiers;
-
 	// Mapping from 'human-readable' macro substring to # of parameters for delegate declarations
 	// Index 0 is 1 parameter, Index 1 is 2, etc...
 	TArray<FString> DelegateParameterCountStrings;
@@ -430,9 +427,6 @@ protected:
 		}
 	}
 
-	// Returns true if the token is a variable specifier
-	bool IsValidVariableSpecifier(const FToken& Token) const;
-
 	// Returns true if the token is a dynamic delegate declaration
 	bool IsValidDelegateDeclaration(const FToken& Token) const;
 
