Localized some vars which caused a pointless test.

diff --git a/Engine/Source/Programs/UnrealHeaderTool/Private/HeaderParser.cpp b/Engine/Source/Programs/UnrealHeaderTool/Private/HeaderParser.cpp
--- a/Engine/Source/Programs/UnrealHeaderTool/Private/HeaderParser.cpp
+++ b/Engine/Source/Programs/UnrealHeaderTool/Private/HeaderParser.cpp
@@ -2144,10 +2144,6 @@ bool FHeaderParser::GetVarType
 {
 	UStruct* OwnerStruct = Scope->IsFileScope() ? nullptr : ((FStructScope*)Scope)->GetStruct();
 	FName RepCallbackName = FName(NAME_None);
-	bool  bIsWeak         = false;
-	bool  bIsLazy         = false;
-	bool  bIsAsset        = false;
-	bool  bWeakIsAuto     = false;
 
 	// Get flags.
 	ObjectFlags = RF_Public;
@@ -2628,11 +2624,6 @@ bool FHeaderParser::GetVarType
 			FError::Throwf(TEXT("Arrays within arrays not supported.") );
 		}
 
-		if (bIsLazy || bIsWeak)
-		{
-			FError::Throwf(TEXT("Object reference and delegate flags must be specified inside array <>."));
-		}
-
 		OriginalVarTypeFlags |= VarProperty.PropertyFlags & (CPF_ContainsInstancedReference | CPF_InstancedReference); // propagate these to the array, we will fix them later
 		VarType.PropertyFlags = OriginalVarTypeFlags;
 		VarProperty.ArrayType = EArrayType::Dynamic;
@@ -2827,13 +2818,18 @@ bool FHeaderParser::GetVarType
 			UClass* TempClass = NULL;
 			bool bExpectStar = true;
 
-			const bool bIsLazyPtrTemplate = VarType.Matches(TEXT("TLazyObjectPtr"));
-			const bool bIsAssetPtrTemplate = VarType.Matches(TEXT("TAssetPtr"));
-			const bool bIsAssetClassTemplate = VarType.Matches(TEXT("TAssetSubclassOf"));
-			const bool bIsWeakPtrTemplate = VarType.Matches(TEXT("TWeakObjectPtr"));
-			const bool bIsAutoweakPtrTemplate = VarType.Matches(TEXT("TAutoWeakObjectPtr"));
+			const bool bIsLazyPtrTemplate        = VarType.Matches(TEXT("TLazyObjectPtr"));
+			const bool bIsAssetPtrTemplate       = VarType.Matches(TEXT("TAssetPtr"));
+			const bool bIsAssetClassTemplate     = VarType.Matches(TEXT("TAssetSubclassOf"));
+			const bool bIsWeakPtrTemplate        = VarType.Matches(TEXT("TWeakObjectPtr"));
+			const bool bIsAutoweakPtrTemplate    = VarType.Matches(TEXT("TAutoWeakObjectPtr"));
 			const bool bIsScriptInterfaceWrapper = VarType.Matches(TEXT("TScriptInterface"));
-			const bool bIsSubobjectPtrTemplate = VarType.Matches(TEXT("TSubobjectPtr"));
+			const bool bIsSubobjectPtrTemplate   = VarType.Matches(TEXT("TSubobjectPtr"));
+
+			bool bIsWeak     = false;
+			bool bIsLazy     = false;
+			bool bIsAsset    = false;
+			bool bWeakIsAuto = false;
 
 			if (VarType.Matches(TEXT("TSubclassOf")))
 			{
