Enabled UHT to digest GENERATED_BODY instead of GENERATED_UCLASS_BODY, GENERATED_USTRUCT_BODY, GENERATED_UINTERFACE_BODY or GENERATED_IINTERFACE_BODY, changed every occurence to the new syntax and fix

diff --git a/Engine/Source/Editor/UnrealEd/Classes/Niagara/NiagaraNodeInput.h b/Engine/Source/Editor/UnrealEd/Classes/Niagara/NiagaraNodeInput.h
--- a/Engine/Source/Editor/UnrealEd/Classes/Niagara/NiagaraNodeInput.h
+++ b/Engine/Source/Editor/UnrealEd/Classes/Niagara/NiagaraNodeInput.h
@@ -6,7 +6,9 @@
 UCLASS(MinimalAPI)
 class UNiagaraNodeInput : public UNiagaraNode
 {
-	GENERATED_UCLASS_BODY()
+	GENERATED_BODY()
+public:
+	UNREALED_API UNiagaraNodeInput(const FObjectInitializer& ObjectInitializer = FObjectInitializer::Get());
 
 	void ReallocatePins();
 
diff --git a/Engine/Source/Editor/UnrealEd/Classes/Niagara/NiagaraNodeOp.h b/Engine/Source/Editor/UnrealEd/Classes/Niagara/NiagaraNodeOp.h
--- a/Engine/Source/Editor/UnrealEd/Classes/Niagara/NiagaraNodeOp.h
+++ b/Engine/Source/Editor/UnrealEd/Classes/Niagara/NiagaraNodeOp.h
@@ -6,7 +6,9 @@
 UCLASS(MinimalAPI)
 class UNiagaraNodeOp : public UNiagaraNode
 {
-	GENERATED_UCLASS_BODY()
+	GENERATED_BODY()
+public:
+	UNREALED_API UNiagaraNodeOp(const FObjectInitializer& ObjectInitializer = FObjectInitializer::Get());
 
 public:
 
@@ -79,4 +81,4 @@ private:
 	UNiagaraNode* NiagaraNode;
 };
 
-#endif
\ No newline at end of file
+#endif
diff --git a/Engine/Source/Editor/UnrealEd/Classes/Niagara/NiagaraNodeOutput.h b/Engine/Source/Editor/UnrealEd/Classes/Niagara/NiagaraNodeOutput.h
--- a/Engine/Source/Editor/UnrealEd/Classes/Niagara/NiagaraNodeOutput.h
+++ b/Engine/Source/Editor/UnrealEd/Classes/Niagara/NiagaraNodeOutput.h
@@ -6,7 +6,9 @@
 UCLASS(MinimalAPI)
 class UNiagaraNodeOutput : public UNiagaraNode
 {
-	GENERATED_UCLASS_BODY()
+	GENERATED_BODY()
+public:
+	UNREALED_API UNiagaraNodeOutput(const FObjectInitializer& ObjectInitializer = FObjectInitializer::Get());
 
 	void ReallocatePins();
 public:
diff --git a/Engine/Source/Editor/UnrealEd/Classes/Niagara/NiagaraNodeOutputUpdate.h b/Engine/Source/Editor/UnrealEd/Classes/Niagara/NiagaraNodeOutputUpdate.h
--- a/Engine/Source/Editor/UnrealEd/Classes/Niagara/NiagaraNodeOutputUpdate.h
+++ b/Engine/Source/Editor/UnrealEd/Classes/Niagara/NiagaraNodeOutputUpdate.h
@@ -6,7 +6,9 @@
 UCLASS(MinimalAPI, Deprecated)
 class UDEPRECATED_NiagaraNodeOutputUpdate : public UNiagaraNode
 {
-	GENERATED_UCLASS_BODY()
+	GENERATED_BODY()
+public:
+	UNREALED_API UDEPRECATED_NiagaraNodeOutputUpdate(const FObjectInitializer& ObjectInitializer = FObjectInitializer::Get());
 
 public:
 
diff --git a/Engine/Source/Editor/UnrealEd/Classes/Niagara/NiagaraScriptSource.h b/Engine/Source/Editor/UnrealEd/Classes/Niagara/NiagaraScriptSource.h
--- a/Engine/Source/Editor/UnrealEd/Classes/Niagara/NiagaraScriptSource.h
+++ b/Engine/Source/Editor/UnrealEd/Classes/Niagara/NiagaraScriptSource.h
@@ -9,7 +9,9 @@
 UCLASS(MinimalAPI)
 class UNiagaraScriptSource : public UNiagaraScriptSourceBase
 {
-	GENERATED_UCLASS_BODY()
+	GENERATED_BODY()
+public:
+	UNREALED_API UNiagaraScriptSource(const FObjectInitializer& ObjectInitializer = FObjectInitializer::Get());
 
 	/** Graph for particle update expression */
 	UPROPERTY()
diff --git a/Engine/Source/Editor/UnrealEd/Classes/Preferences/CascadeOptions.h b/Engine/Source/Editor/UnrealEd/Classes/Preferences/CascadeOptions.h
--- a/Engine/Source/Editor/UnrealEd/Classes/Preferences/CascadeOptions.h
+++ b/Engine/Source/Editor/UnrealEd/Classes/Preferences/CascadeOptions.h
@@ -13,7 +13,9 @@
 UCLASS(hidecategories=Object, config=EditorUserSettings)
 class UNREALED_API UCascadeOptions : public UObject
 {
-	GENERATED_UCLASS_BODY()
+	GENERATED_BODY()
+public:
+	UCascadeOptions(const FObjectInitializer& ObjectInitializer = FObjectInitializer::Get());
 
 	UPROPERTY(EditAnywhere, config, Category=Options)
 	uint32 bShowModuleDump:1;
diff --git a/Engine/Source/Editor/UnrealEd/Classes/Preferences/CurveEdOptions.h b/Engine/Source/Editor/UnrealEd/Classes/Preferences/CurveEdOptions.h
--- a/Engine/Source/Editor/UnrealEd/Classes/Preferences/CurveEdOptions.h
+++ b/Engine/Source/Editor/UnrealEd/Classes/Preferences/CurveEdOptions.h
@@ -13,7 +13,9 @@
 UCLASS(hidecategories=Object, config=EditorUserSettings)
 class UNREALED_API UCurveEdOptions : public UObject
 {
-	GENERATED_UCLASS_BODY()
+	GENERATED_BODY()
+public:
+	UCurveEdOptions(const FObjectInitializer& ObjectInitializer = FObjectInitializer::Get());
 
 	UPROPERTY(EditAnywhere, config, Category=Options)
 	float MinViewRange;
diff --git a/Engine/Source/Editor/UnrealEd/Classes/Preferences/LightmassOptionsObject.h b/Engine/Source/Editor/UnrealEd/Classes/Preferences/LightmassOptionsObject.h
--- a/Engine/Source/Editor/UnrealEd/Classes/Preferences/LightmassOptionsObject.h
+++ b/Engine/Source/Editor/UnrealEd/Classes/Preferences/LightmassOptionsObject.h
@@ -15,7 +15,9 @@
 UCLASS(hidecategories=Object, editinlinenew)
 class ULightmassOptionsObject : public UObject
 {
-	GENERATED_UCLASS_BODY()
+	GENERATED_BODY()
+public:
+	ULightmassOptionsObject(const FObjectInitializer& ObjectInitializer = FObjectInitializer::Get());
 
 	UPROPERTY(EditAnywhere, Category=Debug)
 	struct FLightmassDebugOptions DebugSettings;
diff --git a/Engine/Source/Editor/UnrealEd/Classes/Preferences/MaterialEditorOptions.h b/Engine/Source/Editor/UnrealEd/Classes/Preferences/MaterialEditorOptions.h
--- a/Engine/Source/Editor/UnrealEd/Classes/Preferences/MaterialEditorOptions.h
+++ b/Engine/Source/Editor/UnrealEd/Classes/Preferences/MaterialEditorOptions.h
@@ -12,7 +12,9 @@
 UCLASS(hidecategories=Object, config=EditorUserSettings)
 class UNREALED_API UMaterialEditorOptions : public UObject
 {
-	GENERATED_UCLASS_BODY()
+	GENERATED_BODY()
+public:
+	UMaterialEditorOptions(const FObjectInitializer& ObjectInitializer = FObjectInitializer::Get());
 
 	/** If true, render grid the preview scene. */
 	UPROPERTY(EditAnywhere, config, Category=Options)
diff --git a/Engine/Source/Editor/UnrealEd/Classes/Preferences/PersonaOptions.h b/Engine/Source/Editor/UnrealEd/Classes/Preferences/PersonaOptions.h
--- a/Engine/Source/Editor/UnrealEd/Classes/Preferences/PersonaOptions.h
+++ b/Engine/Source/Editor/UnrealEd/Classes/Preferences/PersonaOptions.h
@@ -13,7 +13,9 @@
 UCLASS(hidecategories=Object, config=EditorUserSettings)
 class UNREALED_API UPersonaOptions : public UObject
 {
-	GENERATED_UCLASS_BODY()
+	GENERATED_BODY()
+public:
+	UPersonaOptions(const FObjectInitializer& ObjectInitializer = FObjectInitializer::Get());
 
 	UPROPERTY(EditAnywhere, config, Category=Options)
 	uint32 bShowFloor:1;
diff --git a/Engine/Source/Editor/UnrealEd/Classes/Preferences/PhATSimOptions.h b/Engine/Source/Editor/UnrealEd/Classes/Preferences/PhATSimOptions.h
--- a/Engine/Source/Editor/UnrealEd/Classes/Preferences/PhATSimOptions.h
+++ b/Engine/Source/Editor/UnrealEd/Classes/Preferences/PhATSimOptions.h
@@ -7,7 +7,9 @@
 UCLASS(hidecategories=Object, config=EditorUserSettings)
 class UNREALED_API UPhATSimOptions : public UObject
 {
-	GENERATED_UCLASS_BODY()
+	GENERATED_BODY()
+public:
+	UPhATSimOptions(const FObjectInitializer& ObjectInitializer = FObjectInitializer::Get());
 
 	/** Lets you manually control the physics/animation */
 	UPROPERTY(EditAnywhere, transient, Category=Anim)
diff --git a/Engine/Source/Editor/UnrealEd/Classes/Preferences/UnrealEdKeyBindings.h b/Engine/Source/Editor/UnrealEd/Classes/Preferences/UnrealEdKeyBindings.h
--- a/Engine/Source/Editor/UnrealEd/Classes/Preferences/UnrealEdKeyBindings.h
+++ b/Engine/Source/Editor/UnrealEd/Classes/Preferences/UnrealEdKeyBindings.h
@@ -11,7 +11,7 @@
 USTRUCT()
 struct FEditorKeyBinding
 {
-	GENERATED_USTRUCT_BODY()
+	GENERATED_BODY()
 
 	UPROPERTY()
 	uint32 bCtrlDown:1;
@@ -38,7 +38,9 @@ struct FEditorKeyBinding
 UCLASS(Config=EditorKeyBindings)
 class UUnrealEdKeyBindings : public UObject
 {
-	GENERATED_UCLASS_BODY()
+	GENERATED_BODY()
+public:
+	UUnrealEdKeyBindings(const FObjectInitializer& ObjectInitializer = FObjectInitializer::Get());
 
 	/** Array of keybindings */
 	UPROPERTY(config)
diff --git a/Engine/Source/Editor/UnrealEd/Classes/Preferences/UnrealEdOptions.h b/Engine/Source/Editor/UnrealEd/Classes/Preferences/UnrealEdOptions.h
--- a/Engine/Source/Editor/UnrealEd/Classes/Preferences/UnrealEdOptions.h
+++ b/Engine/Source/Editor/UnrealEd/Classes/Preferences/UnrealEdOptions.h
@@ -12,7 +12,7 @@
 USTRUCT()
 struct FEditorCommandCategory
 {
-	GENERATED_USTRUCT_BODY()
+	GENERATED_BODY()
 
 	UPROPERTY()
 	FName Parent;
@@ -26,7 +26,7 @@ struct FEditorCommandCategory
 USTRUCT()
 struct FEditorCommand
 {
-	GENERATED_USTRUCT_BODY()
+	GENERATED_BODY()
 
 	UPROPERTY()
 	FName Parent;
@@ -46,7 +46,7 @@ struct FEditorCommand
 USTRUCT()
 struct FClassPickerDefaults
 {
-	GENERATED_USTRUCT_BODY()
+	GENERATED_BODY()
 
 	/** The name of the class to select */
 	UPROPERTY()
@@ -66,7 +66,9 @@ struct FClassPickerDefaults
 UCLASS(Config=Editor)
 class UNREALED_API UUnrealEdOptions : public UObject
 {
-	GENERATED_UCLASS_BODY()
+	GENERATED_BODY()
+public:
+	UUnrealEdOptions(const FObjectInitializer& ObjectInitializer = FObjectInitializer::Get());
 
 	/** Categories of commands. */
 	UPROPERTY(config)
diff --git a/Engine/Source/Editor/UnrealEd/Classes/Settings/ContentBrowserSettings.h b/Engine/Source/Editor/UnrealEd/Classes/Settings/ContentBrowserSettings.h
--- a/Engine/Source/Editor/UnrealEd/Classes/Settings/ContentBrowserSettings.h
+++ b/Engine/Source/Editor/UnrealEd/Classes/Settings/ContentBrowserSettings.h
@@ -17,7 +17,9 @@ UCLASS(config=EditorUserSettings)
 class UNREALED_API UContentBrowserSettings
 	: public UObject
 {
-	GENERATED_UCLASS_BODY()
+	GENERATED_BODY()
+public:
+	UContentBrowserSettings(const FObjectInitializer& ObjectInitializer = FObjectInitializer::Get());
 
 public:
 
diff --git a/Engine/Source/Editor/UnrealEd/Classes/Settings/DestructableMeshEditorSettings.h b/Engine/Source/Editor/UnrealEd/Classes/Settings/DestructableMeshEditorSettings.h
--- a/Engine/Source/Editor/UnrealEd/Classes/Settings/DestructableMeshEditorSettings.h
+++ b/Engine/Source/Editor/UnrealEd/Classes/Settings/DestructableMeshEditorSettings.h
@@ -17,7 +17,9 @@ UCLASS(config=EditorUserSettings)
 class UNREALED_API UDestructableMeshEditorSettings
 	: public UObject
 {
-	GENERATED_UCLASS_BODY()
+	GENERATED_BODY()
+public:
+	UDestructableMeshEditorSettings(const FObjectInitializer& ObjectInitializer = FObjectInitializer::Get());
 
 public:
 
diff --git a/Engine/Source/Editor/UnrealEd/Classes/Settings/EditorExperimentalSettings.h b/Engine/Source/Editor/UnrealEd/Classes/Settings/EditorExperimentalSettings.h
--- a/Engine/Source/Editor/UnrealEd/Classes/Settings/EditorExperimentalSettings.h
+++ b/Engine/Source/Editor/UnrealEd/Classes/Settings/EditorExperimentalSettings.h
@@ -17,7 +17,9 @@ UCLASS(config=EditorUserSettings)
 class UNREALED_API UEditorExperimentalSettings
 	: public UObject
 {
-	GENERATED_UCLASS_BODY()
+	GENERATED_BODY()
+public:
+	UEditorExperimentalSettings(const FObjectInitializer& ObjectInitializer = FObjectInitializer::Get());
 
 public:
 
diff --git a/Engine/Source/Editor/UnrealEd/Classes/Settings/EditorLoadingSavingSettings.h b/Engine/Source/Editor/UnrealEd/Classes/Settings/EditorLoadingSavingSettings.h
--- a/Engine/Source/Editor/UnrealEd/Classes/Settings/EditorLoadingSavingSettings.h
+++ b/Engine/Source/Editor/UnrealEd/Classes/Settings/EditorLoadingSavingSettings.h
@@ -23,7 +23,9 @@ UCLASS(config=EditorUserSettings, autoexpandcategories=(AutoSave, AutoReimport,
 class UNREALED_API UEditorLoadingSavingSettings
 	: public UObject
 {
-	GENERATED_UCLASS_BODY()
+	GENERATED_BODY()
+public:
+	UEditorLoadingSavingSettings(const FObjectInitializer& ObjectInitializer = FObjectInitializer::Get());
 
 public:
 
diff --git a/Engine/Source/Editor/UnrealEd/Classes/Settings/EditorMiscSettings.h b/Engine/Source/Editor/UnrealEd/Classes/Settings/EditorMiscSettings.h
--- a/Engine/Source/Editor/UnrealEd/Classes/Settings/EditorMiscSettings.h
+++ b/Engine/Source/Editor/UnrealEd/Classes/Settings/EditorMiscSettings.h
@@ -13,7 +13,9 @@ UCLASS(config=EditorUserSettings)
 class UNREALED_API UEditorMiscSettings
 	: public UObject
 {
-	GENERATED_UCLASS_BODY()
+	GENERATED_BODY()
+public:
+	UEditorMiscSettings(const FObjectInitializer& ObjectInitializer = FObjectInitializer::Get());
 
 public:
 
diff --git a/Engine/Source/Editor/UnrealEd/Classes/Settings/LevelEditorMiscSettings.h b/Engine/Source/Editor/UnrealEd/Classes/Settings/LevelEditorMiscSettings.h
--- a/Engine/Source/Editor/UnrealEd/Classes/Settings/LevelEditorMiscSettings.h
+++ b/Engine/Source/Editor/UnrealEd/Classes/Settings/LevelEditorMiscSettings.h
@@ -12,7 +12,9 @@ UCLASS(config=EditorUserSettings)
 class UNREALED_API ULevelEditorMiscSettings
 	: public UObject
 {
-	GENERATED_UCLASS_BODY()
+	GENERATED_BODY()
+public:
+	ULevelEditorMiscSettings(const FObjectInitializer& ObjectInitializer = FObjectInitializer::Get());
 
 public:
 
diff --git a/Engine/Source/Editor/UnrealEd/Classes/Settings/LevelEditorPlaySettings.h b/Engine/Source/Editor/UnrealEd/Classes/Settings/LevelEditorPlaySettings.h
--- a/Engine/Source/Editor/UnrealEd/Classes/Settings/LevelEditorPlaySettings.h
+++ b/Engine/Source/Editor/UnrealEd/Classes/Settings/LevelEditorPlaySettings.h
@@ -84,7 +84,7 @@ enum EPlayNetMode
 USTRUCT()
 struct FPlayScreenResolution
 {
-	GENERATED_USTRUCT_BODY()
+	GENERATED_BODY()
 
 public:
 
@@ -113,7 +113,9 @@ UCLASS(config=EditorUserSettings)
 class UNREALED_API ULevelEditorPlaySettings
 	: public UObject
 {
-	GENERATED_UCLASS_BODY()
+	GENERATED_BODY()
+public:
+	ULevelEditorPlaySettings(const FObjectInitializer& ObjectInitializer = FObjectInitializer::Get());
 
 public:
 
diff --git a/Engine/Source/Editor/UnrealEd/Classes/Settings/LevelEditorViewportSettings.h b/Engine/Source/Editor/UnrealEd/Classes/Settings/LevelEditorViewportSettings.h
--- a/Engine/Source/Editor/UnrealEd/Classes/Settings/LevelEditorViewportSettings.h
+++ b/Engine/Source/Editor/UnrealEd/Classes/Settings/LevelEditorViewportSettings.h
@@ -63,7 +63,7 @@ enum EMeasuringToolUnits
 USTRUCT()
 struct UNREALED_API FLevelEditorViewportInstanceSettings
 {
-	GENERATED_USTRUCT_BODY()
+	GENERATED_BODY()
 
 	FLevelEditorViewportInstanceSettings()
 		: ViewportType(LVT_Perspective)
@@ -141,7 +141,7 @@ struct UNREALED_API FLevelEditorViewportInstanceSettings
 USTRUCT()
 struct UNREALED_API FLevelEditorViewportInstanceSettingsKeyValuePair
 {
-	GENERATED_USTRUCT_BODY()
+	GENERATED_BODY()
 
 	/*  Name identifying this config. */
 	UPROPERTY(config)
@@ -159,7 +159,7 @@ struct UNREALED_API FLevelEditorViewportInstanceSettingsKeyValuePair
 USTRUCT()
 struct UNREALED_API FSnapToSurfaceSettings
 {
-	GENERATED_USTRUCT_BODY()
+	GENERATED_BODY()
 
 	FSnapToSurfaceSettings()
 		: bEnabled(false)
@@ -188,7 +188,9 @@ UCLASS(config=EditorUserSettings)
 class UNREALED_API ULevelEditorViewportSettings
 	: public UObject
 {
-	GENERATED_UCLASS_BODY()
+	GENERATED_BODY()
+public:
+	ULevelEditorViewportSettings(const FObjectInitializer& ObjectInitializer = FObjectInitializer::Get());
 
 	/** Enable the use of flight camera controls under various circumstances. */
 	UPROPERTY(EditAnywhere, config, Category=Controls)
diff --git a/Engine/Source/Editor/UnrealEd/Classes/Settings/ProjectPackagingSettings.h b/Engine/Source/Editor/UnrealEd/Classes/Settings/ProjectPackagingSettings.h
--- a/Engine/Source/Editor/UnrealEd/Classes/Settings/ProjectPackagingSettings.h
+++ b/Engine/Source/Editor/UnrealEd/Classes/Settings/ProjectPackagingSettings.h
@@ -50,7 +50,9 @@ UCLASS(config=Game, defaultconfig)
 class UNREALED_API UProjectPackagingSettings
 	: public UObject
 {
-	GENERATED_UCLASS_BODY()
+	GENERATED_BODY()
+public:
+	UProjectPackagingSettings(const FObjectInitializer& ObjectInitializer = FObjectInitializer::Get());
 
 public:
 
diff --git a/Engine/Source/Editor/UnrealEd/Classes/SoundClassGraph/SoundClassGraph.h b/Engine/Source/Editor/UnrealEd/Classes/SoundClassGraph/SoundClassGraph.h
--- a/Engine/Source/Editor/UnrealEd/Classes/SoundClassGraph/SoundClassGraph.h
+++ b/Engine/Source/Editor/UnrealEd/Classes/SoundClassGraph/SoundClassGraph.h
@@ -7,7 +7,9 @@
 UCLASS(MinimalAPI)
 class USoundClassGraph : public UEdGraph
 {
-	GENERATED_UCLASS_BODY()
+	GENERATED_BODY()
+public:
+	UNREALED_API USoundClassGraph(const FObjectInitializer& ObjectInitializer = FObjectInitializer::Get());
 
 	/**
 	 * Set the SoundClass which forms the root of this graph
diff --git a/Engine/Source/Editor/UnrealEd/Classes/SoundClassGraph/SoundClassGraphNode.h b/Engine/Source/Editor/UnrealEd/Classes/SoundClassGraph/SoundClassGraphNode.h
--- a/Engine/Source/Editor/UnrealEd/Classes/SoundClassGraph/SoundClassGraphNode.h
+++ b/Engine/Source/Editor/UnrealEd/Classes/SoundClassGraph/SoundClassGraphNode.h
@@ -6,7 +6,9 @@
 UCLASS(MinimalAPI)
 class USoundClassGraphNode : public UEdGraphNode
 {
-	GENERATED_UCLASS_BODY()
+	GENERATED_BODY()
+public:
+	UNREALED_API USoundClassGraphNode(const FObjectInitializer& ObjectInitializer = FObjectInitializer::Get());
 
 	/** The SoundNode this represents */
 	UPROPERTY(VisibleAnywhere, instanced, Category=Sound)
diff --git a/Engine/Source/Editor/UnrealEd/Classes/SoundClassGraph/SoundClassGraphSchema.h b/Engine/Source/Editor/UnrealEd/Classes/SoundClassGraph/SoundClassGraphSchema.h
--- a/Engine/Source/Editor/UnrealEd/Classes/SoundClassGraph/SoundClassGraphSchema.h
+++ b/Engine/Source/Editor/UnrealEd/Classes/SoundClassGraph/SoundClassGraphSchema.h
@@ -7,7 +7,7 @@
 USTRUCT()
 struct UNREALED_API FSoundClassGraphSchemaAction_NewNode : public FEdGraphSchemaAction
 {
-	GENERATED_USTRUCT_BODY();
+	GENERATED_BODY()
 
 	// Simple type info
 	static FName StaticGetTypeId() {static FName Type("FSoundClassGraphSchemaAction_NewNode"); return Type;}
@@ -34,7 +34,9 @@ struct UNREALED_API FSoundClassGraphSchemaAction_NewNode : public FEdGraphSchema
 UCLASS(MinimalAPI)
 class USoundClassGraphSchema : public UEdGraphSchema
 {
-	GENERATED_UCLASS_BODY()
+	GENERATED_BODY()
+public:
+	UNREALED_API USoundClassGraphSchema(const FObjectInitializer& ObjectInitializer = FObjectInitializer::Get());
 
 	/** Check whether connecting these pins would cause a loop */
 	bool ConnectionCausesLoop(const UEdGraphPin* InputPin, const UEdGraphPin* OutputPin) const;
diff --git a/Engine/Source/Editor/UnrealEd/Classes/SoundCueGraph/SoundCueGraph.h b/Engine/Source/Editor/UnrealEd/Classes/SoundCueGraph/SoundCueGraph.h
--- a/Engine/Source/Editor/UnrealEd/Classes/SoundCueGraph/SoundCueGraph.h
+++ b/Engine/Source/Editor/UnrealEd/Classes/SoundCueGraph/SoundCueGraph.h
@@ -7,7 +7,9 @@
 UCLASS(MinimalAPI)
 class USoundCueGraph : public UEdGraph
 {
-	GENERATED_UCLASS_BODY()
+	GENERATED_BODY()
+public:
+	UNREALED_API USoundCueGraph(const FObjectInitializer& ObjectInitializer = FObjectInitializer::Get());
 
 	/** Returns the SoundCue that contains this graph */
 	UNREALED_API class USoundCue* GetSoundCue() const;
diff --git a/Engine/Source/Editor/UnrealEd/Classes/SoundCueGraph/SoundCueGraphNode.h b/Engine/Source/Editor/UnrealEd/Classes/SoundCueGraph/SoundCueGraphNode.h
--- a/Engine/Source/Editor/UnrealEd/Classes/SoundCueGraph/SoundCueGraphNode.h
+++ b/Engine/Source/Editor/UnrealEd/Classes/SoundCueGraph/SoundCueGraphNode.h
@@ -8,7 +8,9 @@
 UCLASS(MinimalAPI)
 class USoundCueGraphNode : public USoundCueGraphNode_Base
 {
-	GENERATED_UCLASS_BODY()
+	GENERATED_BODY()
+public:
+	UNREALED_API USoundCueGraphNode(const FObjectInitializer& ObjectInitializer = FObjectInitializer::Get());
 
 	/** The SoundNode this represents */
 	UPROPERTY(VisibleAnywhere, instanced, Category=Sound)
diff --git a/Engine/Source/Editor/UnrealEd/Classes/SoundCueGraph/SoundCueGraphNode_Base.h b/Engine/Source/Editor/UnrealEd/Classes/SoundCueGraph/SoundCueGraphNode_Base.h
--- a/Engine/Source/Editor/UnrealEd/Classes/SoundCueGraph/SoundCueGraphNode_Base.h
+++ b/Engine/Source/Editor/UnrealEd/Classes/SoundCueGraph/SoundCueGraphNode_Base.h
@@ -7,7 +7,9 @@
 UCLASS(MinimalAPI)
 class USoundCueGraphNode_Base : public UEdGraphNode
 {
-	GENERATED_UCLASS_BODY()
+	GENERATED_BODY()
+public:
+	UNREALED_API USoundCueGraphNode_Base(const FObjectInitializer& ObjectInitializer = FObjectInitializer::Get());
 
 	/** Create all of the input pins required */
 	virtual void CreateInputPins() {};
diff --git a/Engine/Source/Editor/UnrealEd/Classes/SoundCueGraph/SoundCueGraphNode_Root.h b/Engine/Source/Editor/UnrealEd/Classes/SoundCueGraph/SoundCueGraphNode_Root.h
--- a/Engine/Source/Editor/UnrealEd/Classes/SoundCueGraph/SoundCueGraphNode_Root.h
+++ b/Engine/Source/Editor/UnrealEd/Classes/SoundCueGraph/SoundCueGraphNode_Root.h
@@ -6,7 +6,9 @@
 UCLASS(MinimalAPI)
 class USoundCueGraphNode_Root : public USoundCueGraphNode_Base
 {
-	GENERATED_UCLASS_BODY()
+	GENERATED_BODY()
+public:
+	UNREALED_API USoundCueGraphNode_Root(const FObjectInitializer& ObjectInitializer = FObjectInitializer::Get());
 
 	// UEdGraphNode interface
 	virtual FLinearColor GetNodeTitleColor() const override;
diff --git a/Engine/Source/Editor/UnrealEd/Classes/SoundCueGraph/SoundCueGraphSchema.h b/Engine/Source/Editor/UnrealEd/Classes/SoundCueGraph/SoundCueGraphSchema.h
--- a/Engine/Source/Editor/UnrealEd/Classes/SoundCueGraph/SoundCueGraphSchema.h
+++ b/Engine/Source/Editor/UnrealEd/Classes/SoundCueGraph/SoundCueGraphSchema.h
@@ -11,7 +11,7 @@ class USoundNodeWavePlayer;
 USTRUCT()
 struct UNREALED_API FSoundCueGraphSchemaAction_NewNode : public FEdGraphSchemaAction
 {
-	GENERATED_USTRUCT_BODY();
+	GENERATED_BODY()
 
 	/** Class of node we want to create */
 	UPROPERTY()
@@ -41,7 +41,7 @@ private:
 USTRUCT()
 struct UNREALED_API FSoundCueGraphSchemaAction_NewFromSelected : public FSoundCueGraphSchemaAction_NewNode
 {
-	GENERATED_USTRUCT_BODY();
+	GENERATED_BODY()
 
 	FSoundCueGraphSchemaAction_NewFromSelected() 
 		: FSoundCueGraphSchemaAction_NewNode()
@@ -60,7 +60,7 @@ struct UNREALED_API FSoundCueGraphSchemaAction_NewFromSelected : public FSoundCu
 USTRUCT()
 struct UNREALED_API FSoundCueGraphSchemaAction_NewComment : public FEdGraphSchemaAction
 {
-	GENERATED_USTRUCT_BODY();
+	GENERATED_BODY()
 
 	FSoundCueGraphSchemaAction_NewComment() 
 		: FEdGraphSchemaAction()
@@ -79,7 +79,7 @@ struct UNREALED_API FSoundCueGraphSchemaAction_NewComment : public FEdGraphSchem
 USTRUCT()
 struct UNREALED_API FSoundCueGraphSchemaAction_Paste : public FEdGraphSchemaAction
 {
-	GENERATED_USTRUCT_BODY();
+	GENERATED_BODY()
 
 	FSoundCueGraphSchemaAction_Paste() 
 		: FEdGraphSchemaAction()
@@ -97,7 +97,9 @@ struct UNREALED_API FSoundCueGraphSchemaAction_Paste : public FEdGraphSchemaActi
 UCLASS(MinimalAPI)
 class USoundCueGraphSchema : public UEdGraphSchema
 {
-	GENERATED_UCLASS_BODY()
+	GENERATED_BODY()
+public:
+	UNREALED_API USoundCueGraphSchema(const FObjectInitializer& ObjectInitializer = FObjectInitializer::Get());
 
 	/** Check whether connecting these pins would cause a loop */
 	bool ConnectionCausesLoop(const UEdGraphPin* InputPin, const UEdGraphPin* OutputPin) const;
diff --git a/Engine/Source/Editor/UnrealEd/Classes/TexAligner/TexAligner.h b/Engine/Source/Editor/UnrealEd/Classes/TexAligner/TexAligner.h
--- a/Engine/Source/Editor/UnrealEd/Classes/TexAligner/TexAligner.h
+++ b/Engine/Source/Editor/UnrealEd/Classes/TexAligner/TexAligner.h
@@ -38,7 +38,9 @@ enum ETexAlign
 UCLASS()
 class UTexAligner : public UObject
 {
-	GENERATED_UCLASS_BODY()
+	GENERATED_BODY()
+public:
+	UTexAligner(const FObjectInitializer& ObjectInitializer = FObjectInitializer::Get());
 
 	/** The default alignment this aligner represents. */
 	UPROPERTY()
diff --git a/Engine/Source/Editor/UnrealEd/Classes/TexAligner/TexAlignerBox.h b/Engine/Source/Editor/UnrealEd/Classes/TexAligner/TexAlignerBox.h
--- a/Engine/Source/Editor/UnrealEd/Classes/TexAligner/TexAlignerBox.h
+++ b/Engine/Source/Editor/UnrealEd/Classes/TexAligner/TexAlignerBox.h
@@ -12,7 +12,9 @@
 UCLASS(hidecategories=Object)
 class UTexAlignerBox : public UTexAligner
 {
-	GENERATED_UCLASS_BODY()
+	GENERATED_BODY()
+public:
+	UTexAlignerBox(const FObjectInitializer& ObjectInitializer = FObjectInitializer::Get());
 
 
 	// Begin UObject Interface
diff --git a/Engine/Source/Editor/UnrealEd/Classes/TexAligner/TexAlignerDefault.h b/Engine/Source/Editor/UnrealEd/Classes/TexAligner/TexAlignerDefault.h
--- a/Engine/Source/Editor/UnrealEd/Classes/TexAligner/TexAlignerDefault.h
+++ b/Engine/Source/Editor/UnrealEd/Classes/TexAligner/TexAlignerDefault.h
@@ -12,7 +12,9 @@
 UCLASS(hidecategories=Object)
 class UTexAlignerDefault : public UTexAligner
 {
-	GENERATED_UCLASS_BODY()
+	GENERATED_BODY()
+public:
+	UTexAlignerDefault(const FObjectInitializer& ObjectInitializer = FObjectInitializer::Get());
 
 
 	// Begin UObject Interface
diff --git a/Engine/Source/Editor/UnrealEd/Classes/TexAligner/TexAlignerFit.h b/Engine/Source/Editor/UnrealEd/Classes/TexAligner/TexAlignerFit.h
--- a/Engine/Source/Editor/UnrealEd/Classes/TexAligner/TexAlignerFit.h
+++ b/Engine/Source/Editor/UnrealEd/Classes/TexAligner/TexAlignerFit.h
@@ -12,7 +12,9 @@
 UCLASS(hidecategories=Object)
 class UTexAlignerFit : public UTexAligner
 {
-	GENERATED_UCLASS_BODY()
+	GENERATED_BODY()
+public:
+	UTexAlignerFit(const FObjectInitializer& ObjectInitializer = FObjectInitializer::Get());
 
 
 	// Begin UObject Interface
diff --git a/Engine/Source/Editor/UnrealEd/Classes/TexAligner/TexAlignerPlanar.h b/Engine/Source/Editor/UnrealEd/Classes/TexAligner/TexAlignerPlanar.h
--- a/Engine/Source/Editor/UnrealEd/Classes/TexAligner/TexAlignerPlanar.h
+++ b/Engine/Source/Editor/UnrealEd/Classes/TexAligner/TexAlignerPlanar.h
@@ -12,7 +12,9 @@
 UCLASS(hidecategories=Object)
 class UTexAlignerPlanar : public UTexAligner
 {
-	GENERATED_UCLASS_BODY()
+	GENERATED_BODY()
+public:
+	UTexAlignerPlanar(const FObjectInitializer& ObjectInitializer = FObjectInitializer::Get());
 
 
 	// Begin UObject Interface
diff --git a/Engine/Source/Editor/UnrealEd/Classes/ThumbnailRendering/AnimBlueprintThumbnailRenderer.h b/Engine/Source/Editor/UnrealEd/Classes/ThumbnailRendering/AnimBlueprintThumbnailRenderer.h
--- a/Engine/Source/Editor/UnrealEd/Classes/ThumbnailRendering/AnimBlueprintThumbnailRenderer.h
+++ b/Engine/Source/Editor/UnrealEd/Classes/ThumbnailRendering/AnimBlueprintThumbnailRenderer.h
@@ -11,7 +11,9 @@
 UCLASS(config=Editor, MinimalAPI)
 class UAnimBlueprintThumbnailRenderer : public UDefaultSizedThumbnailRenderer
 {
-	GENERATED_UCLASS_BODY()
+	GENERATED_BODY()
+public:
+	UNREALED_API UAnimBlueprintThumbnailRenderer(const FObjectInitializer& ObjectInitializer = FObjectInitializer::Get());
 
 
 	// Begin UThumbnailRenderer Object
diff --git a/Engine/Source/Editor/UnrealEd/Classes/ThumbnailRendering/AnimSequenceThumbnailRenderer.h b/Engine/Source/Editor/UnrealEd/Classes/ThumbnailRendering/AnimSequenceThumbnailRenderer.h
--- a/Engine/Source/Editor/UnrealEd/Classes/ThumbnailRendering/AnimSequenceThumbnailRenderer.h
+++ b/Engine/Source/Editor/UnrealEd/Classes/ThumbnailRendering/AnimSequenceThumbnailRenderer.h
@@ -11,7 +11,9 @@
 UCLASS(config=Editor, MinimalAPI)
 class UAnimSequenceThumbnailRenderer : public UDefaultSizedThumbnailRenderer
 {
-	GENERATED_UCLASS_BODY()
+	GENERATED_BODY()
+public:
+	UNREALED_API UAnimSequenceThumbnailRenderer(const FObjectInitializer& ObjectInitializer = FObjectInitializer::Get());
 
 
 	// Begin UThumbnailRenderer Object
diff --git a/Engine/Source/Editor/UnrealEd/Classes/ThumbnailRendering/BlendSpaceThumbnailRenderer.h b/Engine/Source/Editor/UnrealEd/Classes/ThumbnailRendering/BlendSpaceThumbnailRenderer.h
--- a/Engine/Source/Editor/UnrealEd/Classes/ThumbnailRendering/BlendSpaceThumbnailRenderer.h
+++ b/Engine/Source/Editor/UnrealEd/Classes/ThumbnailRendering/BlendSpaceThumbnailRenderer.h
@@ -11,7 +11,9 @@
 UCLASS(config=Editor, MinimalAPI)
 class UBlendSpaceThumbnailRenderer : public UDefaultSizedThumbnailRenderer
 {
-	GENERATED_UCLASS_BODY()
+	GENERATED_BODY()
+public:
+	UNREALED_API UBlendSpaceThumbnailRenderer(const FObjectInitializer& ObjectInitializer = FObjectInitializer::Get());
 
 
 	// Begin UThumbnailRenderer Object
diff --git a/Engine/Source/Editor/UnrealEd/Classes/ThumbnailRendering/BlueprintThumbnailRenderer.h b/Engine/Source/Editor/UnrealEd/Classes/ThumbnailRendering/BlueprintThumbnailRenderer.h
--- a/Engine/Source/Editor/UnrealEd/Classes/ThumbnailRendering/BlueprintThumbnailRenderer.h
+++ b/Engine/Source/Editor/UnrealEd/Classes/ThumbnailRendering/BlueprintThumbnailRenderer.h
@@ -11,7 +11,9 @@
 UCLASS(config=Editor,MinimalAPI)
 class UBlueprintThumbnailRenderer : public UDefaultSizedThumbnailRenderer
 {
-	GENERATED_UCLASS_BODY()
+	GENERATED_BODY()
+public:
+	UNREALED_API UBlueprintThumbnailRenderer(const FObjectInitializer& ObjectInitializer = FObjectInitializer::Get());
 
 
 	// Begin UThumbnailRenderer Object
diff --git a/Engine/Source/Editor/UnrealEd/Classes/ThumbnailRendering/ClassThumbnailRenderer.h b/Engine/Source/Editor/UnrealEd/Classes/ThumbnailRendering/ClassThumbnailRenderer.h
--- a/Engine/Source/Editor/UnrealEd/Classes/ThumbnailRendering/ClassThumbnailRenderer.h
+++ b/Engine/Source/Editor/UnrealEd/Classes/ThumbnailRendering/ClassThumbnailRenderer.h
@@ -6,7 +6,9 @@
 UCLASS(config=Editor,MinimalAPI)
 class UClassThumbnailRenderer : public UDefaultSizedThumbnailRenderer
 {
-	GENERATED_UCLASS_BODY()
+	GENERATED_BODY()
+public:
+	UNREALED_API UClassThumbnailRenderer(const FObjectInitializer& ObjectInitializer = FObjectInitializer::Get());
 
 
 	// Begin UThumbnailRenderer Object
diff --git a/Engine/Source/Editor/UnrealEd/Classes/ThumbnailRendering/DefaultSizedThumbnailRenderer.h b/Engine/Source/Editor/UnrealEd/Classes/ThumbnailRendering/DefaultSizedThumbnailRenderer.h
--- a/Engine/Source/Editor/UnrealEd/Classes/ThumbnailRendering/DefaultSizedThumbnailRenderer.h
+++ b/Engine/Source/Editor/UnrealEd/Classes/ThumbnailRendering/DefaultSizedThumbnailRenderer.h
@@ -11,7 +11,9 @@
 UCLASS(abstract, config=Editor, MinimalAPI)
 class UDefaultSizedThumbnailRenderer : public UThumbnailRenderer
 {
-	GENERATED_UCLASS_BODY()
+	GENERATED_BODY()
+public:
+	UNREALED_API UDefaultSizedThumbnailRenderer(const FObjectInitializer& ObjectInitializer = FObjectInitializer::Get());
 
 	/**
 	 * The default width of this thumbnail
diff --git a/Engine/Source/Editor/UnrealEd/Classes/ThumbnailRendering/FontThumbnailRenderer.h b/Engine/Source/Editor/UnrealEd/Classes/ThumbnailRendering/FontThumbnailRenderer.h
--- a/Engine/Source/Editor/UnrealEd/Classes/ThumbnailRendering/FontThumbnailRenderer.h
+++ b/Engine/Source/Editor/UnrealEd/Classes/ThumbnailRendering/FontThumbnailRenderer.h
@@ -11,7 +11,9 @@
 UCLASS()
 class UFontThumbnailRenderer : public UTextureThumbnailRenderer
 {
-	GENERATED_UCLASS_BODY()
+	GENERATED_BODY()
+public:
+	UFontThumbnailRenderer(const FObjectInitializer& ObjectInitializer = FObjectInitializer::Get());
 
 
 	// Begin UThumbnailRenderer Object
diff --git a/Engine/Source/Editor/UnrealEd/Classes/ThumbnailRendering/LevelThumbnailRenderer.h b/Engine/Source/Editor/UnrealEd/Classes/ThumbnailRendering/LevelThumbnailRenderer.h
--- a/Engine/Source/Editor/UnrealEd/Classes/ThumbnailRendering/LevelThumbnailRenderer.h
+++ b/Engine/Source/Editor/UnrealEd/Classes/ThumbnailRendering/LevelThumbnailRenderer.h
@@ -11,7 +11,9 @@
 UCLASS(config=Editor)
 class ULevelThumbnailRenderer : public UDefaultSizedThumbnailRenderer
 {
-	GENERATED_UCLASS_BODY()
+	GENERATED_BODY()
+public:
+	ULevelThumbnailRenderer(const FObjectInitializer& ObjectInitializer = FObjectInitializer::Get());
 
 
 	// Begin UThumbnailRenderer Object
diff --git a/Engine/Source/Editor/UnrealEd/Classes/ThumbnailRendering/MaterialFunctionThumbnailRenderer.h b/Engine/Source/Editor/UnrealEd/Classes/ThumbnailRendering/MaterialFunctionThumbnailRenderer.h
--- a/Engine/Source/Editor/UnrealEd/Classes/ThumbnailRendering/MaterialFunctionThumbnailRenderer.h
+++ b/Engine/Source/Editor/UnrealEd/Classes/ThumbnailRendering/MaterialFunctionThumbnailRenderer.h
@@ -11,7 +11,9 @@
 UCLASS(config=Editor)
 class UMaterialFunctionThumbnailRenderer : public UDefaultSizedThumbnailRenderer
 {
-	GENERATED_UCLASS_BODY()
+	GENERATED_BODY()
+public:
+	UMaterialFunctionThumbnailRenderer(const FObjectInitializer& ObjectInitializer = FObjectInitializer::Get());
 
 
 	// Begin UThumbnailRenderer Object
diff --git a/Engine/Source/Editor/UnrealEd/Classes/ThumbnailRendering/MaterialInstanceThumbnailRenderer.h b/Engine/Source/Editor/UnrealEd/Classes/ThumbnailRendering/MaterialInstanceThumbnailRenderer.h
--- a/Engine/Source/Editor/UnrealEd/Classes/ThumbnailRendering/MaterialInstanceThumbnailRenderer.h
+++ b/Engine/Source/Editor/UnrealEd/Classes/ThumbnailRendering/MaterialInstanceThumbnailRenderer.h
@@ -11,7 +11,9 @@
 UCLASS(config=Editor)
 class UMaterialInstanceThumbnailRenderer : public UDefaultSizedThumbnailRenderer
 {
-	GENERATED_UCLASS_BODY()
+	GENERATED_BODY()
+public:
+	UMaterialInstanceThumbnailRenderer(const FObjectInitializer& ObjectInitializer = FObjectInitializer::Get());
 
 
 	// Begin UThumbnailRenderer Object
diff --git a/Engine/Source/Editor/UnrealEd/Classes/ThumbnailRendering/ParticleSystemThumbnailRenderer.h b/Engine/Source/Editor/UnrealEd/Classes/ThumbnailRendering/ParticleSystemThumbnailRenderer.h
--- a/Engine/Source/Editor/UnrealEd/Classes/ThumbnailRendering/ParticleSystemThumbnailRenderer.h
+++ b/Engine/Source/Editor/UnrealEd/Classes/ThumbnailRendering/ParticleSystemThumbnailRenderer.h
@@ -11,7 +11,9 @@
 UCLASS(config=Editor)
 class UParticleSystemThumbnailRenderer : public UTextureThumbnailRenderer
 {
-	GENERATED_UCLASS_BODY()
+	GENERATED_BODY()
+public:
+	UParticleSystemThumbnailRenderer(const FObjectInitializer& ObjectInitializer = FObjectInitializer::Get());
 
 	UPROPERTY()
 	class UTexture2D* NoImage;
diff --git a/Engine/Source/Editor/UnrealEd/Classes/ThumbnailRendering/SceneThumbnailInfo.h b/Engine/Source/Editor/UnrealEd/Classes/ThumbnailRendering/SceneThumbnailInfo.h
--- a/Engine/Source/Editor/UnrealEd/Classes/ThumbnailRendering/SceneThumbnailInfo.h
+++ b/Engine/Source/Editor/UnrealEd/Classes/ThumbnailRendering/SceneThumbnailInfo.h
@@ -13,7 +13,9 @@
 UCLASS(MinimalAPI)
 class USceneThumbnailInfo : public UThumbnailInfo
 {
-	GENERATED_UCLASS_BODY()
+	GENERATED_BODY()
+public:
+	UNREALED_API USceneThumbnailInfo(const FObjectInitializer& ObjectInitializer = FObjectInitializer::Get());
 
 	/** The pitch of the orbit camera around the asset */
 	UPROPERTY(EditAnywhere, Category=Thumbnail)
diff --git a/Engine/Source/Editor/UnrealEd/Classes/ThumbnailRendering/SceneThumbnailInfoWithPrimitive.h b/Engine/Source/Editor/UnrealEd/Classes/ThumbnailRendering/SceneThumbnailInfoWithPrimitive.h
--- a/Engine/Source/Editor/UnrealEd/Classes/ThumbnailRendering/SceneThumbnailInfoWithPrimitive.h
+++ b/Engine/Source/Editor/UnrealEd/Classes/ThumbnailRendering/SceneThumbnailInfoWithPrimitive.h
@@ -14,7 +14,9 @@
 UCLASS(MinimalAPI)
 class USceneThumbnailInfoWithPrimitive : public USceneThumbnailInfo
 {
-	GENERATED_UCLASS_BODY()
+	GENERATED_BODY()
+public:
+	UNREALED_API USceneThumbnailInfoWithPrimitive(const FObjectInitializer& ObjectInitializer = FObjectInitializer::Get());
 
 	/** The type of primitive used in this thumbnail */
 	UPROPERTY(EditAnywhere, Category=Thumbnail)
diff --git a/Engine/Source/Editor/UnrealEd/Classes/ThumbnailRendering/SkeletalMeshThumbnailRenderer.h b/Engine/Source/Editor/UnrealEd/Classes/ThumbnailRendering/SkeletalMeshThumbnailRenderer.h
--- a/Engine/Source/Editor/UnrealEd/Classes/ThumbnailRendering/SkeletalMeshThumbnailRenderer.h
+++ b/Engine/Source/Editor/UnrealEd/Classes/ThumbnailRendering/SkeletalMeshThumbnailRenderer.h
@@ -11,7 +11,9 @@
 UCLASS(config=Editor, MinimalAPI)
 class USkeletalMeshThumbnailRenderer : public UDefaultSizedThumbnailRenderer
 {
-	GENERATED_UCLASS_BODY()
+	GENERATED_BODY()
+public:
+	UNREALED_API USkeletalMeshThumbnailRenderer(const FObjectInitializer& ObjectInitializer = FObjectInitializer::Get());
 
 
 	// Begin UThumbnailRenderer Object
diff --git a/Engine/Source/Editor/UnrealEd/Classes/ThumbnailRendering/SlateBrushThumbnailRenderer.h b/Engine/Source/Editor/UnrealEd/Classes/ThumbnailRendering/SlateBrushThumbnailRenderer.h
--- a/Engine/Source/Editor/UnrealEd/Classes/ThumbnailRendering/SlateBrushThumbnailRenderer.h
+++ b/Engine/Source/Editor/UnrealEd/Classes/ThumbnailRendering/SlateBrushThumbnailRenderer.h
@@ -11,7 +11,9 @@
 UCLASS(config=Editor,MinimalAPI)
 class USlateBrushThumbnailRenderer : public UDefaultSizedThumbnailRenderer
 {
-	GENERATED_UCLASS_BODY()
+	GENERATED_BODY()
+public:
+	UNREALED_API USlateBrushThumbnailRenderer(const FObjectInitializer& ObjectInitializer = FObjectInitializer::Get());
 
 public:
 
diff --git a/Engine/Source/Editor/UnrealEd/Classes/ThumbnailRendering/SoundWaveThumbnailRenderer.h b/Engine/Source/Editor/UnrealEd/Classes/ThumbnailRendering/SoundWaveThumbnailRenderer.h
--- a/Engine/Source/Editor/UnrealEd/Classes/ThumbnailRendering/SoundWaveThumbnailRenderer.h
+++ b/Engine/Source/Editor/UnrealEd/Classes/ThumbnailRendering/SoundWaveThumbnailRenderer.h
@@ -11,7 +11,9 @@
 UCLASS()
 class USoundWaveThumbnailRenderer : public UThumbnailRenderer
 {
-	GENERATED_UCLASS_BODY()
+	GENERATED_BODY()
+public:
+	USoundWaveThumbnailRenderer(const FObjectInitializer& ObjectInitializer = FObjectInitializer::Get());
 
 	// Begin UThumbnailRenderer Object
 	virtual void Draw(UObject* Object, int32 X, int32 Y, uint32 Width, uint32 Height, FRenderTarget*, FCanvas* Canvas) override;
diff --git a/Engine/Source/Editor/UnrealEd/Classes/ThumbnailRendering/StaticMeshThumbnailRenderer.h b/Engine/Source/Editor/UnrealEd/Classes/ThumbnailRendering/StaticMeshThumbnailRenderer.h
--- a/Engine/Source/Editor/UnrealEd/Classes/ThumbnailRendering/StaticMeshThumbnailRenderer.h
+++ b/Engine/Source/Editor/UnrealEd/Classes/ThumbnailRendering/StaticMeshThumbnailRenderer.h
@@ -11,7 +11,9 @@
 UCLASS(config=Editor,MinimalAPI)
 class UStaticMeshThumbnailRenderer : public UDefaultSizedThumbnailRenderer
 {
-	GENERATED_UCLASS_BODY()
+	GENERATED_BODY()
+public:
+	UNREALED_API UStaticMeshThumbnailRenderer(const FObjectInitializer& ObjectInitializer = FObjectInitializer::Get());
 
 
 	// Begin UThumbnailRenderer Object
diff --git a/Engine/Source/Editor/UnrealEd/Classes/ThumbnailRendering/SubsurfaceProfileRenderer.h b/Engine/Source/Editor/UnrealEd/Classes/ThumbnailRendering/SubsurfaceProfileRenderer.h
--- a/Engine/Source/Editor/UnrealEd/Classes/ThumbnailRendering/SubsurfaceProfileRenderer.h
+++ b/Engine/Source/Editor/UnrealEd/Classes/ThumbnailRendering/SubsurfaceProfileRenderer.h
@@ -11,7 +11,9 @@
 UCLASS()
 class USubsurfaceProfileRenderer : public UTextureThumbnailRenderer
 {
-	GENERATED_UCLASS_BODY()
+	GENERATED_BODY()
+public:
+	USubsurfaceProfileRenderer(const FObjectInitializer& ObjectInitializer = FObjectInitializer::Get());
 
 
 	// Begin UThumbnailRenderer Object
diff --git a/Engine/Source/Editor/UnrealEd/Classes/ThumbnailRendering/TextureThumbnailRenderer.h b/Engine/Source/Editor/UnrealEd/Classes/ThumbnailRendering/TextureThumbnailRenderer.h
--- a/Engine/Source/Editor/UnrealEd/Classes/ThumbnailRendering/TextureThumbnailRenderer.h
+++ b/Engine/Source/Editor/UnrealEd/Classes/ThumbnailRendering/TextureThumbnailRenderer.h
@@ -11,7 +11,9 @@
 UCLASS(MinimalAPI)
 class UTextureThumbnailRenderer : public UThumbnailRenderer
 {
-	GENERATED_UCLASS_BODY()
+	GENERATED_BODY()
+public:
+	UNREALED_API UTextureThumbnailRenderer(const FObjectInitializer& ObjectInitializer = FObjectInitializer::Get());
 
 
 	// Begin UThumbnailRenderer Object
diff --git a/Engine/Source/Editor/UnrealEd/Classes/ThumbnailRendering/ThumbnailManager.h b/Engine/Source/Editor/UnrealEd/Classes/ThumbnailRendering/ThumbnailManager.h
--- a/Engine/Source/Editor/UnrealEd/Classes/ThumbnailRendering/ThumbnailManager.h
+++ b/Engine/Source/Editor/UnrealEd/Classes/ThumbnailRendering/ThumbnailManager.h
@@ -33,7 +33,7 @@ enum EThumbnailPrimType
 USTRUCT()
 struct FThumbnailRenderingInfo
 {
-	GENERATED_USTRUCT_BODY()
+	GENERATED_BODY()
 
 	/**
 	 * The name of the class that this thumbnail is for (so we can lazy bind)
@@ -73,7 +73,9 @@ public:
 UCLASS(config=Editor)
 class UThumbnailManager : public UObject
 {
-	GENERATED_UCLASS_BODY()
+	GENERATED_BODY()
+public:
+	UThumbnailManager(const FObjectInitializer& ObjectInitializer = FObjectInitializer::Get());
 
 protected:
 	/**
diff --git a/Engine/Source/Editor/UnrealEd/Classes/ThumbnailRendering/ThumbnailRenderer.h b/Engine/Source/Editor/UnrealEd/Classes/ThumbnailRendering/ThumbnailRenderer.h
--- a/Engine/Source/Editor/UnrealEd/Classes/ThumbnailRendering/ThumbnailRenderer.h
+++ b/Engine/Source/Editor/UnrealEd/Classes/ThumbnailRendering/ThumbnailRenderer.h
@@ -13,7 +13,9 @@
 UCLASS(abstract, MinimalAPI)
 class UThumbnailRenderer : public UObject
 {
-	GENERATED_UCLASS_BODY()
+	GENERATED_BODY()
+public:
+	UNREALED_API UThumbnailRenderer(const FObjectInitializer& ObjectInitializer = FObjectInitializer::Get());
 
 
 public:
diff --git a/Engine/Source/Editor/UnrealEd/Classes/ThumbnailRendering/WorldThumbnailInfo.h b/Engine/Source/Editor/UnrealEd/Classes/ThumbnailRendering/WorldThumbnailInfo.h
--- a/Engine/Source/Editor/UnrealEd/Classes/ThumbnailRendering/WorldThumbnailInfo.h
+++ b/Engine/Source/Editor/UnrealEd/Classes/ThumbnailRendering/WorldThumbnailInfo.h
@@ -26,7 +26,9 @@ namespace EOrthoThumbnailDirection
 UCLASS(MinimalAPI)
 class UWorldThumbnailInfo : public USceneThumbnailInfo
 {
-	GENERATED_UCLASS_BODY()
+	GENERATED_BODY()
+public:
+	UNREALED_API UWorldThumbnailInfo(const FObjectInitializer& ObjectInitializer = FObjectInitializer::Get());
 
 	/** The type of projection to use */
 	UPROPERTY(EditAnywhere, Category=Thumbnail)
diff --git a/Engine/Source/Editor/UnrealEd/Classes/ThumbnailRendering/WorldThumbnailRenderer.h b/Engine/Source/Editor/UnrealEd/Classes/ThumbnailRendering/WorldThumbnailRenderer.h
--- a/Engine/Source/Editor/UnrealEd/Classes/ThumbnailRendering/WorldThumbnailRenderer.h
+++ b/Engine/Source/Editor/UnrealEd/Classes/ThumbnailRendering/WorldThumbnailRenderer.h
@@ -11,7 +11,9 @@
 UCLASS(config=Editor)
 class UWorldThumbnailRenderer : public UDefaultSizedThumbnailRenderer
 {
-	GENERATED_UCLASS_BODY()
+	GENERATED_BODY()
+public:
+	UWorldThumbnailRenderer(const FObjectInitializer& ObjectInitializer = FObjectInitializer::Get());
 
 	// Begin UThumbnailRenderer Object
 	virtual bool CanVisualizeAsset(UObject* Object) override;
diff --git a/Engine/Source/Editor/UnrealEd/Classes/UserDefinedStructure/UserDefinedStructEditorData.h b/Engine/Source/Editor/UnrealEd/Classes/UserDefinedStructure/UserDefinedStructEditorData.h
--- a/Engine/Source/Editor/UnrealEd/Classes/UserDefinedStructure/UserDefinedStructEditorData.h
+++ b/Engine/Source/Editor/UnrealEd/Classes/UserDefinedStructure/UserDefinedStructEditorData.h
@@ -7,7 +7,7 @@
 USTRUCT()
 struct FStructVariableDescription
 {
-	GENERATED_USTRUCT_BODY()
+	GENERATED_BODY()
 
 	UPROPERTY()
 	FName VarName;
@@ -65,7 +65,9 @@ struct FStructVariableDescription
 UCLASS()
 class UNREALED_API UUserDefinedStructEditorData : public UObject
 {
-	GENERATED_UCLASS_BODY()
+	GENERATED_BODY()
+public:
+	UUserDefinedStructEditorData(const FObjectInitializer& ObjectInitializer = FObjectInitializer::Get());
 
 private:
 	// the property is used to generate an uniqe name id for member variable
@@ -89,4 +91,4 @@ public:
 
 	uint32 GenerateUniqueNameIdForMemberVariable();
 	class UUserDefinedStruct* GetOwnerStruct() const;
-};
\ No newline at end of file
+};
diff --git a/Engine/Source/Editor/UnrealEd/Public/AutoReimport/AutoReimportManager.h b/Engine/Source/Editor/UnrealEd/Public/AutoReimport/AutoReimportManager.h
--- a/Engine/Source/Editor/UnrealEd/Public/AutoReimport/AutoReimportManager.h
+++ b/Engine/Source/Editor/UnrealEd/Public/AutoReimport/AutoReimportManager.h
@@ -22,7 +22,9 @@ struct FPathAndMountPoint
 UCLASS(config=Editor, transient)
 class UNREALED_API UAutoReimportManager : public UObject
 {
-	GENERATED_UCLASS_BODY()
+	GENERATED_BODY()
+public:
+	UAutoReimportManager(const FObjectInitializer& ObjectInitializer = FObjectInitializer::Get());
 public:
 
 	~UAutoReimportManager();
diff --git a/Engine/Source/Editor/UnrealEd/Public/CookerSettings.h b/Engine/Source/Editor/UnrealEd/Public/CookerSettings.h
--- a/Engine/Source/Editor/UnrealEd/Public/CookerSettings.h
+++ b/Engine/Source/Editor/UnrealEd/Public/CookerSettings.h
@@ -15,7 +15,9 @@ UCLASS(config=Engine, defaultconfig)
 class UNREALED_API UCookerSettings
 	: public UObject
 {
-	GENERATED_UCLASS_BODY()
+	GENERATED_BODY()
+public:
+	UCookerSettings(const FObjectInitializer& ObjectInitializer = FObjectInitializer::Get());
 
 public:
 
diff --git a/Engine/Source/Editor/WorldBrowser/Private/Tiles/WorldTileDetails.h b/Engine/Source/Editor/WorldBrowser/Private/Tiles/WorldTileDetails.h
--- a/Engine/Source/Editor/WorldBrowser/Private/Tiles/WorldTileDetails.h
+++ b/Engine/Source/Editor/WorldBrowser/Private/Tiles/WorldTileDetails.h
@@ -12,7 +12,7 @@ class FWorldTileModel;
 USTRUCT()
 struct FTileLODEntryDetails
 {
-	GENERATED_USTRUCT_BODY()
+	GENERATED_BODY()
 	
 	// Maximum deviation of details percentage
 	UPROPERTY(Category=LODDetails, VisibleAnywhere)
@@ -37,7 +37,9 @@ struct FTileLODEntryDetails
 UCLASS()
 class UWorldTileDetails : public UObject
 {
-	GENERATED_UCLASS_BODY()
+	GENERATED_BODY()
+public:
+	UWorldTileDetails(const FObjectInitializer& ObjectInitializer = FObjectInitializer::Get());
 	
 	// Whether this tile properties can be edited via details panel
 	UPROPERTY(Category=Tile, VisibleAnywhere)
diff --git a/Engine/Source/Programs/UnrealHeaderTool/Private/CodeGenerator.cpp b/Engine/Source/Programs/UnrealHeaderTool/Private/CodeGenerator.cpp
--- a/Engine/Source/Programs/UnrealHeaderTool/Private/CodeGenerator.cpp
+++ b/Engine/Source/Programs/UnrealHeaderTool/Private/CodeGenerator.cpp
@@ -1822,7 +1822,7 @@ void FNativeClassHeaderGenerator::ExportClassesFromSourceFileInner(FUnrealSource
 		{
 			bool bIsIInterface = Class->HasAnyClassFlags(CLASS_Interface);
 
-			FString MacroName = FString::Printf(TEXT("GENERATED_%s_BODY()"), bIsIInterface ? TEXT("IINTERFACE") : TEXT("UCLASS"));
+			auto MacroName = FString::Printf(TEXT("GENERATED_%s_BODY()"), bIsIInterface ? TEXT("IINTERFACE") : TEXT("UCLASS"));
 
 			auto DeprecationWarning = bIsIInterface ? FString(TEXT("")) : GetGeneratedMacroDeprecationWarning(*MacroName);
 
@@ -1831,20 +1831,15 @@ void FNativeClassHeaderGenerator::ExportClassesFromSourceFileInner(FUnrealSource
 
 			auto Public = TEXT("public:" LINE_TERMINATOR);
 
-			FString BodyMacros;
+			auto GeneratedBodyLine = bIsIInterface ? ClassData->GetInterfaceGeneratedBodyLine() : ClassData->GetGeneratedBodyLine();
+			auto LegacyGeneratedBody = InClassMacroCalls + (bIsIInterface ? TEXT("") : StandardUObjectConstructorsMacroCall);
+			auto GeneratedBody = InClassNoPureDeclsMacroCalls + (bIsIInterface ? TEXT("") : EnhancedUObjectConstructorsMacroCall);
 
-			if (bIsIInterface)
-			{
-				BodyMacros = Macroize(*SourceFile.GetGeneratedBodyMacroName(ClassData->GetInterfaceGeneratedBodyLine()), *(DeprecationPushString + FString(Public) + InClassMacroCalls + Public + DeprecationPopString));
-			}
-			else
-			{
-				auto StandardWrappedInClassMacroCalls = Public + InClassMacroCalls + StandardUObjectConstructorsMacroCall + Public;
-				auto EnhancedWrappedInClassMacroCalls = Public + InClassNoPureDeclsMacroCalls + EnhancedUObjectConstructorsMacroCall + GetPreservedAccessSpecifierString(Class);
+			auto WrappedLegacyGeneratedBody = DeprecationWarning + DeprecationPushString + Public + LegacyGeneratedBody + Public + DeprecationPopString;
+			auto WrappedGeneratedBody = FString(DeprecationPushString) + Public + GeneratedBody + GetPreservedAccessSpecifierString(Class) + DeprecationPopString;
 
-				BodyMacros = Macroize(*SourceFile.GetGeneratedBodyMacroName(ClassData->GetGeneratedBodyLine(), true), *(DeprecationWarning + DeprecationPushString + StandardWrappedInClassMacroCalls + DeprecationPopString)) +
-					Macroize(*SourceFile.GetGeneratedBodyMacroName(ClassData->GetGeneratedBodyLine(), false), *(FString(DeprecationPushString) + EnhancedWrappedInClassMacroCalls + DeprecationPopString));
-			}
+			auto BodyMacros = Macroize(*SourceFile.GetGeneratedBodyMacroName(GeneratedBodyLine, true), *WrappedLegacyGeneratedBody) +
+				Macroize(*SourceFile.GetGeneratedBodyMacroName(GeneratedBodyLine, false), *WrappedGeneratedBody);
 
 			GeneratedHeaderText.Log(*BodyMacros);
 		}
@@ -2348,7 +2343,7 @@ void FNativeClassHeaderGenerator::ExportEnums( const TArray<UEnum*>& Enums )
 	}
 }
 
-// Exports the header text for the list of structs specified (GENERATED_USTRUCT_BODY impls)
+// Exports the header text for the list of structs specified (GENERATED_BODY impls)
 void FNativeClassHeaderGenerator::ExportGeneratedStructBodyMacros(FUnrealSourceFile& SourceFile, const TArray<UScriptStruct*>& NativeStructs)
 {
 	// reverse the order.
@@ -2796,7 +2791,7 @@ void FNativeClassHeaderGenerator::ExportDelegateDefinitions(FUnrealSourceFile& S
 		}
 		else
 		{
-			FString MacroName = SourceFile.GetGeneratedMacroName(FunctionData.DelegateMacroLine, TEXT("_DELEGATE"));
+			FString MacroName = SourceFile.GetGeneratedMacroName(FunctionData.MacroLine, TEXT("_DELEGATE"));
 			WriteMacro(HeaderOutput, MacroName, DelegateOutput);
 		}
 	}
@@ -3619,8 +3614,9 @@ void FNativeClassHeaderGenerator::ExportNativeFunctionHeader( const FFuncInfo& F
  * @param FunctionData function data for the current function
  * @param Parameters list of parameters in the function
  * @param Return return parameter for the function
+ * @param DeprecationWarningOutputDevice Device to output deprecation warnings for _Validate and _Implementation functions.
  */
-void FNativeClassHeaderGenerator::ExportFunctionThunk(FStringOutputDevice& RPCWrappers, UFunction* Function, const FFuncInfo& FunctionData, const TArray<UProperty*>& Parameters, UProperty* Return)
+void FNativeClassHeaderGenerator::ExportFunctionThunk(FStringOutputDevice& RPCWrappers, UFunction* Function, const FFuncInfo& FunctionData, const TArray<UProperty*>& Parameters, UProperty* Return, FStringOutputDevice& DeprecationWarningOutputDevice)
 {
 	// export the GET macro for this parameter
 	FString ParameterList;
@@ -3767,7 +3763,7 @@ void FNativeClassHeaderGenerator::ExportFunctionThunk(FStringOutputDevice& RPCWr
 		// NOTE - We don't need to set "Result", since RPC calls don't return any values...
 		if (bShouldEnableValidateDeprecation)
 		{
-			RPCWrappers.Logf(TEXT("EMIT_CUSTOM_WARNING(\"Function %s::%s needs additional native validation function declaration %s due to its properties. Currently %s declaration is autogenerated by UHT. Since next release you'll have to provide declaration on your own. Please update your code before upgrading to the next release, otherwise your project will no longer compile.\")\r\n"), *ClassName, *FunctionName, *FunctionData.CppValidationImplName, *FunctionData.CppValidationImplName);
+			DeprecationWarningOutputDevice.Logf(TEXT("EMIT_CUSTOM_WARNING_AT_LINE(%d, \"Function %s::%s needs additional native validation function declaration %s due to its properties. Currently %s declaration is autogenerated by UHT. Since next release you'll have to provide declaration on your own. Please update your code before upgrading to the next release, otherwise your project will no longer compile.\")\r\n"), FunctionData.MacroLine, *ClassName, *FunctionName, *FunctionData.CppValidationImplName, *FunctionData.CppValidationImplName);
 		}
 		RPCWrappers.Logf(TEXT("%sif (!this->%s(%s))%s"), FCString::Spc(8), *FunctionData.CppValidationImplName, *ParameterList, LINE_TERMINATOR);
 		RPCWrappers.Logf(TEXT("%s{%s"), FCString::Spc(8), LINE_TERMINATOR);
@@ -3778,7 +3774,7 @@ void FNativeClassHeaderGenerator::ExportFunctionThunk(FStringOutputDevice& RPCWr
 
 	if (bShouldEnableImplementationDeprecation)
 	{
-		RPCWrappers.Logf(TEXT("EMIT_CUSTOM_WARNING(\"Function %s::%s needs additional native declaration %s due to its properties. Currently %s declaration is autogenerated by UHT. Since next release you'll have to provide declaration on your own. Please update your code before upgrading to the next release, otherwise your project will no longer compile.\")\r\n"), *ClassName, *FunctionName, *FunctionData.CppImplName, *FunctionData.CppImplName);
+		DeprecationWarningOutputDevice.Logf(TEXT("EMIT_CUSTOM_WARNING_AT_LINE(%d, \"Function %s::%s needs additional native declaration %s due to its properties. Currently %s declaration is autogenerated by UHT. Since next release you'll have to provide declaration on your own. Please update your code before upgrading to the next release, otherwise your project will no longer compile.\")\r\n"), FunctionData.MacroLine, *ClassName, *FunctionName, *FunctionData.CppImplName, *FunctionData.CppImplName);
 	}
 	// write out the return value
 	RPCWrappers.Log(FCString::Spc(8));
@@ -3898,10 +3894,12 @@ void FNativeClassHeaderGenerator::ExportNativeFunctions(FUnrealSourceFile& Sourc
 				FString ParameterList = GetFunctionParameterString(Function);
 
 				auto Virtual = (!FunctionData.FunctionReference->HasAnyFunctionFlags(FUNC_Static) && !(FunctionData.FunctionExportFlags & FUNCEXPORT_Final)) ? TEXT("virtual") : TEXT("");
-				AutogeneratedBlueprintFunctionDeclarations.Logf(TEXT("\t%s bool %s(%s);\r\n"), Virtual, *FunctionData.CppValidationImplName, *ParameterList);
+				FStringOutputDevice ValidDecl;
+				ValidDecl.Logf(TEXT("\t%s bool %s(%s);\r\n"), Virtual, *FunctionData.CppValidationImplName, *ParameterList);
+				AutogeneratedBlueprintFunctionDeclarations.Log(*ValidDecl);
 				if (!bHasValidate)
 				{
-					AutogeneratedBlueprintFunctionDeclarationsOnlyNotDeclared.Logf(TEXT("\t%s bool %s(%s);\r\n"), Virtual, *FunctionData.CppValidationImplName, *ParameterList);
+					AutogeneratedBlueprintFunctionDeclarationsOnlyNotDeclared.Logf(*ValidDecl);
 				}
 			}
 
@@ -3911,7 +3909,7 @@ void FNativeClassHeaderGenerator::ExportNativeFunctions(FUnrealSourceFile& Sourc
 				AutogeneratedBlueprintFunctionDeclarationsOnlyNotDeclared.Log(*FunctionDeclaration);
 			}
 
-			ExportFunctionChecks(FunctionData, MemberFunctionCheckClasses, StaticChecks, NameLookupCPP.GetNameCPP(Class));
+			ExportFunctionChecks(FunctionData, MemberFunctionCheckClasses, StaticChecks, Class->HasAnyClassFlags(CLASS_Interface) ? *(FString(TEXT("I")) + Class->GetName()) : NameLookupCPP.GetNameCPP(Class));
 		}
 
 		if (bMultiLineUFUNCTION)
@@ -3931,7 +3929,7 @@ void FNativeClassHeaderGenerator::ExportNativeFunctions(FUnrealSourceFile& Sourc
 		RPCWrappers.Logf(TEXT("%s%s{%s"), LINE_TERMINATOR, FCString::Spc(4), LINE_TERMINATOR);
 
 		auto Parameters = GetFunctionParmsAndReturn(FunctionData.FunctionReference);
-		ExportFunctionThunk(RPCWrappers, Function, FunctionData, Parameters.Parms, Parameters.Return);
+		ExportFunctionThunk(RPCWrappers, Function, FunctionData, Parameters.Parms, Parameters.Return, AutogeneratedBlueprintFunctionDeclarationsOnlyNotDeclared);
 
 		RPCWrappers.Logf(TEXT("%s}%s"), FCString::Spc(4), LINE_TERMINATOR);
 	}
diff --git a/Engine/Source/Runtime/AIModule/Classes/AIResourceInterface.h b/Engine/Source/Runtime/AIModule/Classes/AIResourceInterface.h
--- a/Engine/Source/Runtime/AIModule/Classes/AIResourceInterface.h
+++ b/Engine/Source/Runtime/AIModule/Classes/AIResourceInterface.h
@@ -7,12 +7,15 @@
 UINTERFACE(MinimalAPI, meta=(CannotImplementInterfaceInBlueprint))
 class UAIResourceInterface : public UInterface
 {
-	GENERATED_UINTERFACE_BODY()
+	GENERATED_BODY()
+public:
+	AIMODULE_API UAIResourceInterface(const FObjectInitializer& ObjectInitializer = FObjectInitializer::Get());
 };
 
 class IAIResourceInterface
 {
-	GENERATED_IINTERFACE_BODY()
+	GENERATED_BODY()
+public:
 
 	/** If resource is lockable lock it with indicated priority */
 	virtual void LockResource(EAIRequestPriority::Type LockSource) {}
diff --git a/Engine/Source/Runtime/AIModule/Classes/AISystem.h b/Engine/Source/Runtime/AIModule/Classes/AISystem.h
--- a/Engine/Source/Runtime/AIModule/Classes/AISystem.h
+++ b/Engine/Source/Runtime/AIModule/Classes/AISystem.h
@@ -18,7 +18,9 @@ class UBlackboardComponent;
 UCLASS(config=Engine, defaultconfig)
 class AIMODULE_API UAISystem : public UAISystemBase
 {
-	GENERATED_UCLASS_BODY()
+	GENERATED_BODY()
+public:
+	UAISystem(const FObjectInitializer& ObjectInitializer = FObjectInitializer::Get());
 
 	UPROPERTY(globalconfig, EditAnywhere, Category = "AISystem", meta = (MetaClass = "AIPerceptionSystem", DisplayName = "Perception System Class"))
 	FStringClassReference PerceptionSystemClassName;
diff --git a/Engine/Source/Runtime/AIModule/Classes/Actions/PawnAction.h b/Engine/Source/Runtime/AIModule/Classes/Actions/PawnAction.h
--- a/Engine/Source/Runtime/AIModule/Classes/Actions/PawnAction.h
+++ b/Engine/Source/Runtime/AIModule/Classes/Actions/PawnAction.h
@@ -43,7 +43,9 @@ namespace EPawnActionFailHandling
 UCLASS(abstract, EditInlineNew)
 class AIMODULE_API UPawnAction : public UObject
 {
-	GENERATED_UCLASS_BODY()
+	GENERATED_BODY()
+public:
+	UPawnAction(const FObjectInitializer& ObjectInitializer = FObjectInitializer::Get());
 
 	friend UPawnActionsComponent;
 	friend FPawnActionStack;
diff --git a/Engine/Source/Runtime/AIModule/Classes/Actions/PawnAction_BlueprintBase.h b/Engine/Source/Runtime/AIModule/Classes/Actions/PawnAction_BlueprintBase.h
--- a/Engine/Source/Runtime/AIModule/Classes/Actions/PawnAction_BlueprintBase.h
+++ b/Engine/Source/Runtime/AIModule/Classes/Actions/PawnAction_BlueprintBase.h
@@ -7,7 +7,9 @@
 UCLASS(abstract, Blueprintable)
 class AIMODULE_API UPawnAction_BlueprintBase : public UPawnAction
 {
-	GENERATED_UCLASS_BODY()
+	GENERATED_BODY()
+public:
+	UPawnAction_BlueprintBase(const FObjectInitializer& ObjectInitializer = FObjectInitializer::Get());
 
 public:
 
diff --git a/Engine/Source/Runtime/AIModule/Classes/Actions/PawnAction_Move.h b/Engine/Source/Runtime/AIModule/Classes/Actions/PawnAction_Move.h
--- a/Engine/Source/Runtime/AIModule/Classes/Actions/PawnAction_Move.h
+++ b/Engine/Source/Runtime/AIModule/Classes/Actions/PawnAction_Move.h
@@ -22,7 +22,9 @@ namespace EPawnActionMoveMode
 UCLASS()
 class AIMODULE_API UPawnAction_Move : public UPawnAction
 {
-	GENERATED_UCLASS_BODY()
+	GENERATED_BODY()
+public:
+	UPawnAction_Move(const FObjectInitializer& ObjectInitializer = FObjectInitializer::Get());
 protected:
 	UPROPERTY(Category = PawnAction, EditAnywhere, BlueprintReadWrite)
 	AActor* GoalActor;
diff --git a/Engine/Source/Runtime/AIModule/Classes/Actions/PawnAction_Repeat.h b/Engine/Source/Runtime/AIModule/Classes/Actions/PawnAction_Repeat.h
--- a/Engine/Source/Runtime/AIModule/Classes/Actions/PawnAction_Repeat.h
+++ b/Engine/Source/Runtime/AIModule/Classes/Actions/PawnAction_Repeat.h
@@ -7,7 +7,9 @@
 UCLASS()
 class AIMODULE_API UPawnAction_Repeat : public UPawnAction
 {
-	GENERATED_UCLASS_BODY()
+	GENERATED_BODY()
+public:
+	UPawnAction_Repeat(const FObjectInitializer& ObjectInitializer = FObjectInitializer::Get());
 
 	enum
 	{
@@ -37,4 +39,4 @@ protected:
 	virtual void OnChildFinished(UPawnAction& Action, EPawnActionResult::Type WithResult) override;
 
 	bool PushSubAction();
-};
\ No newline at end of file
+};
diff --git a/Engine/Source/Runtime/AIModule/Classes/Actions/PawnAction_Sequence.h b/Engine/Source/Runtime/AIModule/Classes/Actions/PawnAction_Sequence.h
--- a/Engine/Source/Runtime/AIModule/Classes/Actions/PawnAction_Sequence.h
+++ b/Engine/Source/Runtime/AIModule/Classes/Actions/PawnAction_Sequence.h
@@ -7,7 +7,9 @@
 UCLASS()
 class AIMODULE_API UPawnAction_Sequence : public UPawnAction
 {
-	GENERATED_UCLASS_BODY()
+	GENERATED_BODY()
+public:
+	UPawnAction_Sequence(const FObjectInitializer& ObjectInitializer = FObjectInitializer::Get());
 
 	UPROPERTY()
 	TArray<UPawnAction*> ActionSequence;
@@ -30,4 +32,4 @@ protected:
 	virtual void OnChildFinished(UPawnAction& Action, EPawnActionResult::Type WithResult) override;
 
 	bool PushNextActionCopy();
-};
\ No newline at end of file
+};
diff --git a/Engine/Source/Runtime/AIModule/Classes/Actions/PawnAction_Wait.h b/Engine/Source/Runtime/AIModule/Classes/Actions/PawnAction_Wait.h
--- a/Engine/Source/Runtime/AIModule/Classes/Actions/PawnAction_Wait.h
+++ b/Engine/Source/Runtime/AIModule/Classes/Actions/PawnAction_Wait.h
@@ -9,7 +9,9 @@
 UCLASS()
 class AIMODULE_API UPawnAction_Wait : public UPawnAction
 {
-	GENERATED_UCLASS_BODY()
+	GENERATED_BODY()
+public:
+	UPawnAction_Wait(const FObjectInitializer& ObjectInitializer = FObjectInitializer::Get());
 		
 	UPROPERTY()
 	float TimeToWait;
@@ -27,4 +29,4 @@ class AIMODULE_API UPawnAction_Wait : public UPawnAction
 	virtual EPawnActionAbortState::Type PerformAbort(EAIForceParam::Type ShouldForce);
 
 	void TimerDone();
-};
\ No newline at end of file
+};
diff --git a/Engine/Source/Runtime/AIModule/Classes/Actions/PawnActionsComponent.h b/Engine/Source/Runtime/AIModule/Classes/Actions/PawnActionsComponent.h
--- a/Engine/Source/Runtime/AIModule/Classes/Actions/PawnActionsComponent.h
+++ b/Engine/Source/Runtime/AIModule/Classes/Actions/PawnActionsComponent.h
@@ -11,7 +11,7 @@ class UPawnAction;
 USTRUCT()
 struct AIMODULE_API FPawnActionEvent
 {
-	GENERATED_USTRUCT_BODY()
+	GENERATED_BODY()
 
 	// used for marking FPawnActionEvent instances created solely for comparisons uses
 	static const int32 FakeActionIndex = INDEX_NONE;
@@ -37,7 +37,7 @@ struct AIMODULE_API FPawnActionEvent
 USTRUCT()
 struct AIMODULE_API FPawnActionStack
 {
-	GENERATED_USTRUCT_BODY()
+	GENERATED_BODY()
 
 private:
 	UPROPERTY()
@@ -68,7 +68,9 @@ public:
 UCLASS()
 class AIMODULE_API UPawnActionsComponent : public UActorComponent
 {
-	GENERATED_UCLASS_BODY()
+	GENERATED_BODY()
+public:
+	UPawnActionsComponent(const FObjectInitializer& ObjectInitializer = FObjectInitializer::Get());
 
 protected:
 	UPROPERTY(BlueprintReadOnly, Category="PawnActions")
@@ -161,4 +163,4 @@ protected:
 private:
 	/** Removed all pending action events associated with PawnAction. Private to make sure it's called only in special cases */
 	void RemoveEventsForAction(UPawnAction& PawnAction);
-};
\ No newline at end of file
+};
diff --git a/Engine/Source/Runtime/AIModule/Classes/BehaviorTree/BTAuxiliaryNode.h b/Engine/Source/Runtime/AIModule/Classes/BehaviorTree/BTAuxiliaryNode.h
--- a/Engine/Source/Runtime/AIModule/Classes/BehaviorTree/BTAuxiliaryNode.h
+++ b/Engine/Source/Runtime/AIModule/Classes/BehaviorTree/BTAuxiliaryNode.h
@@ -30,7 +30,9 @@ struct FBTAuxiliaryMemory : public FBTInstancedNodeMemory
 UCLASS(Abstract)
 class AIMODULE_API UBTAuxiliaryNode : public UBTNode
 {
-	GENERATED_UCLASS_BODY()
+	GENERATED_BODY()
+public:
+	UBTAuxiliaryNode(const FObjectInitializer& ObjectInitializer = FObjectInitializer::Get());
 
 	/** wrapper for node instancing: OnBecomeRelevant */
 	void WrappedOnBecomeRelevant(UBehaviorTreeComponent& OwnerComp, uint8* NodeMemory) const;
diff --git a/Engine/Source/Runtime/AIModule/Classes/BehaviorTree/BTCompositeNode.h b/Engine/Source/Runtime/AIModule/Classes/BehaviorTree/BTCompositeNode.h
--- a/Engine/Source/Runtime/AIModule/Classes/BehaviorTree/BTCompositeNode.h
+++ b/Engine/Source/Runtime/AIModule/Classes/BehaviorTree/BTCompositeNode.h
@@ -49,7 +49,7 @@ namespace EBTDecoratorLogic
 USTRUCT()
 struct FBTDecoratorLogic
 {
-	GENERATED_USTRUCT_BODY()
+	GENERATED_BODY()
 
 	UPROPERTY()
 	TEnumAsByte<EBTDecoratorLogic::Type> Operation;
@@ -64,7 +64,7 @@ struct FBTDecoratorLogic
 USTRUCT()
 struct FBTCompositeChild
 {
-	GENERATED_USTRUCT_BODY()
+	GENERATED_BODY()
 
 	/** child node */
 	UPROPERTY()
@@ -85,7 +85,9 @@ struct FBTCompositeChild
 UCLASS(Abstract)
 class AIMODULE_API UBTCompositeNode : public UBTNode
 {
-	GENERATED_UCLASS_BODY()
+	GENERATED_BODY()
+public:
+	UBTCompositeNode(const FObjectInitializer& ObjectInitializer = FObjectInitializer::Get());
 
 	/** child nodes */
 	UPROPERTY()
diff --git a/Engine/Source/Runtime/AIModule/Classes/BehaviorTree/BTDecorator.h b/Engine/Source/Runtime/AIModule/Classes/BehaviorTree/BTDecorator.h
--- a/Engine/Source/Runtime/AIModule/Classes/BehaviorTree/BTDecorator.h
+++ b/Engine/Source/Runtime/AIModule/Classes/BehaviorTree/BTDecorator.h
@@ -30,7 +30,9 @@ struct FBehaviorTreeSearchData;
 UCLASS(Abstract)
 class AIMODULE_API UBTDecorator : public UBTAuxiliaryNode
 {
-	GENERATED_UCLASS_BODY()
+	GENERATED_BODY()
+public:
+	UBTDecorator(const FObjectInitializer& ObjectInitializer = FObjectInitializer::Get());
 
 	/** fill in data about tree structure */
 	void InitializeDecorator(uint8 InChildIndex);
diff --git a/Engine/Source/Runtime/AIModule/Classes/BehaviorTree/BTNode.h b/Engine/Source/Runtime/AIModule/Classes/BehaviorTree/BTNode.h
--- a/Engine/Source/Runtime/AIModule/Classes/BehaviorTree/BTNode.h
+++ b/Engine/Source/Runtime/AIModule/Classes/BehaviorTree/BTNode.h
@@ -22,7 +22,9 @@ struct FBTInstancedNodeMemory
 UCLASS(Abstract,config=Game)
 class AIMODULE_API UBTNode : public UObject
 {
-	GENERATED_UCLASS_BODY()
+	GENERATED_BODY()
+public:
+	UBTNode(const FObjectInitializer& ObjectInitializer = FObjectInitializer::Get());
 
 	virtual UWorld* GetWorld() const override;
 
diff --git a/Engine/Source/Runtime/AIModule/Classes/BehaviorTree/BTService.h b/Engine/Source/Runtime/AIModule/Classes/BehaviorTree/BTService.h
--- a/Engine/Source/Runtime/AIModule/Classes/BehaviorTree/BTService.h
+++ b/Engine/Source/Runtime/AIModule/Classes/BehaviorTree/BTService.h
@@ -33,7 +33,9 @@
 UCLASS(Abstract)
 class AIMODULE_API UBTService : public UBTAuxiliaryNode
 {
-	GENERATED_UCLASS_BODY()
+	GENERATED_BODY()
+public:
+	UBTService(const FObjectInitializer& ObjectInitializer = FObjectInitializer::Get());
 
 	virtual FString GetStaticDescription() const override;
 
diff --git a/Engine/Source/Runtime/AIModule/Classes/BehaviorTree/BTTaskNode.h b/Engine/Source/Runtime/AIModule/Classes/BehaviorTree/BTTaskNode.h
--- a/Engine/Source/Runtime/AIModule/Classes/BehaviorTree/BTTaskNode.h
+++ b/Engine/Source/Runtime/AIModule/Classes/BehaviorTree/BTTaskNode.h
@@ -26,7 +26,9 @@ struct FAIMessage;
 UCLASS(Abstract)
 class AIMODULE_API UBTTaskNode : public UBTNode
 {
-	GENERATED_UCLASS_BODY()
+	GENERATED_BODY()
+public:
+	UBTTaskNode(const FObjectInitializer& ObjectInitializer = FObjectInitializer::Get());
 
 	/** starts this task, should return Succeeded, Failed or InProgress
 	 *  (use FinishLatentTask() when returning InProgress)
diff --git a/Engine/Source/Runtime/AIModule/Classes/BehaviorTree/BehaviorTree.h b/Engine/Source/Runtime/AIModule/Classes/BehaviorTree/BehaviorTree.h
--- a/Engine/Source/Runtime/AIModule/Classes/BehaviorTree/BehaviorTree.h
+++ b/Engine/Source/Runtime/AIModule/Classes/BehaviorTree/BehaviorTree.h
@@ -13,7 +13,9 @@ class UBTDecorator;
 UCLASS(BlueprintType)
 class AIMODULE_API UBehaviorTree : public UObject
 {
-	GENERATED_UCLASS_BODY()
+	GENERATED_BODY()
+public:
+	UBehaviorTree(const FObjectInitializer& ObjectInitializer = FObjectInitializer::Get());
 
 	/** root node of loaded tree */
 	UPROPERTY()
diff --git a/Engine/Source/Runtime/AIModule/Classes/BehaviorTree/BehaviorTreeComponent.h b/Engine/Source/Runtime/AIModule/Classes/BehaviorTree/BehaviorTreeComponent.h
--- a/Engine/Source/Runtime/AIModule/Classes/BehaviorTree/BehaviorTreeComponent.h
+++ b/Engine/Source/Runtime/AIModule/Classes/BehaviorTree/BehaviorTreeComponent.h
@@ -78,7 +78,9 @@ struct FBTPendingInitializeInfo
 UCLASS()
 class AIMODULE_API UBehaviorTreeComponent : public UBrainComponent
 {
-	GENERATED_UCLASS_BODY()
+	GENERATED_BODY()
+public:
+	UBehaviorTreeComponent(const FObjectInitializer& ObjectInitializer = FObjectInitializer::Get());
 
 	// Begin UBrainComponent overrides
 	virtual void RestartLogic() override;
diff --git a/Engine/Source/Runtime/AIModule/Classes/BehaviorTree/BehaviorTreeManager.h b/Engine/Source/Runtime/AIModule/Classes/BehaviorTree/BehaviorTreeManager.h
--- a/Engine/Source/Runtime/AIModule/Classes/BehaviorTree/BehaviorTreeManager.h
+++ b/Engine/Source/Runtime/AIModule/Classes/BehaviorTree/BehaviorTreeManager.h
@@ -11,7 +11,7 @@ class UBehaviorTree;
 USTRUCT()
 struct FBehaviorTreeTemplateInfo
 {
-	GENERATED_USTRUCT_BODY()
+	GENERATED_BODY()
 
 	/** behavior tree asset */
 	UPROPERTY()
@@ -28,7 +28,9 @@ struct FBehaviorTreeTemplateInfo
 UCLASS(config=Engine)
 class AIMODULE_API UBehaviorTreeManager : public UObject
 {
-	GENERATED_UCLASS_BODY()
+	GENERATED_BODY()
+public:
+	UBehaviorTreeManager(const FObjectInitializer& ObjectInitializer = FObjectInitializer::Get());
 
 	/** limit for recording execution steps for debugger */
 	UPROPERTY(config)
diff --git a/Engine/Source/Runtime/AIModule/Classes/BehaviorTree/BehaviorTreeTypes.h b/Engine/Source/Runtime/AIModule/Classes/BehaviorTree/BehaviorTreeTypes.h
--- a/Engine/Source/Runtime/AIModule/Classes/BehaviorTree/BehaviorTreeTypes.h
+++ b/Engine/Source/Runtime/AIModule/Classes/BehaviorTree/BehaviorTreeTypes.h
@@ -426,7 +426,7 @@ struct FBehaviorTreePropertyMemory
 USTRUCT(BlueprintType)
 struct AIMODULE_API FBlackboardKeySelector
 {
-	GENERATED_USTRUCT_BODY()
+	GENERATED_BODY()
 
 	FBlackboardKeySelector() : SelectedKeyID(FBlackboard::InvalidKey)
 	{}
@@ -510,7 +510,9 @@ public:
 UCLASS(Abstract)
 class AIMODULE_API UBehaviorTreeTypes : public UObject
 {
-	GENERATED_UCLASS_BODY()
+	GENERATED_BODY()
+public:
+	UBehaviorTreeTypes(const FObjectInitializer& ObjectInitializer = FObjectInitializer::Get());
 
 	static FString DescribeNodeHelper(const UBTNode* Node);
 
diff --git a/Engine/Source/Runtime/AIModule/Classes/BehaviorTree/Blackboard/BlackboardKeyType.h b/Engine/Source/Runtime/AIModule/Classes/BehaviorTree/Blackboard/BlackboardKeyType.h
--- a/Engine/Source/Runtime/AIModule/Classes/BehaviorTree/Blackboard/BlackboardKeyType.h
+++ b/Engine/Source/Runtime/AIModule/Classes/BehaviorTree/Blackboard/BlackboardKeyType.h
@@ -77,7 +77,9 @@ struct FBlackboardInstancedKeyMemory
 UCLASS(EditInlineNew, Abstract, CollapseCategories, AutoExpandCategories=(Blackboard))
 class AIMODULE_API UBlackboardKeyType : public UObject
 {
-	GENERATED_UCLASS_BODY()
+	GENERATED_BODY()
+public:
+	UBlackboardKeyType(const FObjectInitializer& ObjectInitializer = FObjectInitializer::Get());
 
 	/** handle dynamic data size */
 	virtual void PreInitialize(UBlackboardComponent& OwnerComp);
diff --git a/Engine/Source/Runtime/AIModule/Classes/BehaviorTree/Blackboard/BlackboardKeyType_Bool.h b/Engine/Source/Runtime/AIModule/Classes/BehaviorTree/Blackboard/BlackboardKeyType_Bool.h
--- a/Engine/Source/Runtime/AIModule/Classes/BehaviorTree/Blackboard/BlackboardKeyType_Bool.h
+++ b/Engine/Source/Runtime/AIModule/Classes/BehaviorTree/Blackboard/BlackboardKeyType_Bool.h
@@ -6,7 +6,9 @@
 UCLASS(EditInlineNew, meta=(DisplayName="Bool"))
 class AIMODULE_API UBlackboardKeyType_Bool : public UBlackboardKeyType
 {
-	GENERATED_UCLASS_BODY()
+	GENERATED_BODY()
+public:
+	UBlackboardKeyType_Bool(const FObjectInitializer& ObjectInitializer = FObjectInitializer::Get());
 
 	typedef bool FDataType;
 	static const FDataType InvalidValue;
diff --git a/Engine/Source/Runtime/AIModule/Classes/BehaviorTree/Blackboard/BlackboardKeyType_Class.h b/Engine/Source/Runtime/AIModule/Classes/BehaviorTree/Blackboard/BlackboardKeyType_Class.h
--- a/Engine/Source/Runtime/AIModule/Classes/BehaviorTree/Blackboard/BlackboardKeyType_Class.h
+++ b/Engine/Source/Runtime/AIModule/Classes/BehaviorTree/Blackboard/BlackboardKeyType_Class.h
@@ -6,7 +6,9 @@
 UCLASS(EditInlineNew, meta=(DisplayName="Class"))
 class AIMODULE_API UBlackboardKeyType_Class : public UBlackboardKeyType
 {
-	GENERATED_UCLASS_BODY()
+	GENERATED_BODY()
+public:
+	UBlackboardKeyType_Class(const FObjectInitializer& ObjectInitializer = FObjectInitializer::Get());
 
 	typedef UClass* FDataType;
 	static const FDataType InvalidValue;
diff --git a/Engine/Source/Runtime/AIModule/Classes/BehaviorTree/Blackboard/BlackboardKeyType_Enum.h b/Engine/Source/Runtime/AIModule/Classes/BehaviorTree/Blackboard/BlackboardKeyType_Enum.h
--- a/Engine/Source/Runtime/AIModule/Classes/BehaviorTree/Blackboard/BlackboardKeyType_Enum.h
+++ b/Engine/Source/Runtime/AIModule/Classes/BehaviorTree/Blackboard/BlackboardKeyType_Enum.h
@@ -6,7 +6,9 @@
 UCLASS(EditInlineNew, meta=(DisplayName="Enum"))
 class AIMODULE_API UBlackboardKeyType_Enum : public UBlackboardKeyType
 {
-	GENERATED_UCLASS_BODY()
+	GENERATED_BODY()
+public:
+	UBlackboardKeyType_Enum(const FObjectInitializer& ObjectInitializer = FObjectInitializer::Get());
 
 	typedef uint8 FDataType;
 	static const FDataType InvalidValue;
diff --git a/Engine/Source/Runtime/AIModule/Classes/BehaviorTree/Blackboard/BlackboardKeyType_Float.h b/Engine/Source/Runtime/AIModule/Classes/BehaviorTree/Blackboard/BlackboardKeyType_Float.h
--- a/Engine/Source/Runtime/AIModule/Classes/BehaviorTree/Blackboard/BlackboardKeyType_Float.h
+++ b/Engine/Source/Runtime/AIModule/Classes/BehaviorTree/Blackboard/BlackboardKeyType_Float.h
@@ -6,7 +6,9 @@
 UCLASS(EditInlineNew, meta=(DisplayName="Float"))
 class AIMODULE_API UBlackboardKeyType_Float : public UBlackboardKeyType
 {
-	GENERATED_UCLASS_BODY()
+	GENERATED_BODY()
+public:
+	UBlackboardKeyType_Float(const FObjectInitializer& ObjectInitializer = FObjectInitializer::Get());
 
 	typedef float FDataType;
 	static const FDataType InvalidValue;
diff --git a/Engine/Source/Runtime/AIModule/Classes/BehaviorTree/Blackboard/BlackboardKeyType_Int.h b/Engine/Source/Runtime/AIModule/Classes/BehaviorTree/Blackboard/BlackboardKeyType_Int.h
--- a/Engine/Source/Runtime/AIModule/Classes/BehaviorTree/Blackboard/BlackboardKeyType_Int.h
+++ b/Engine/Source/Runtime/AIModule/Classes/BehaviorTree/Blackboard/BlackboardKeyType_Int.h
@@ -6,7 +6,9 @@
 UCLASS(EditInlineNew, meta=(DisplayName="Int"))
 class AIMODULE_API UBlackboardKeyType_Int : public UBlackboardKeyType
 {
-	GENERATED_UCLASS_BODY()
+	GENERATED_BODY()
+public:
+	UBlackboardKeyType_Int(const FObjectInitializer& ObjectInitializer = FObjectInitializer::Get());
 
 	typedef int32 FDataType;
 	static const FDataType InvalidValue;
diff --git a/Engine/Source/Runtime/AIModule/Classes/BehaviorTree/Blackboard/BlackboardKeyType_Name.h b/Engine/Source/Runtime/AIModule/Classes/BehaviorTree/Blackboard/BlackboardKeyType_Name.h
--- a/Engine/Source/Runtime/AIModule/Classes/BehaviorTree/Blackboard/BlackboardKeyType_Name.h
+++ b/Engine/Source/Runtime/AIModule/Classes/BehaviorTree/Blackboard/BlackboardKeyType_Name.h
@@ -6,7 +6,9 @@
 UCLASS(EditInlineNew, meta=(DisplayName="Name"))
 class AIMODULE_API UBlackboardKeyType_Name : public UBlackboardKeyType
 {
-	GENERATED_UCLASS_BODY()
+	GENERATED_BODY()
+public:
+	UBlackboardKeyType_Name(const FObjectInitializer& ObjectInitializer = FObjectInitializer::Get());
 
 	typedef FName FDataType;
 	static const FDataType InvalidValue;
diff --git a/Engine/Source/Runtime/AIModule/Classes/BehaviorTree/Blackboard/BlackboardKeyType_NativeEnum.h b/Engine/Source/Runtime/AIModule/Classes/BehaviorTree/Blackboard/BlackboardKeyType_NativeEnum.h
--- a/Engine/Source/Runtime/AIModule/Classes/BehaviorTree/Blackboard/BlackboardKeyType_NativeEnum.h
+++ b/Engine/Source/Runtime/AIModule/Classes/BehaviorTree/Blackboard/BlackboardKeyType_NativeEnum.h
@@ -9,7 +9,9 @@
 UCLASS(NotEditInlineNew)
 class AIMODULE_API UBlackboardKeyType_NativeEnum : public UBlackboardKeyType
 {
-	GENERATED_UCLASS_BODY()
+	GENERATED_BODY()
+public:
+	UBlackboardKeyType_NativeEnum(const FObjectInitializer& ObjectInitializer = FObjectInitializer::Get());
 
 	typedef uint8 FDataType;
 	static const FDataType InvalidValue;
diff --git a/Engine/Source/Runtime/AIModule/Classes/BehaviorTree/Blackboard/BlackboardKeyType_Object.h b/Engine/Source/Runtime/AIModule/Classes/BehaviorTree/Blackboard/BlackboardKeyType_Object.h
--- a/Engine/Source/Runtime/AIModule/Classes/BehaviorTree/Blackboard/BlackboardKeyType_Object.h
+++ b/Engine/Source/Runtime/AIModule/Classes/BehaviorTree/Blackboard/BlackboardKeyType_Object.h
@@ -6,7 +6,9 @@
 UCLASS(EditInlineNew, meta=(DisplayName="Object"))
 class AIMODULE_API UBlackboardKeyType_Object : public UBlackboardKeyType
 {
-	GENERATED_UCLASS_BODY()
+	GENERATED_BODY()
+public:
+	UBlackboardKeyType_Object(const FObjectInitializer& ObjectInitializer = FObjectInitializer::Get());
 
 	typedef UObject* FDataType;
 	static const FDataType InvalidValue;
diff --git a/Engine/Source/Runtime/AIModule/Classes/BehaviorTree/Blackboard/BlackboardKeyType_Rotator.h b/Engine/Source/Runtime/AIModule/Classes/BehaviorTree/Blackboard/BlackboardKeyType_Rotator.h
--- a/Engine/Source/Runtime/AIModule/Classes/BehaviorTree/Blackboard/BlackboardKeyType_Rotator.h
+++ b/Engine/Source/Runtime/AIModule/Classes/BehaviorTree/Blackboard/BlackboardKeyType_Rotator.h
@@ -5,7 +5,9 @@
 UCLASS(EditInlineNew, meta=(DisplayName="Rotator"))
 class AIMODULE_API UBlackboardKeyType_Rotator : public UBlackboardKeyType
 {
-	GENERATED_UCLASS_BODY()
+	GENERATED_BODY()
+public:
+	UBlackboardKeyType_Rotator(const FObjectInitializer& ObjectInitializer = FObjectInitializer::Get());
 	
 	typedef FRotator FDataType; 
 	static const FDataType InvalidValue;
diff --git a/Engine/Source/Runtime/AIModule/Classes/BehaviorTree/Blackboard/BlackboardKeyType_String.h b/Engine/Source/Runtime/AIModule/Classes/BehaviorTree/Blackboard/BlackboardKeyType_String.h
--- a/Engine/Source/Runtime/AIModule/Classes/BehaviorTree/Blackboard/BlackboardKeyType_String.h
+++ b/Engine/Source/Runtime/AIModule/Classes/BehaviorTree/Blackboard/BlackboardKeyType_String.h
@@ -6,7 +6,9 @@
 UCLASS(EditInlineNew, meta=(DisplayName="String"))
 class AIMODULE_API UBlackboardKeyType_String : public UBlackboardKeyType
 {
-	GENERATED_UCLASS_BODY()
+	GENERATED_BODY()
+public:
+	UBlackboardKeyType_String(const FObjectInitializer& ObjectInitializer = FObjectInitializer::Get());
 
 	typedef FString FDataType;
 	static const FDataType InvalidValue;
diff --git a/Engine/Source/Runtime/AIModule/Classes/BehaviorTree/Blackboard/BlackboardKeyType_Vector.h b/Engine/Source/Runtime/AIModule/Classes/BehaviorTree/Blackboard/BlackboardKeyType_Vector.h
--- a/Engine/Source/Runtime/AIModule/Classes/BehaviorTree/Blackboard/BlackboardKeyType_Vector.h
+++ b/Engine/Source/Runtime/AIModule/Classes/BehaviorTree/Blackboard/BlackboardKeyType_Vector.h
@@ -6,7 +6,9 @@
 UCLASS(EditInlineNew, meta=(DisplayName="Vector"))
 class AIMODULE_API UBlackboardKeyType_Vector : public UBlackboardKeyType
 {
-	GENERATED_UCLASS_BODY()
+	GENERATED_BODY()
+public:
+	UBlackboardKeyType_Vector(const FObjectInitializer& ObjectInitializer = FObjectInitializer::Get());
 
 	typedef FVector FDataType; 
 	static const FDataType InvalidValue;
diff --git a/Engine/Source/Runtime/AIModule/Classes/BehaviorTree/BlackboardData.h b/Engine/Source/Runtime/AIModule/Classes/BehaviorTree/BlackboardData.h
--- a/Engine/Source/Runtime/AIModule/Classes/BehaviorTree/BlackboardData.h
+++ b/Engine/Source/Runtime/AIModule/Classes/BehaviorTree/BlackboardData.h
@@ -11,7 +11,7 @@ class UBlackboardData;
 USTRUCT()
 struct FBlackboardEntry
 {
-	GENERATED_USTRUCT_BODY()
+	GENERATED_BODY()
 
 	UPROPERTY(EditAnywhere, Category=Blackboard)
 	FName EntryName;
@@ -35,7 +35,9 @@ struct FBlackboardEntry
 UCLASS(AutoExpandCategories=(Blackboard))
 class AIMODULE_API UBlackboardData : public UDataAsset
 {
-	GENERATED_UCLASS_BODY()
+	GENERATED_BODY()
+public:
+	UBlackboardData(const FObjectInitializer& ObjectInitializer = FObjectInitializer::Get());
 	DECLARE_MULTICAST_DELEGATE_OneParam(FKeyUpdate, UBlackboardData* /*asset*/);
 
 	/** parent blackboard (keys can be overridden) */
diff --git a/Engine/Source/Runtime/AIModule/Classes/BehaviorTree/Composites/BTComposite_Selector.h b/Engine/Source/Runtime/AIModule/Classes/BehaviorTree/Composites/BTComposite_Selector.h
--- a/Engine/Source/Runtime/AIModule/Classes/BehaviorTree/Composites/BTComposite_Selector.h
+++ b/Engine/Source/Runtime/AIModule/Classes/BehaviorTree/Composites/BTComposite_Selector.h
@@ -11,7 +11,9 @@
 UCLASS()
 class AIMODULE_API UBTComposite_Selector: public UBTCompositeNode
 {
-	GENERATED_UCLASS_BODY()
+	GENERATED_BODY()
+public:
+	UBTComposite_Selector(const FObjectInitializer& ObjectInitializer = FObjectInitializer::Get());
 
 	int32 GetNextChildHandler(struct FBehaviorTreeSearchData& SearchData, int32 PrevChild, EBTNodeResult::Type LastResult) const;
 
diff --git a/Engine/Source/Runtime/AIModule/Classes/BehaviorTree/Composites/BTComposite_Sequence.h b/Engine/Source/Runtime/AIModule/Classes/BehaviorTree/Composites/BTComposite_Sequence.h
--- a/Engine/Source/Runtime/AIModule/Classes/BehaviorTree/Composites/BTComposite_Sequence.h
+++ b/Engine/Source/Runtime/AIModule/Classes/BehaviorTree/Composites/BTComposite_Sequence.h
@@ -11,7 +11,9 @@
 UCLASS()
 class AIMODULE_API UBTComposite_Sequence : public UBTCompositeNode
 {
-	GENERATED_UCLASS_BODY()
+	GENERATED_BODY()
+public:
+	UBTComposite_Sequence(const FObjectInitializer& ObjectInitializer = FObjectInitializer::Get());
 
 	int32 GetNextChildHandler(struct FBehaviorTreeSearchData& SearchData, int32 PrevChild, EBTNodeResult::Type LastResult) const;
 
diff --git a/Engine/Source/Runtime/AIModule/Classes/BehaviorTree/Composites/BTComposite_SimpleParallel.h b/Engine/Source/Runtime/AIModule/Classes/BehaviorTree/Composites/BTComposite_SimpleParallel.h
--- a/Engine/Source/Runtime/AIModule/Classes/BehaviorTree/Composites/BTComposite_SimpleParallel.h
+++ b/Engine/Source/Runtime/AIModule/Classes/BehaviorTree/Composites/BTComposite_SimpleParallel.h
@@ -52,7 +52,9 @@ struct FBTParallelMemory : public FBTCompositeMemory
 UCLASS()
 class AIMODULE_API UBTComposite_SimpleParallel : public UBTCompositeNode
 {
-	GENERATED_UCLASS_BODY()
+	GENERATED_BODY()
+public:
+	UBTComposite_SimpleParallel(const FObjectInitializer& ObjectInitializer = FObjectInitializer::Get());
 
 	/** how background tree should be handled when main task finishes execution */
 	UPROPERTY(EditInstanceOnly, Category=Parallel)
diff --git a/Engine/Source/Runtime/AIModule/Classes/BehaviorTree/Decorators/BTDecorator_Blackboard.h b/Engine/Source/Runtime/AIModule/Classes/BehaviorTree/Decorators/BTDecorator_Blackboard.h
--- a/Engine/Source/Runtime/AIModule/Classes/BehaviorTree/Decorators/BTDecorator_Blackboard.h
+++ b/Engine/Source/Runtime/AIModule/Classes/BehaviorTree/Decorators/BTDecorator_Blackboard.h
@@ -31,7 +31,9 @@ namespace EBTBlackboardRestart
 UCLASS(HideCategories=(Condition))
 class AIMODULE_API UBTDecorator_Blackboard : public UBTDecorator_BlackboardBase
 {
-	GENERATED_UCLASS_BODY()
+	GENERATED_BODY()
+public:
+	UBTDecorator_Blackboard(const FObjectInitializer& ObjectInitializer = FObjectInitializer::Get());
 
 	virtual bool CalculateRawConditionValue(UBehaviorTreeComponent& OwnerComp, uint8* NodeMemory) const override;
 	virtual void OnBlackboardChange(const UBlackboardComponent& Blackboard, FBlackboard::FKey ChangedKeyID) override;
diff --git a/Engine/Source/Runtime/AIModule/Classes/BehaviorTree/Decorators/BTDecorator_BlackboardBase.h b/Engine/Source/Runtime/AIModule/Classes/BehaviorTree/Decorators/BTDecorator_BlackboardBase.h
--- a/Engine/Source/Runtime/AIModule/Classes/BehaviorTree/Decorators/BTDecorator_BlackboardBase.h
+++ b/Engine/Source/Runtime/AIModule/Classes/BehaviorTree/Decorators/BTDecorator_BlackboardBase.h
@@ -8,7 +8,9 @@ class UBlackboardComponent;
 UCLASS(Abstract)
 class AIMODULE_API UBTDecorator_BlackboardBase : public UBTDecorator
 {
-	GENERATED_UCLASS_BODY()
+	GENERATED_BODY()
+public:
+	UBTDecorator_BlackboardBase(const FObjectInitializer& ObjectInitializer = FObjectInitializer::Get());
 
 	/** initialize any asset related data */
 	virtual void InitializeFromAsset(UBehaviorTree& Asset) override;
diff --git a/Engine/Source/Runtime/AIModule/Classes/BehaviorTree/Decorators/BTDecorator_CheckGameplayTagsOnActor.h b/Engine/Source/Runtime/AIModule/Classes/BehaviorTree/Decorators/BTDecorator_CheckGameplayTagsOnActor.h
--- a/Engine/Source/Runtime/AIModule/Classes/BehaviorTree/Decorators/BTDecorator_CheckGameplayTagsOnActor.h
+++ b/Engine/Source/Runtime/AIModule/Classes/BehaviorTree/Decorators/BTDecorator_CheckGameplayTagsOnActor.h
@@ -14,7 +14,9 @@
 UCLASS()
 class AIMODULE_API UBTDecorator_CheckGameplayTagsOnActor : public UBTDecorator
 {
-	GENERATED_UCLASS_BODY()
+	GENERATED_BODY()
+public:
+	UBTDecorator_CheckGameplayTagsOnActor(const FObjectInitializer& ObjectInitializer = FObjectInitializer::Get());
 
 	virtual bool CalculateRawConditionValue(UBehaviorTreeComponent& OwnerComp, uint8* NodeMemory) const override;
 	virtual void DescribeRuntimeValues(const UBehaviorTreeComponent& OwnerComp, uint8* NodeMemory, EBTDescriptionVerbosity::Type Verbosity, TArray<FString>& Values) const override;
diff --git a/Engine/Source/Runtime/AIModule/Classes/BehaviorTree/Decorators/BTDecorator_CompareBBEntries.h b/Engine/Source/Runtime/AIModule/Classes/BehaviorTree/Decorators/BTDecorator_CompareBBEntries.h
--- a/Engine/Source/Runtime/AIModule/Classes/BehaviorTree/Decorators/BTDecorator_CompareBBEntries.h
+++ b/Engine/Source/Runtime/AIModule/Classes/BehaviorTree/Decorators/BTDecorator_CompareBBEntries.h
@@ -24,7 +24,9 @@ namespace EBlackBoardEntryComparison
 UCLASS(HideCategories=(Condition))
 class AIMODULE_API UBTDecorator_CompareBBEntries : public UBTDecorator
 {
-	GENERATED_UCLASS_BODY()
+	GENERATED_BODY()
+public:
+	UBTDecorator_CompareBBEntries(const FObjectInitializer& ObjectInitializer = FObjectInitializer::Get());
 
 protected:
 
diff --git a/Engine/Source/Runtime/AIModule/Classes/BehaviorTree/Decorators/BTDecorator_ConditionalLoop.h b/Engine/Source/Runtime/AIModule/Classes/BehaviorTree/Decorators/BTDecorator_ConditionalLoop.h
--- a/Engine/Source/Runtime/AIModule/Classes/BehaviorTree/Decorators/BTDecorator_ConditionalLoop.h
+++ b/Engine/Source/Runtime/AIModule/Classes/BehaviorTree/Decorators/BTDecorator_ConditionalLoop.h
@@ -11,7 +11,9 @@
 UCLASS(HideCategories=(FlowControl))
 class AIMODULE_API UBTDecorator_ConditionalLoop : public UBTDecorator_Blackboard
 {
-	GENERATED_UCLASS_BODY()
+	GENERATED_BODY()
+public:
+	UBTDecorator_ConditionalLoop(const FObjectInitializer& ObjectInitializer = FObjectInitializer::Get());
 
 	virtual bool CalculateRawConditionValue(UBehaviorTreeComponent& OwnerComp, uint8* NodeMemory) const override;
 	virtual void OnBlackboardChange(const UBlackboardComponent& Blackboard, FBlackboard::FKey ChangedKeyID) override;
diff --git a/Engine/Source/Runtime/AIModule/Classes/BehaviorTree/Decorators/BTDecorator_ConeCheck.h b/Engine/Source/Runtime/AIModule/Classes/BehaviorTree/Decorators/BTDecorator_ConeCheck.h
--- a/Engine/Source/Runtime/AIModule/Classes/BehaviorTree/Decorators/BTDecorator_ConeCheck.h
+++ b/Engine/Source/Runtime/AIModule/Classes/BehaviorTree/Decorators/BTDecorator_ConeCheck.h
@@ -19,7 +19,9 @@ struct FBTConeCheckDecoratorMemory
 UCLASS()
 class AIMODULE_API UBTDecorator_ConeCheck : public UBTDecorator
 {
-	GENERATED_UCLASS_BODY()
+	GENERATED_BODY()
+public:
+	UBTDecorator_ConeCheck(const FObjectInitializer& ObjectInitializer = FObjectInitializer::Get());
 
 	typedef FBTConeCheckDecoratorMemory TNodeInstanceMemory;
 
diff --git a/Engine/Source/Runtime/AIModule/Classes/BehaviorTree/Decorators/BTDecorator_Cooldown.h b/Engine/Source/Runtime/AIModule/Classes/BehaviorTree/Decorators/BTDecorator_Cooldown.h
--- a/Engine/Source/Runtime/AIModule/Classes/BehaviorTree/Decorators/BTDecorator_Cooldown.h
+++ b/Engine/Source/Runtime/AIModule/Classes/BehaviorTree/Decorators/BTDecorator_Cooldown.h
@@ -17,7 +17,9 @@ struct FBTCooldownDecoratorMemory
 UCLASS(HideCategories=(Condition))
 class AIMODULE_API UBTDecorator_Cooldown : public UBTDecorator
 {
-	GENERATED_UCLASS_BODY()
+	GENERATED_BODY()
+public:
+	UBTDecorator_Cooldown(const FObjectInitializer& ObjectInitializer = FObjectInitializer::Get());
 
 	/** max allowed time for execution of underlying node */
 	UPROPERTY(Category=Decorator, EditAnywhere)
diff --git a/Engine/Source/Runtime/AIModule/Classes/BehaviorTree/Decorators/BTDecorator_DoesPathExist.h b/Engine/Source/Runtime/AIModule/Classes/BehaviorTree/Decorators/BTDecorator_DoesPathExist.h
--- a/Engine/Source/Runtime/AIModule/Classes/BehaviorTree/Decorators/BTDecorator_DoesPathExist.h
+++ b/Engine/Source/Runtime/AIModule/Classes/BehaviorTree/Decorators/BTDecorator_DoesPathExist.h
@@ -24,7 +24,9 @@ namespace EPathExistanceQueryType
 UCLASS()
 class AIMODULE_API UBTDecorator_DoesPathExist : public UBTDecorator
 {
-	GENERATED_UCLASS_BODY()
+	GENERATED_BODY()
+public:
+	UBTDecorator_DoesPathExist(const FObjectInitializer& ObjectInitializer = FObjectInitializer::Get());
 
 protected:
 
diff --git a/Engine/Source/Runtime/AIModule/Classes/BehaviorTree/Decorators/BTDecorator_ForceSuccess.h b/Engine/Source/Runtime/AIModule/Classes/BehaviorTree/Decorators/BTDecorator_ForceSuccess.h
--- a/Engine/Source/Runtime/AIModule/Classes/BehaviorTree/Decorators/BTDecorator_ForceSuccess.h
+++ b/Engine/Source/Runtime/AIModule/Classes/BehaviorTree/Decorators/BTDecorator_ForceSuccess.h
@@ -16,7 +16,9 @@
 UCLASS(HideCategories=(Condition))
 class AIMODULE_API UBTDecorator_ForceSuccess : public UBTDecorator
 {
-	GENERATED_UCLASS_BODY()
+	GENERATED_BODY()
+public:
+	UBTDecorator_ForceSuccess(const FObjectInitializer& ObjectInitializer = FObjectInitializer::Get());
 
 protected:
 
diff --git a/Engine/Source/Runtime/AIModule/Classes/BehaviorTree/Decorators/BTDecorator_KeepInCone.h b/Engine/Source/Runtime/AIModule/Classes/BehaviorTree/Decorators/BTDecorator_KeepInCone.h
--- a/Engine/Source/Runtime/AIModule/Classes/BehaviorTree/Decorators/BTDecorator_KeepInCone.h
+++ b/Engine/Source/Runtime/AIModule/Classes/BehaviorTree/Decorators/BTDecorator_KeepInCone.h
@@ -16,7 +16,9 @@ struct FBTKeepInConeDecoratorMemory
 UCLASS(HideCategories=(Condition))
 class AIMODULE_API UBTDecorator_KeepInCone : public UBTDecorator
 {
-	GENERATED_UCLASS_BODY()
+	GENERATED_BODY()
+public:
+	UBTDecorator_KeepInCone(const FObjectInitializer& ObjectInitializer = FObjectInitializer::Get());
 
 	typedef FBTKeepInConeDecoratorMemory TNodeInstanceMemory;
 
diff --git a/Engine/Source/Runtime/AIModule/Classes/BehaviorTree/Decorators/BTDecorator_Loop.h b/Engine/Source/Runtime/AIModule/Classes/BehaviorTree/Decorators/BTDecorator_Loop.h
--- a/Engine/Source/Runtime/AIModule/Classes/BehaviorTree/Decorators/BTDecorator_Loop.h
+++ b/Engine/Source/Runtime/AIModule/Classes/BehaviorTree/Decorators/BTDecorator_Loop.h
@@ -16,7 +16,9 @@ struct FBTLoopDecoratorMemory
 UCLASS(HideCategories=(Condition))
 class AIMODULE_API UBTDecorator_Loop : public UBTDecorator
 {
-	GENERATED_UCLASS_BODY()
+	GENERATED_BODY()
+public:
+	UBTDecorator_Loop(const FObjectInitializer& ObjectInitializer = FObjectInitializer::Get());
 
 	/** number of executions */
 	UPROPERTY(Category=Decorator, EditAnywhere)
diff --git a/Engine/Source/Runtime/AIModule/Classes/BehaviorTree/Decorators/BTDecorator_ReachedMoveGoal.h b/Engine/Source/Runtime/AIModule/Classes/BehaviorTree/Decorators/BTDecorator_ReachedMoveGoal.h
--- a/Engine/Source/Runtime/AIModule/Classes/BehaviorTree/Decorators/BTDecorator_ReachedMoveGoal.h
+++ b/Engine/Source/Runtime/AIModule/Classes/BehaviorTree/Decorators/BTDecorator_ReachedMoveGoal.h
@@ -11,7 +11,9 @@
 UCLASS()
 class AIMODULE_API UBTDecorator_ReachedMoveGoal : public UBTDecorator
 {
-	GENERATED_UCLASS_BODY()
+	GENERATED_BODY()
+public:
+	UBTDecorator_ReachedMoveGoal(const FObjectInitializer& ObjectInitializer = FObjectInitializer::Get());
 
 	virtual bool CalculateRawConditionValue(UBehaviorTreeComponent& OwnerComp, uint8* NodeMemory) const override;
 
diff --git a/Engine/Source/Runtime/AIModule/Classes/BehaviorTree/Decorators/BTDecorator_TagCooldown.h b/Engine/Source/Runtime/AIModule/Classes/BehaviorTree/Decorators/BTDecorator_TagCooldown.h
--- a/Engine/Source/Runtime/AIModule/Classes/BehaviorTree/Decorators/BTDecorator_TagCooldown.h
+++ b/Engine/Source/Runtime/AIModule/Classes/BehaviorTree/Decorators/BTDecorator_TagCooldown.h
@@ -17,7 +17,9 @@ struct FBTTagCooldownDecoratorMemory
 UCLASS(HideCategories=(Condition))
 class AIMODULE_API UBTDecorator_TagCooldown : public UBTDecorator
 {
-	GENERATED_UCLASS_BODY()
+	GENERATED_BODY()
+public:
+	UBTDecorator_TagCooldown(const FObjectInitializer& ObjectInitializer = FObjectInitializer::Get());
 
 	/** gameplay tag that will be used for the cooldown */
 	UPROPERTY(Category = Decorator, EditAnywhere)
diff --git a/Engine/Source/Runtime/AIModule/Classes/BehaviorTree/Decorators/BTDecorator_TimeLimit.h b/Engine/Source/Runtime/AIModule/Classes/BehaviorTree/Decorators/BTDecorator_TimeLimit.h
--- a/Engine/Source/Runtime/AIModule/Classes/BehaviorTree/Decorators/BTDecorator_TimeLimit.h
+++ b/Engine/Source/Runtime/AIModule/Classes/BehaviorTree/Decorators/BTDecorator_TimeLimit.h
@@ -10,7 +10,9 @@
 UCLASS(HideCategories=(Condition))
 class AIMODULE_API UBTDecorator_TimeLimit : public UBTDecorator
 {
-	GENERATED_UCLASS_BODY()
+	GENERATED_BODY()
+public:
+	UBTDecorator_TimeLimit(const FObjectInitializer& ObjectInitializer = FObjectInitializer::Get());
 		
 	/** max allowed time for execution of underlying node */
 	UPROPERTY(Category=Decorator, EditAnywhere)
diff --git a/Engine/Source/Runtime/AIModule/Classes/BehaviorTree/Services/BTService_BlackboardBase.h b/Engine/Source/Runtime/AIModule/Classes/BehaviorTree/Services/BTService_BlackboardBase.h
--- a/Engine/Source/Runtime/AIModule/Classes/BehaviorTree/Services/BTService_BlackboardBase.h
+++ b/Engine/Source/Runtime/AIModule/Classes/BehaviorTree/Services/BTService_BlackboardBase.h
@@ -6,7 +6,9 @@
 UCLASS(Abstract)
 class AIMODULE_API UBTService_BlackboardBase : public UBTService
 {
-	GENERATED_UCLASS_BODY()
+	GENERATED_BODY()
+public:
+	UBTService_BlackboardBase(const FObjectInitializer& ObjectInitializer = FObjectInitializer::Get());
 
 	/** initialize any asset related data */
 	virtual void InitializeFromAsset(UBehaviorTree& Asset) override;
diff --git a/Engine/Source/Runtime/AIModule/Classes/BehaviorTree/Services/BTService_BlueprintBase.h b/Engine/Source/Runtime/AIModule/Classes/BehaviorTree/Services/BTService_BlueprintBase.h
--- a/Engine/Source/Runtime/AIModule/Classes/BehaviorTree/Services/BTService_BlueprintBase.h
+++ b/Engine/Source/Runtime/AIModule/Classes/BehaviorTree/Services/BTService_BlueprintBase.h
@@ -17,7 +17,9 @@ class FBehaviorBlueprintDetails;
 UCLASS(Abstract, Blueprintable)
 class AIMODULE_API UBTService_BlueprintBase : public UBTService
 {
-	GENERATED_UCLASS_BODY()
+	GENERATED_BODY()
+public:
+	UBTService_BlueprintBase(const FObjectInitializer& ObjectInitializer = FObjectInitializer::Get());
 
 	virtual void PostInitProperties() override;
 
diff --git a/Engine/Source/Runtime/AIModule/Classes/BehaviorTree/Tasks/BTTask_BlackboardBase.h b/Engine/Source/Runtime/AIModule/Classes/BehaviorTree/Tasks/BTTask_BlackboardBase.h
--- a/Engine/Source/Runtime/AIModule/Classes/BehaviorTree/Tasks/BTTask_BlackboardBase.h
+++ b/Engine/Source/Runtime/AIModule/Classes/BehaviorTree/Tasks/BTTask_BlackboardBase.h
@@ -10,7 +10,9 @@ class UBehaviorTree;
 UCLASS(Abstract)
 class AIMODULE_API UBTTask_BlackboardBase : public UBTTaskNode
 {
-	GENERATED_UCLASS_BODY()
+	GENERATED_BODY()
+public:
+	UBTTask_BlackboardBase(const FObjectInitializer& ObjectInitializer = FObjectInitializer::Get());
 
 	/** initialize any asset related data */
 	virtual void InitializeFromAsset(UBehaviorTree& Asset) override;
diff --git a/Engine/Source/Runtime/AIModule/Classes/BehaviorTree/Tasks/BTTask_BlueprintBase.h b/Engine/Source/Runtime/AIModule/Classes/BehaviorTree/Tasks/BTTask_BlueprintBase.h
--- a/Engine/Source/Runtime/AIModule/Classes/BehaviorTree/Tasks/BTTask_BlueprintBase.h
+++ b/Engine/Source/Runtime/AIModule/Classes/BehaviorTree/Tasks/BTTask_BlueprintBase.h
@@ -18,7 +18,9 @@ class FBehaviorBlueprintDetails;
 UCLASS(Abstract, Blueprintable)
 class AIMODULE_API UBTTask_BlueprintBase : public UBTTaskNode
 {
-	GENERATED_UCLASS_BODY()
+	GENERATED_BODY()
+public:
+	UBTTask_BlueprintBase(const FObjectInitializer& ObjectInitializer = FObjectInitializer::Get());
 
 	/** setup node name */
 	virtual void PostInitProperties() override;
diff --git a/Engine/Source/Runtime/AIModule/Classes/BehaviorTree/Tasks/BTTask_MakeNoise.h b/Engine/Source/Runtime/AIModule/Classes/BehaviorTree/Tasks/BTTask_MakeNoise.h
--- a/Engine/Source/Runtime/AIModule/Classes/BehaviorTree/Tasks/BTTask_MakeNoise.h
+++ b/Engine/Source/Runtime/AIModule/Classes/BehaviorTree/Tasks/BTTask_MakeNoise.h
@@ -10,7 +10,9 @@
 UCLASS()
 class AIMODULE_API UBTTask_MakeNoise : public UBTTaskNode
 {
-	GENERATED_UCLASS_BODY()
+	GENERATED_BODY()
+public:
+	UBTTask_MakeNoise(const FObjectInitializer& ObjectInitializer = FObjectInitializer::Get());
 
 	/** Loudnes of generated noise */
 	UPROPERTY(Category=Node, EditAnywhere, meta=(ClampMin = "0", ClampMax = "1", UIMin = "0", UIMax = "1"))
diff --git a/Engine/Source/Runtime/AIModule/Classes/BehaviorTree/Tasks/BTTask_MoveDirectlyToward.h b/Engine/Source/Runtime/AIModule/Classes/BehaviorTree/Tasks/BTTask_MoveDirectlyToward.h
--- a/Engine/Source/Runtime/AIModule/Classes/BehaviorTree/Tasks/BTTask_MoveDirectlyToward.h
+++ b/Engine/Source/Runtime/AIModule/Classes/BehaviorTree/Tasks/BTTask_MoveDirectlyToward.h
@@ -17,7 +17,9 @@ struct FBTMoveDirectlyTowardMemory
 UCLASS(config=Game)
 class AIMODULE_API UBTTask_MoveDirectlyToward : public UBTTask_BlackboardBase
 {
-	GENERATED_UCLASS_BODY()
+	GENERATED_BODY()
+public:
+	UBTTask_MoveDirectlyToward(const FObjectInitializer& ObjectInitializer = FObjectInitializer::Get());
 
 	UPROPERTY(config, Category=Node, EditAnywhere, meta=(ClampMin = "0.0"))
 	float AcceptableRadius;
diff --git a/Engine/Source/Runtime/AIModule/Classes/BehaviorTree/Tasks/BTTask_MoveTo.h b/Engine/Source/Runtime/AIModule/Classes/BehaviorTree/Tasks/BTTask_MoveTo.h
--- a/Engine/Source/Runtime/AIModule/Classes/BehaviorTree/Tasks/BTTask_MoveTo.h
+++ b/Engine/Source/Runtime/AIModule/Classes/BehaviorTree/Tasks/BTTask_MoveTo.h
@@ -21,7 +21,9 @@ struct FBTMoveToTaskMemory
 UCLASS(config=Game)
 class AIMODULE_API UBTTask_MoveTo : public UBTTask_BlackboardBase
 {
-	GENERATED_UCLASS_BODY()
+	GENERATED_BODY()
+public:
+	UBTTask_MoveTo(const FObjectInitializer& ObjectInitializer = FObjectInitializer::Get());
 
 	UPROPERTY(config, Category=Node, EditAnywhere, meta=(ClampMin = "0.0"))
 	float AcceptableRadius;
diff --git a/Engine/Source/Runtime/AIModule/Classes/BehaviorTree/Tasks/BTTask_PawnActionBase.h b/Engine/Source/Runtime/AIModule/Classes/BehaviorTree/Tasks/BTTask_PawnActionBase.h
--- a/Engine/Source/Runtime/AIModule/Classes/BehaviorTree/Tasks/BTTask_PawnActionBase.h
+++ b/Engine/Source/Runtime/AIModule/Classes/BehaviorTree/Tasks/BTTask_PawnActionBase.h
@@ -20,7 +20,9 @@ class UPawnAction;
 UCLASS(Abstract)
 class AIMODULE_API UBTTask_PawnActionBase : public UBTTaskNode
 {
-	GENERATED_UCLASS_BODY()
+	GENERATED_BODY()
+public:
+	UBTTask_PawnActionBase(const FObjectInitializer& ObjectInitializer = FObjectInitializer::Get());
 
 	virtual EBTNodeResult::Type AbortTask(UBehaviorTreeComponent& OwnerComp, uint8* NodeMemory) override;
 
diff --git a/Engine/Source/Runtime/AIModule/Classes/BehaviorTree/Tasks/BTTask_PlayAnimation.h b/Engine/Source/Runtime/AIModule/Classes/BehaviorTree/Tasks/BTTask_PlayAnimation.h
--- a/Engine/Source/Runtime/AIModule/Classes/BehaviorTree/Tasks/BTTask_PlayAnimation.h
+++ b/Engine/Source/Runtime/AIModule/Classes/BehaviorTree/Tasks/BTTask_PlayAnimation.h
@@ -14,7 +14,9 @@ class UAnimationAsset;
 UCLASS()
 class AIMODULE_API UBTTask_PlayAnimation : public UBTTaskNode
 {
-	GENERATED_UCLASS_BODY()
+	GENERATED_BODY()
+public:
+	UBTTask_PlayAnimation(const FObjectInitializer& ObjectInitializer = FObjectInitializer::Get());
 
 	/** Animation asset to play. Note that it needs to match the skeleton of pawn this BT is controlling */
 	UPROPERTY(Category = Node, EditAnywhere)
diff --git a/Engine/Source/Runtime/AIModule/Classes/BehaviorTree/Tasks/BTTask_PlaySound.h b/Engine/Source/Runtime/AIModule/Classes/BehaviorTree/Tasks/BTTask_PlaySound.h
--- a/Engine/Source/Runtime/AIModule/Classes/BehaviorTree/Tasks/BTTask_PlaySound.h
+++ b/Engine/Source/Runtime/AIModule/Classes/BehaviorTree/Tasks/BTTask_PlaySound.h
@@ -13,7 +13,9 @@ class USoundCue;
 UCLASS()
 class AIMODULE_API UBTTask_PlaySound : public UBTTaskNode
 {
-	GENERATED_UCLASS_BODY()
+	GENERATED_BODY()
+public:
+	UBTTask_PlaySound(const FObjectInitializer& ObjectInitializer = FObjectInitializer::Get());
 
 	/** CUE to play */
 	UPROPERTY(Category=Node, EditAnywhere)
