Back out changelist 2481333

diff --git a/Engine/Source/Editor/UnrealEd/Classes/Factories/ReimportFbxStaticMeshFactory.h b/Engine/Source/Editor/UnrealEd/Classes/Factories/ReimportFbxStaticMeshFactory.h
--- a/Engine/Source/Editor/UnrealEd/Classes/Factories/ReimportFbxStaticMeshFactory.h
+++ b/Engine/Source/Editor/UnrealEd/Classes/Factories/ReimportFbxStaticMeshFactory.h
@@ -10,9 +10,7 @@
 UCLASS(collapsecategories)
 class UReimportFbxStaticMeshFactory : public UFbxFactory, public FReimportHandler
 {
-	GENERATED_BODY()
-public:
-	UReimportFbxStaticMeshFactory(const FObjectInitializer& ObjectInitializer = FObjectInitializer::Get());
+	GENERATED_UCLASS_BODY()
 
 
 	// Begin FReimportHandler interface
diff --git a/Engine/Source/Editor/UnrealEd/Classes/Factories/ReimportSoundFactory.h b/Engine/Source/Editor/UnrealEd/Classes/Factories/ReimportSoundFactory.h
--- a/Engine/Source/Editor/UnrealEd/Classes/Factories/ReimportSoundFactory.h
+++ b/Engine/Source/Editor/UnrealEd/Classes/Factories/ReimportSoundFactory.h
@@ -10,9 +10,7 @@
 UCLASS(hidecategories=Object, collapsecategories)
 class UReimportSoundFactory : public USoundFactory, public FReimportHandler
 {
-	GENERATED_BODY()
-public:
-	UReimportSoundFactory(const FObjectInitializer& ObjectInitializer = FObjectInitializer::Get());
+	GENERATED_UCLASS_BODY()
 
 
 	// Begin FReimportHandler interface
diff --git a/Engine/Source/Editor/UnrealEd/Classes/Factories/ReimportSoundSurroundFactory.h b/Engine/Source/Editor/UnrealEd/Classes/Factories/ReimportSoundSurroundFactory.h
--- a/Engine/Source/Editor/UnrealEd/Classes/Factories/ReimportSoundSurroundFactory.h
+++ b/Engine/Source/Editor/UnrealEd/Classes/Factories/ReimportSoundSurroundFactory.h
@@ -10,9 +10,7 @@
 UCLASS(hidecategories=Object, collapsecategories)
 class UReimportSoundSurroundFactory : public USoundSurroundFactory, public FReimportHandler
 {
-	GENERATED_BODY()
-public:
-	UReimportSoundSurroundFactory(const FObjectInitializer& ObjectInitializer = FObjectInitializer::Get());
+	GENERATED_UCLASS_BODY()
 
 	UPROPERTY()
 	TArray<FString> ReimportPaths;
diff --git a/Engine/Source/Editor/UnrealEd/Classes/Factories/ReimportVectorFieldStaticFactory.h b/Engine/Source/Editor/UnrealEd/Classes/Factories/ReimportVectorFieldStaticFactory.h
--- a/Engine/Source/Editor/UnrealEd/Classes/Factories/ReimportVectorFieldStaticFactory.h
+++ b/Engine/Source/Editor/UnrealEd/Classes/Factories/ReimportVectorFieldStaticFactory.h
@@ -10,9 +10,7 @@
 UCLASS()
 class UReimportVectorFieldStaticFactory : public UVectorFieldStaticFactory, public FReimportHandler
 {
-	GENERATED_BODY()
-public:
-	UReimportVectorFieldStaticFactory(const FObjectInitializer& ObjectInitializer = FObjectInitializer::Get());
+	GENERATED_UCLASS_BODY()
 
 	// Begin FReimportHandler interface
 	virtual bool CanReimport( UObject* Obj, TArray<FString>& OutFilenames ) override;
diff --git a/Engine/Source/Editor/UnrealEd/Classes/Factories/ReverbEffectFactory.h b/Engine/Source/Editor/UnrealEd/Classes/Factories/ReverbEffectFactory.h
--- a/Engine/Source/Editor/UnrealEd/Classes/Factories/ReverbEffectFactory.h
+++ b/Engine/Source/Editor/UnrealEd/Classes/Factories/ReverbEffectFactory.h
@@ -10,9 +10,7 @@
 UCLASS(hidecategories=Object)
 class UReverbEffectFactory : public UFactory
 {
-	GENERATED_BODY()
-public:
-	UReverbEffectFactory(const FObjectInitializer& ObjectInitializer = FObjectInitializer::Get());
+	GENERATED_UCLASS_BODY()
 
 
 	// Begin UFactory Interface
diff --git a/Engine/Source/Editor/UnrealEd/Classes/Factories/SkeletonFactory.h b/Engine/Source/Editor/UnrealEd/Classes/Factories/SkeletonFactory.h
--- a/Engine/Source/Editor/UnrealEd/Classes/Factories/SkeletonFactory.h
+++ b/Engine/Source/Editor/UnrealEd/Classes/Factories/SkeletonFactory.h
@@ -10,9 +10,7 @@
 UCLASS(HideCategories=Object,MinimalAPI)
 class USkeletonFactory : public UFactory
 {
-	GENERATED_BODY()
-public:
-	UNREALED_API USkeletonFactory(const FObjectInitializer& ObjectInitializer = FObjectInitializer::Get());
+	GENERATED_UCLASS_BODY()
 
 	// The skeletal mesh with which to initialize this skeleton.
 	UPROPERTY()
diff --git a/Engine/Source/Editor/UnrealEd/Classes/Factories/SlateBrushAssetFactory.h b/Engine/Source/Editor/UnrealEd/Classes/Factories/SlateBrushAssetFactory.h
--- a/Engine/Source/Editor/UnrealEd/Classes/Factories/SlateBrushAssetFactory.h
+++ b/Engine/Source/Editor/UnrealEd/Classes/Factories/SlateBrushAssetFactory.h
@@ -8,9 +8,7 @@
 UCLASS(hidecategories=Object, MinimalAPI)
 class USlateBrushAssetFactory : public UFactory
 {
-	GENERATED_BODY()
-public:
-	UNREALED_API USlateBrushAssetFactory(const FObjectInitializer& ObjectInitializer = FObjectInitializer::Get());
+	GENERATED_UCLASS_BODY()
 
 	/** An initial texture to assign to the newly created slate brush */
 	UPROPERTY()
diff --git a/Engine/Source/Editor/UnrealEd/Classes/Factories/SlateWidgetStyleAssetFactory.h b/Engine/Source/Editor/UnrealEd/Classes/Factories/SlateWidgetStyleAssetFactory.h
--- a/Engine/Source/Editor/UnrealEd/Classes/Factories/SlateWidgetStyleAssetFactory.h
+++ b/Engine/Source/Editor/UnrealEd/Classes/Factories/SlateWidgetStyleAssetFactory.h
@@ -8,9 +8,7 @@
 UCLASS(hidecategories=Object)
 class USlateWidgetStyleAssetFactory : public UFactory
 {
-	GENERATED_BODY()
-public:
-	USlateWidgetStyleAssetFactory(const FObjectInitializer& ObjectInitializer = FObjectInitializer::Get());
+	GENERATED_UCLASS_BODY()
 
 	UPROPERTY(EditAnywhere, Category=StyleType)
 	TSubclassOf< USlateWidgetStyleContainerBase > StyleType;
diff --git a/Engine/Source/Editor/UnrealEd/Classes/Factories/SoundAttenuationFactory.h b/Engine/Source/Editor/UnrealEd/Classes/Factories/SoundAttenuationFactory.h
--- a/Engine/Source/Editor/UnrealEd/Classes/Factories/SoundAttenuationFactory.h
+++ b/Engine/Source/Editor/UnrealEd/Classes/Factories/SoundAttenuationFactory.h
@@ -10,9 +10,7 @@
 UCLASS(hidecategories=Object, MinimalAPI)
 class USoundAttenuationFactory : public UFactory
 {
-	GENERATED_BODY()
-public:
-	UNREALED_API USoundAttenuationFactory(const FObjectInitializer& ObjectInitializer = FObjectInitializer::Get());
+	GENERATED_UCLASS_BODY()
 
 	// Begin UFactory Interface
 	virtual UObject* FactoryCreateNew(UClass* Class,UObject* InParent,FName Name,EObjectFlags Flags,UObject* Context,FFeedbackContext* Warn) override;
diff --git a/Engine/Source/Editor/UnrealEd/Classes/Factories/SoundClassFactory.h b/Engine/Source/Editor/UnrealEd/Classes/Factories/SoundClassFactory.h
--- a/Engine/Source/Editor/UnrealEd/Classes/Factories/SoundClassFactory.h
+++ b/Engine/Source/Editor/UnrealEd/Classes/Factories/SoundClassFactory.h
@@ -10,9 +10,7 @@
 UCLASS(hidecategories=Object)
 class USoundClassFactory : public UFactory
 {
-	GENERATED_BODY()
-public:
-	USoundClassFactory(const FObjectInitializer& ObjectInitializer = FObjectInitializer::Get());
+	GENERATED_UCLASS_BODY()
 
 
 	// Begin UFactory Interface
diff --git a/Engine/Source/Editor/UnrealEd/Classes/Factories/SoundCueFactoryNew.h b/Engine/Source/Editor/UnrealEd/Classes/Factories/SoundCueFactoryNew.h
--- a/Engine/Source/Editor/UnrealEd/Classes/Factories/SoundCueFactoryNew.h
+++ b/Engine/Source/Editor/UnrealEd/Classes/Factories/SoundCueFactoryNew.h
@@ -10,9 +10,7 @@
 UCLASS(hidecategories=Object, MinimalAPI)
 class USoundCueFactoryNew : public UFactory
 {
-	GENERATED_BODY()
-public:
-	UNREALED_API USoundCueFactoryNew(const FObjectInitializer& ObjectInitializer = FObjectInitializer::Get());
+	GENERATED_UCLASS_BODY()
 
 
 	// Begin UFactory Interface
diff --git a/Engine/Source/Editor/UnrealEd/Classes/Factories/SoundFactory.h b/Engine/Source/Editor/UnrealEd/Classes/Factories/SoundFactory.h
--- a/Engine/Source/Editor/UnrealEd/Classes/Factories/SoundFactory.h
+++ b/Engine/Source/Editor/UnrealEd/Classes/Factories/SoundFactory.h
@@ -10,9 +10,7 @@
 UCLASS(hidecategories=Object)
 class USoundFactory : public UFactory
 {
-	GENERATED_BODY()
-public:
-	USoundFactory(const FObjectInitializer& ObjectInitializer = FObjectInitializer::Get());
+	GENERATED_UCLASS_BODY()
 
 	/** If enabled, a sound cue will automatically be created for the sound */
 	UPROPERTY(EditAnywhere, Category=SoundFactory, meta=(ToolTip="If enabled, a sound cue will automatically be created for the sound"))
diff --git a/Engine/Source/Editor/UnrealEd/Classes/Factories/SoundMixFactory.h b/Engine/Source/Editor/UnrealEd/Classes/Factories/SoundMixFactory.h
--- a/Engine/Source/Editor/UnrealEd/Classes/Factories/SoundMixFactory.h
+++ b/Engine/Source/Editor/UnrealEd/Classes/Factories/SoundMixFactory.h
@@ -10,9 +10,7 @@
 UCLASS(hidecategories=Object)
 class USoundMixFactory : public UFactory
 {
-	GENERATED_BODY()
-public:
-	USoundMixFactory(const FObjectInitializer& ObjectInitializer = FObjectInitializer::Get());
+	GENERATED_UCLASS_BODY()
 
 
 	// Begin UFactory Interface
diff --git a/Engine/Source/Editor/UnrealEd/Classes/Factories/SoundSurroundFactory.h b/Engine/Source/Editor/UnrealEd/Classes/Factories/SoundSurroundFactory.h
--- a/Engine/Source/Editor/UnrealEd/Classes/Factories/SoundSurroundFactory.h
+++ b/Engine/Source/Editor/UnrealEd/Classes/Factories/SoundSurroundFactory.h
@@ -10,9 +10,7 @@
 UCLASS(hidecategories=Object)
 class USoundSurroundFactory : public UFactory
 {
-	GENERATED_BODY()
-public:
-	USoundSurroundFactory(const FObjectInitializer& ObjectInitializer = FObjectInitializer::Get());
+	GENERATED_UCLASS_BODY()
 
 	UPROPERTY()
 	float CueVolume;
diff --git a/Engine/Source/Editor/UnrealEd/Classes/Factories/SubsurfaceProfileFactory.h b/Engine/Source/Editor/UnrealEd/Classes/Factories/SubsurfaceProfileFactory.h
--- a/Engine/Source/Editor/UnrealEd/Classes/Factories/SubsurfaceProfileFactory.h
+++ b/Engine/Source/Editor/UnrealEd/Classes/Factories/SubsurfaceProfileFactory.h
@@ -6,9 +6,7 @@
 UCLASS(hidecategories=Object)
 class USubsurfaceProfileFactory : public UFactory
 {
-	GENERATED_BODY()
-public:
-	USubsurfaceProfileFactory(const FObjectInitializer& ObjectInitializer = FObjectInitializer::Get());
+	GENERATED_UCLASS_BODY()
 
 
 	// Begin UFactory Interface
diff --git a/Engine/Source/Editor/UnrealEd/Classes/Factories/Texture2dFactoryNew.h b/Engine/Source/Editor/UnrealEd/Classes/Factories/Texture2dFactoryNew.h
--- a/Engine/Source/Editor/UnrealEd/Classes/Factories/Texture2dFactoryNew.h
+++ b/Engine/Source/Editor/UnrealEd/Classes/Factories/Texture2dFactoryNew.h
@@ -10,9 +10,7 @@
 UCLASS(hidecategories=Object, MinimalAPI)
 class UTexture2DFactoryNew : public UFactory
 {
-	GENERATED_BODY()
-public:
-	UNREALED_API UTexture2DFactoryNew(const FObjectInitializer& ObjectInitializer = FObjectInitializer::Get());
+	GENERATED_UCLASS_BODY()
 
 	/** width of new texture */
 	UPROPERTY()
diff --git a/Engine/Source/Editor/UnrealEd/Classes/Factories/TextureFactory.h b/Engine/Source/Editor/UnrealEd/Classes/Factories/TextureFactory.h
--- a/Engine/Source/Editor/UnrealEd/Classes/Factories/TextureFactory.h
+++ b/Engine/Source/Editor/UnrealEd/Classes/Factories/TextureFactory.h
@@ -5,11 +5,10 @@
 #include "Factories/Factory.h"
 #include "TextureFactory.generated.h"
 
-UCLASS(collapsecategories, hidecategories=Object)
+UCLASS(customconstructor, collapsecategories, hidecategories=Object)
 class UNREALED_API UTextureFactory : public UFactory
 {
-	GENERATED_BODY()
-public:
+	GENERATED_UCLASS_BODY()
 
 	UPROPERTY()
 	uint32 NoCompression:1;
diff --git a/Engine/Source/Editor/UnrealEd/Classes/Factories/TextureRenderTargetCubeFactoryNew.h b/Engine/Source/Editor/UnrealEd/Classes/Factories/TextureRenderTargetCubeFactoryNew.h
--- a/Engine/Source/Editor/UnrealEd/Classes/Factories/TextureRenderTargetCubeFactoryNew.h
+++ b/Engine/Source/Editor/UnrealEd/Classes/Factories/TextureRenderTargetCubeFactoryNew.h
@@ -10,9 +10,7 @@
 UCLASS(hidecategories=(Object, Texture))
 class UTextureRenderTargetCubeFactoryNew : public UFactory
 {
-	GENERATED_BODY()
-public:
-	UTextureRenderTargetCubeFactoryNew(const FObjectInitializer& ObjectInitializer = FObjectInitializer::Get());
+	GENERATED_UCLASS_BODY()
 
 	/** width of new texture */
 	UPROPERTY(meta=(ToolTip="Width of the texture render target"))
diff --git a/Engine/Source/Editor/UnrealEd/Classes/Factories/TextureRenderTargetFactoryNew.h b/Engine/Source/Editor/UnrealEd/Classes/Factories/TextureRenderTargetFactoryNew.h
--- a/Engine/Source/Editor/UnrealEd/Classes/Factories/TextureRenderTargetFactoryNew.h
+++ b/Engine/Source/Editor/UnrealEd/Classes/Factories/TextureRenderTargetFactoryNew.h
@@ -10,9 +10,7 @@
 UCLASS(hidecategories=(Object, Texture))
 class UTextureRenderTargetFactoryNew : public UFactory
 {
-	GENERATED_BODY()
-public:
-	UTextureRenderTargetFactoryNew(const FObjectInitializer& ObjectInitializer = FObjectInitializer::Get());
+	GENERATED_UCLASS_BODY()
 
 	/** Width of the texture render target */
 	UPROPERTY(meta=(ToolTip="Width of the texture render target"))
diff --git a/Engine/Source/Editor/UnrealEd/Classes/Factories/TouchInterfaceFactory.h b/Engine/Source/Editor/UnrealEd/Classes/Factories/TouchInterfaceFactory.h
--- a/Engine/Source/Editor/UnrealEd/Classes/Factories/TouchInterfaceFactory.h
+++ b/Engine/Source/Editor/UnrealEd/Classes/Factories/TouchInterfaceFactory.h
@@ -6,9 +6,7 @@
 UCLASS(hidecategories=Object)
 class UTouchInterfaceFactory : public UFactory
 {
-	GENERATED_BODY()
-public:
-	UTouchInterfaceFactory(const FObjectInitializer& ObjectInitializer = FObjectInitializer::Get());
+	GENERATED_UCLASS_BODY()
 
 
 	// Begin UFactory Interface
diff --git a/Engine/Source/Editor/UnrealEd/Classes/Factories/TrueTypeFontFactory.h b/Engine/Source/Editor/UnrealEd/Classes/Factories/TrueTypeFontFactory.h
--- a/Engine/Source/Editor/UnrealEd/Classes/Factories/TrueTypeFontFactory.h
+++ b/Engine/Source/Editor/UnrealEd/Classes/Factories/TrueTypeFontFactory.h
@@ -10,9 +10,7 @@
 UCLASS(hidecategories=Object, collapsecategories)
 class UNREALED_API UTrueTypeFontFactory : public UTextureFactory, public FReimportHandler
 {
-	GENERATED_BODY()
-public:
-	UTrueTypeFontFactory(const FObjectInitializer& ObjectInitializer = FObjectInitializer::Get());
+	GENERATED_UCLASS_BODY()
 
 	/** Import options for this font */
 	UPROPERTY(EditAnywhere, Instanced, Category = TrueTypeFontFactory, meta = (ToolTip = "Import options for the font"))
diff --git a/Engine/Source/Editor/UnrealEd/Classes/Factories/VectorFieldStaticFactory.h b/Engine/Source/Editor/UnrealEd/Classes/Factories/VectorFieldStaticFactory.h
--- a/Engine/Source/Editor/UnrealEd/Classes/Factories/VectorFieldStaticFactory.h
+++ b/Engine/Source/Editor/UnrealEd/Classes/Factories/VectorFieldStaticFactory.h
@@ -11,9 +11,7 @@
 UCLASS()
 class UVectorFieldStaticFactory : public UFactory
 {
-	GENERATED_BODY()
-public:
-	UVectorFieldStaticFactory(const FObjectInitializer& ObjectInitializer = FObjectInitializer::Get());
+	GENERATED_UCLASS_BODY()
 
 
 	// Begin UFactory Interface
diff --git a/Engine/Source/Editor/UnrealEd/Classes/Factories/WorldFactory.h b/Engine/Source/Editor/UnrealEd/Classes/Factories/WorldFactory.h
--- a/Engine/Source/Editor/UnrealEd/Classes/Factories/WorldFactory.h
+++ b/Engine/Source/Editor/UnrealEd/Classes/Factories/WorldFactory.h
@@ -11,9 +11,7 @@
 UCLASS(MinimalAPI)
 class UWorldFactory : public UFactory
 {
-	GENERATED_BODY()
-public:
-	UNREALED_API UWorldFactory(const FObjectInitializer& ObjectInitializer = FObjectInitializer::Get());
+	GENERATED_UCLASS_BODY()
 
 	TEnumAsByte<EWorldType::Type> WorldType;
 	bool bInformEngineOfWorld;
diff --git a/Engine/Source/Editor/UnrealEd/Classes/MaterialEditor/DEditorFontParameterValue.h b/Engine/Source/Editor/UnrealEd/Classes/MaterialEditor/DEditorFontParameterValue.h
--- a/Engine/Source/Editor/UnrealEd/Classes/MaterialEditor/DEditorFontParameterValue.h
+++ b/Engine/Source/Editor/UnrealEd/Classes/MaterialEditor/DEditorFontParameterValue.h
@@ -10,7 +10,7 @@
 USTRUCT()
 struct UNREALED_API FDFontParameters
 {
-	GENERATED_BODY()
+	GENERATED_USTRUCT_BODY()
 
 	UPROPERTY(EditAnywhere, Category = DFontParameter)
 	class UFont* FontValue;
@@ -22,9 +22,7 @@ struct UNREALED_API FDFontParameters
 UCLASS(hidecategories = Object, collapsecategories)
 class UNREALED_API UDEditorFontParameterValue : public UDEditorParameterValue
 {
-	GENERATED_BODY()
-public:
-	UDEditorFontParameterValue(const FObjectInitializer& ObjectInitializer = FObjectInitializer::Get());
+	GENERATED_UCLASS_BODY()
 
 	UPROPERTY(EditAnywhere, Category = DEditorFontParameterValue)
 	struct FDFontParameters ParameterValue;
diff --git a/Engine/Source/Editor/UnrealEd/Classes/MaterialEditor/DEditorParameterValue.h b/Engine/Source/Editor/UnrealEd/Classes/MaterialEditor/DEditorParameterValue.h
--- a/Engine/Source/Editor/UnrealEd/Classes/MaterialEditor/DEditorParameterValue.h
+++ b/Engine/Source/Editor/UnrealEd/Classes/MaterialEditor/DEditorParameterValue.h
@@ -9,9 +9,7 @@
 UCLASS(hidecategories=Object, collapsecategories, editinlinenew)
 class UNREALED_API UDEditorParameterValue : public UObject
 {
-	GENERATED_BODY()
-public:
-	UDEditorParameterValue(const FObjectInitializer& ObjectInitializer = FObjectInitializer::Get());
+	GENERATED_UCLASS_BODY()
 
 	UPROPERTY(EditAnywhere, Category=DEditorParameterValue)
 	uint32 bOverride:1;
diff --git a/Engine/Source/Editor/UnrealEd/Classes/MaterialEditor/DEditorScalarParameterValue.h b/Engine/Source/Editor/UnrealEd/Classes/MaterialEditor/DEditorScalarParameterValue.h
--- a/Engine/Source/Editor/UnrealEd/Classes/MaterialEditor/DEditorScalarParameterValue.h
+++ b/Engine/Source/Editor/UnrealEd/Classes/MaterialEditor/DEditorScalarParameterValue.h
@@ -9,9 +9,7 @@
 UCLASS(hidecategories=Object, collapsecategories, editinlinenew)
 class UNREALED_API UDEditorScalarParameterValue : public UDEditorParameterValue
 {
-	GENERATED_BODY()
-public:
-	UDEditorScalarParameterValue(const FObjectInitializer& ObjectInitializer = FObjectInitializer::Get());
+	GENERATED_UCLASS_BODY()
 
 	UPROPERTY(EditAnywhere, Category=DEditorScalarParameterValue)
 	float ParameterValue;
diff --git a/Engine/Source/Editor/UnrealEd/Classes/MaterialEditor/DEditorStaticComponentMaskParameterValue.h b/Engine/Source/Editor/UnrealEd/Classes/MaterialEditor/DEditorStaticComponentMaskParameterValue.h
--- a/Engine/Source/Editor/UnrealEd/Classes/MaterialEditor/DEditorStaticComponentMaskParameterValue.h
+++ b/Engine/Source/Editor/UnrealEd/Classes/MaterialEditor/DEditorStaticComponentMaskParameterValue.h
@@ -9,7 +9,7 @@
 USTRUCT()
 struct UNREALED_API FDComponentMaskParameter
 {
-	GENERATED_BODY()
+	GENERATED_USTRUCT_BODY()
 
 	UPROPERTY(EditAnywhere, Category=DComponentMaskParameter)
 	uint32 R:1;
@@ -46,9 +46,7 @@ struct UNREALED_API FDComponentMaskParameter
 UCLASS(hidecategories=Object, collapsecategories)
 class UNREALED_API UDEditorStaticComponentMaskParameterValue : public UDEditorParameterValue
 {
-	GENERATED_BODY()
-public:
-	UDEditorStaticComponentMaskParameterValue(const FObjectInitializer& ObjectInitializer = FObjectInitializer::Get());
+	GENERATED_UCLASS_BODY()
 
 	UPROPERTY(EditAnywhere, Category=DEditorStaticComponentMaskParameterValue)
 	struct FDComponentMaskParameter ParameterValue;
diff --git a/Engine/Source/Editor/UnrealEd/Classes/MaterialEditor/DEditorTextureParameterValue.h b/Engine/Source/Editor/UnrealEd/Classes/MaterialEditor/DEditorTextureParameterValue.h
--- a/Engine/Source/Editor/UnrealEd/Classes/MaterialEditor/DEditorTextureParameterValue.h
+++ b/Engine/Source/Editor/UnrealEd/Classes/MaterialEditor/DEditorTextureParameterValue.h
@@ -9,9 +9,7 @@
 UCLASS(hidecategories=Object, collapsecategories)
 class UNREALED_API UDEditorTextureParameterValue : public UDEditorParameterValue
 {
-	GENERATED_BODY()
-public:
-	UDEditorTextureParameterValue(const FObjectInitializer& ObjectInitializer = FObjectInitializer::Get());
+	GENERATED_UCLASS_BODY()
 
 	UPROPERTY(EditAnywhere, Category=DEditorTextureParameterValue)
 	class UTexture* ParameterValue;
diff --git a/Engine/Source/Editor/UnrealEd/Classes/MaterialEditor/DEditorVectorParameterValue.h b/Engine/Source/Editor/UnrealEd/Classes/MaterialEditor/DEditorVectorParameterValue.h
--- a/Engine/Source/Editor/UnrealEd/Classes/MaterialEditor/DEditorVectorParameterValue.h
+++ b/Engine/Source/Editor/UnrealEd/Classes/MaterialEditor/DEditorVectorParameterValue.h
@@ -9,9 +9,7 @@
 UCLASS(hidecategories=Object, collapsecategories, editinlinenew)
 class UNREALED_API UDEditorVectorParameterValue : public UDEditorParameterValue
 {
-	GENERATED_BODY()
-public:
-	UDEditorVectorParameterValue(const FObjectInitializer& ObjectInitializer = FObjectInitializer::Get());
+	GENERATED_UCLASS_BODY()
 
 	UPROPERTY(EditAnywhere, Category=DEditorVectorParameterValue)
 	FLinearColor ParameterValue;
diff --git a/Engine/Source/Editor/UnrealEd/Classes/MaterialEditor/MaterialEditorInstanceConstant.h b/Engine/Source/Editor/UnrealEd/Classes/MaterialEditor/MaterialEditorInstanceConstant.h
--- a/Engine/Source/Editor/UnrealEd/Classes/MaterialEditor/MaterialEditorInstanceConstant.h
+++ b/Engine/Source/Editor/UnrealEd/Classes/MaterialEditor/MaterialEditorInstanceConstant.h
@@ -13,7 +13,7 @@
 USTRUCT()
 struct FEditorParameterGroup
 {
-	GENERATED_BODY()
+	GENERATED_USTRUCT_BODY()
 
 	UPROPERTY()
 	FName GroupName;
@@ -26,7 +26,7 @@ struct FEditorParameterGroup
 USTRUCT()
 struct FEditorParameterValue
 {
-	GENERATED_BODY()
+	GENERATED_USTRUCT_BODY()
 
 	UPROPERTY(EditAnywhere, Category=EditorParameterValue)
 	uint32 bOverride:1;
@@ -48,7 +48,7 @@ struct FEditorParameterValue
 USTRUCT()
 struct FEditorVectorParameterValue : public FEditorParameterValue
 {
-	GENERATED_BODY()
+	GENERATED_USTRUCT_BODY()
 
 	UPROPERTY(EditAnywhere, Category=EditorVectorParameterValue)
 	FLinearColor ParameterValue;
@@ -64,7 +64,7 @@ struct FEditorVectorParameterValue : public FEditorParameterValue
 USTRUCT()
 struct FEditorScalarParameterValue : public FEditorParameterValue
 {
-	GENERATED_BODY()
+	GENERATED_USTRUCT_BODY()
 
 	UPROPERTY(EditAnywhere, Category=EditorScalarParameterValue)
 	float ParameterValue;
@@ -80,7 +80,7 @@ struct FEditorScalarParameterValue : public FEditorParameterValue
 USTRUCT()
 struct FEditorTextureParameterValue : public FEditorParameterValue
 {
-	GENERATED_BODY()
+	GENERATED_USTRUCT_BODY()
 
 	UPROPERTY(EditAnywhere, Category=EditorTextureParameterValue)
 	class UTexture* ParameterValue;
@@ -96,7 +96,7 @@ struct FEditorTextureParameterValue : public FEditorParameterValue
 USTRUCT()
 struct FEditorFontParameterValue : public FEditorParameterValue
 {
-	GENERATED_BODY()
+	GENERATED_USTRUCT_BODY()
 
 	UPROPERTY(EditAnywhere, Category=EditorFontParameterValue)
 	class UFont* FontValue;
@@ -116,7 +116,7 @@ struct FEditorFontParameterValue : public FEditorParameterValue
 USTRUCT()
 struct FEditorStaticSwitchParameterValue : public FEditorParameterValue
 {
-	GENERATED_BODY()
+	GENERATED_USTRUCT_BODY()
 
 	UPROPERTY(EditAnywhere, Category=EditorStaticSwitchParameterValue)
 	uint32 ParameterValue:1;
@@ -142,7 +142,7 @@ struct FEditorStaticSwitchParameterValue : public FEditorParameterValue
 USTRUCT()
 struct FComponentMaskParameter
 {
-	GENERATED_BODY()
+	GENERATED_USTRUCT_BODY()
 
 	UPROPERTY(EditAnywhere, Category=ComponentMaskParameter)
 	uint32 R:1;
@@ -180,7 +180,7 @@ struct FComponentMaskParameter
 USTRUCT()
 struct FEditorStaticComponentMaskParameterValue : public FEditorParameterValue
 {
-	GENERATED_BODY()
+	GENERATED_USTRUCT_BODY()
 
 	UPROPERTY(EditAnywhere, Category=EditorStaticComponentMaskParameterValue)
 	struct FComponentMaskParameter ParameterValue;
@@ -202,9 +202,7 @@ struct FEditorStaticComponentMaskParameterValue : public FEditorParameterValue
 UCLASS(hidecategories=Object, collapsecategories)
 class UNREALED_API UMaterialEditorInstanceConstant : public UObject
 {
-	GENERATED_BODY()
-public:
-	UMaterialEditorInstanceConstant(const FObjectInitializer& ObjectInitializer = FObjectInitializer::Get());
+	GENERATED_UCLASS_BODY()
 
 	/** Physical material to use for this graphics material. Used for sounds, effects etc.*/
 	UPROPERTY(EditAnywhere, Category=MaterialEditorInstanceConstant)
diff --git a/Engine/Source/Editor/UnrealEd/Classes/MaterialEditor/MaterialEditorMeshComponent.h b/Engine/Source/Editor/UnrealEd/Classes/MaterialEditor/MaterialEditorMeshComponent.h
--- a/Engine/Source/Editor/UnrealEd/Classes/MaterialEditor/MaterialEditorMeshComponent.h
+++ b/Engine/Source/Editor/UnrealEd/Classes/MaterialEditor/MaterialEditorMeshComponent.h
@@ -8,9 +8,7 @@
 UCLASS()
 class UNREALED_API UMaterialEditorMeshComponent : public UStaticMeshComponent
 {
-	GENERATED_BODY()
-public:
-	UMaterialEditorMeshComponent(const FObjectInitializer& ObjectInitializer = FObjectInitializer::Get());
+	GENERATED_UCLASS_BODY()
 
 	// USceneComponent Interface
 	virtual FBoxSphereBounds CalcBounds(const FTransform& LocalToWorld) const override;
diff --git a/Engine/Source/Editor/UnrealEd/Classes/MaterialEditor/PreviewMaterial.h b/Engine/Source/Editor/UnrealEd/Classes/MaterialEditor/PreviewMaterial.h
--- a/Engine/Source/Editor/UnrealEd/Classes/MaterialEditor/PreviewMaterial.h
+++ b/Engine/Source/Editor/UnrealEd/Classes/MaterialEditor/PreviewMaterial.h
@@ -9,9 +9,7 @@
 UCLASS()
 class UNREALED_API UPreviewMaterial : public UMaterial
 {
-	GENERATED_BODY()
-public:
-	UPreviewMaterial(const FObjectInitializer& ObjectInitializer = FObjectInitializer::Get());
+	GENERATED_UCLASS_BODY()
 
 
 	// Begin UMaterial interface.
diff --git a/Engine/Source/Editor/UnrealEd/Classes/MaterialGraph/MaterialGraph.h b/Engine/Source/Editor/UnrealEd/Classes/MaterialGraph/MaterialGraph.h
--- a/Engine/Source/Editor/UnrealEd/Classes/MaterialGraph/MaterialGraph.h
+++ b/Engine/Source/Editor/UnrealEd/Classes/MaterialGraph/MaterialGraph.h
@@ -69,9 +69,7 @@ struct FMaterialInputInfo
 UCLASS(MinimalAPI)
 class UMaterialGraph : public UEdGraph
 {
-	GENERATED_BODY()
-public:
-	UNREALED_API UMaterialGraph(const FObjectInitializer& ObjectInitializer = FObjectInitializer::Get());
+	GENERATED_UCLASS_BODY()
 
 	/** Material this Graph represents */
 	UPROPERTY()
diff --git a/Engine/Source/Editor/UnrealEd/Classes/MaterialGraph/MaterialGraphNode.h b/Engine/Source/Editor/UnrealEd/Classes/MaterialGraph/MaterialGraphNode.h
--- a/Engine/Source/Editor/UnrealEd/Classes/MaterialGraph/MaterialGraphNode.h
+++ b/Engine/Source/Editor/UnrealEd/Classes/MaterialGraph/MaterialGraphNode.h
@@ -7,9 +7,7 @@
 UCLASS(MinimalAPI)
 class UMaterialGraphNode : public UMaterialGraphNode_Base
 {
-	GENERATED_BODY()
-public:
-	UNREALED_API UMaterialGraphNode(const FObjectInitializer& ObjectInitializer = FObjectInitializer::Get());
+	GENERATED_UCLASS_BODY()
 
 	/** Material Expression this node is representing */
 	UPROPERTY()
diff --git a/Engine/Source/Editor/UnrealEd/Classes/MaterialGraph/MaterialGraphNode_Base.h b/Engine/Source/Editor/UnrealEd/Classes/MaterialGraph/MaterialGraphNode_Base.h
--- a/Engine/Source/Editor/UnrealEd/Classes/MaterialGraph/MaterialGraphNode_Base.h
+++ b/Engine/Source/Editor/UnrealEd/Classes/MaterialGraph/MaterialGraphNode_Base.h
@@ -7,9 +7,7 @@
 UCLASS(MinimalAPI)
 class UMaterialGraphNode_Base : public UEdGraphNode
 {
-	GENERATED_BODY()
-public:
-	UNREALED_API UMaterialGraphNode_Base(const FObjectInitializer& ObjectInitializer = FObjectInitializer::Get());
+	GENERATED_UCLASS_BODY()
 
 	/** Create all of the input pins required */
 	virtual void CreateInputPins() {};
diff --git a/Engine/Source/Editor/UnrealEd/Classes/MaterialGraph/MaterialGraphNode_Root.h b/Engine/Source/Editor/UnrealEd/Classes/MaterialGraph/MaterialGraphNode_Root.h
--- a/Engine/Source/Editor/UnrealEd/Classes/MaterialGraph/MaterialGraphNode_Root.h
+++ b/Engine/Source/Editor/UnrealEd/Classes/MaterialGraph/MaterialGraphNode_Root.h
@@ -7,9 +7,7 @@
 UCLASS(MinimalAPI)
 class UMaterialGraphNode_Root : public UMaterialGraphNode_Base
 {
-	GENERATED_BODY()
-public:
-	UNREALED_API UMaterialGraphNode_Root(const FObjectInitializer& ObjectInitializer = FObjectInitializer::Get());
+	GENERATED_UCLASS_BODY()
 
 	/** Material whose inputs this root node represents */
 	UPROPERTY()
diff --git a/Engine/Source/Editor/UnrealEd/Classes/Niagara/EdGraphSchema_Niagara.h b/Engine/Source/Editor/UnrealEd/Classes/Niagara/EdGraphSchema_Niagara.h
--- a/Engine/Source/Editor/UnrealEd/Classes/Niagara/EdGraphSchema_Niagara.h
+++ b/Engine/Source/Editor/UnrealEd/Classes/Niagara/EdGraphSchema_Niagara.h
@@ -9,7 +9,7 @@
 USTRUCT()
 struct UNREALED_API FNiagaraSchemaAction_NewNode : public FEdGraphSchemaAction
 {
-	GENERATED_BODY()
+	GENERATED_USTRUCT_BODY();
 
 	/** Template of node we want to create */
 	UPROPERTY()
@@ -45,9 +45,7 @@ struct UNREALED_API FNiagaraSchemaAction_NewNode : public FEdGraphSchemaAction
 UCLASS(MinimalAPI)
 class UEdGraphSchema_Niagara : public UEdGraphSchema
 {
-	GENERATED_BODY()
-public:
-	UNREALED_API UEdGraphSchema_Niagara(const FObjectInitializer& ObjectInitializer = FObjectInitializer::Get());
+	GENERATED_UCLASS_BODY()
 
 	// Allowable PinType.PinCategory values
 	UPROPERTY()
diff --git a/Engine/Source/Editor/UnrealEd/Classes/Niagara/NiagaraGraph.h b/Engine/Source/Editor/UnrealEd/Classes/Niagara/NiagaraGraph.h
--- a/Engine/Source/Editor/UnrealEd/Classes/Niagara/NiagaraGraph.h
+++ b/Engine/Source/Editor/UnrealEd/Classes/Niagara/NiagaraGraph.h
@@ -7,9 +7,7 @@
 UCLASS(MinimalAPI)
 class UNiagaraGraph : public UEdGraph
 {
-	GENERATED_BODY()
-public:
-	UNREALED_API UNiagaraGraph(const FObjectInitializer& ObjectInitializer = FObjectInitializer::Get());
+	GENERATED_UCLASS_BODY()
 	
 	/** Get the source that owns this graph */
 	class UNiagaraScriptSource* GetSource() const;
diff --git a/Engine/Source/Editor/UnrealEd/Classes/Niagara/NiagaraNode.h b/Engine/Source/Editor/UnrealEd/Classes/Niagara/NiagaraNode.h
--- a/Engine/Source/Editor/UnrealEd/Classes/Niagara/NiagaraNode.h
+++ b/Engine/Source/Editor/UnrealEd/Classes/Niagara/NiagaraNode.h
@@ -6,9 +6,7 @@
 UCLASS()
 class UNREALED_API UNiagaraNode : public UEdGraphNode
 {
-	GENERATED_BODY()
-public:
-	UNiagaraNode(const FObjectInitializer& ObjectInitializer = FObjectInitializer::Get());
+	GENERATED_UCLASS_BODY()
 
 public:
 
diff --git a/Engine/Source/Editor/UnrealEd/Classes/Niagara/NiagaraNodeConstant.h b/Engine/Source/Editor/UnrealEd/Classes/Niagara/NiagaraNodeConstant.h
--- a/Engine/Source/Editor/UnrealEd/Classes/Niagara/NiagaraNodeConstant.h
+++ b/Engine/Source/Editor/UnrealEd/Classes/Niagara/NiagaraNodeConstant.h
@@ -8,9 +8,7 @@
 UCLASS(MinimalAPI, Deprecated)
 class UDEPRECATED_NiagaraNodeConstant : public UNiagaraNode
 {
-	GENERATED_BODY()
-public:
-	UNREALED_API UDEPRECATED_NiagaraNodeConstant(const FObjectInitializer& ObjectInitializer = FObjectInitializer::Get());
+	GENERATED_UCLASS_BODY()
 
 public:
 	/** The type of the constant we're creating. */
diff --git a/Engine/Source/Editor/UnrealEd/Classes/Niagara/NiagaraNodeGetAttr.h b/Engine/Source/Editor/UnrealEd/Classes/Niagara/NiagaraNodeGetAttr.h
--- a/Engine/Source/Editor/UnrealEd/Classes/Niagara/NiagaraNodeGetAttr.h
+++ b/Engine/Source/Editor/UnrealEd/Classes/Niagara/NiagaraNodeGetAttr.h
@@ -6,9 +6,7 @@
 UCLASS(MinimalAPI, Deprecated)
 class UDEPRECATED_NiagaraNodeGetAttr : public UNiagaraNode
 {
-	GENERATED_BODY()
-public:
-	UNREALED_API UDEPRECATED_NiagaraNodeGetAttr(const FObjectInitializer& ObjectInitializer = FObjectInitializer::Get());
+	GENERATED_UCLASS_BODY()
 
 public:
 
diff --git a/Engine/Source/Editor/UnrealEd/Classes/Niagara/NiagaraNodeInput.h b/Engine/Source/Editor/UnrealEd/Classes/Niagara/NiagaraNodeInput.h
--- a/Engine/Source/Editor/UnrealEd/Classes/Niagara/NiagaraNodeInput.h
+++ b/Engine/Source/Editor/UnrealEd/Classes/Niagara/NiagaraNodeInput.h
@@ -6,9 +6,7 @@
 UCLASS(MinimalAPI)
 class UNiagaraNodeInput : public UNiagaraNode
 {
-	GENERATED_BODY()
-public:
-	UNREALED_API UNiagaraNodeInput(const FObjectInitializer& ObjectInitializer = FObjectInitializer::Get());
+	GENERATED_UCLASS_BODY()
 
 	void ReallocatePins();
 
diff --git a/Engine/Source/Editor/UnrealEd/Classes/Niagara/NiagaraNodeOp.h b/Engine/Source/Editor/UnrealEd/Classes/Niagara/NiagaraNodeOp.h
--- a/Engine/Source/Editor/UnrealEd/Classes/Niagara/NiagaraNodeOp.h
+++ b/Engine/Source/Editor/UnrealEd/Classes/Niagara/NiagaraNodeOp.h
@@ -6,9 +6,7 @@
 UCLASS(MinimalAPI)
 class UNiagaraNodeOp : public UNiagaraNode
 {
-	GENERATED_BODY()
-public:
-	UNREALED_API UNiagaraNodeOp(const FObjectInitializer& ObjectInitializer = FObjectInitializer::Get());
+	GENERATED_UCLASS_BODY()
 
 public:
 
@@ -81,4 +79,4 @@ private:
 	UNiagaraNode* NiagaraNode;
 };
 
-#endif
+#endif
\ No newline at end of file
diff --git a/Engine/Source/Editor/UnrealEd/Classes/Niagara/NiagaraNodeOutput.h b/Engine/Source/Editor/UnrealEd/Classes/Niagara/NiagaraNodeOutput.h
--- a/Engine/Source/Editor/UnrealEd/Classes/Niagara/NiagaraNodeOutput.h
+++ b/Engine/Source/Editor/UnrealEd/Classes/Niagara/NiagaraNodeOutput.h
@@ -6,9 +6,7 @@
 UCLASS(MinimalAPI)
 class UNiagaraNodeOutput : public UNiagaraNode
 {
-	GENERATED_BODY()
-public:
-	UNREALED_API UNiagaraNodeOutput(const FObjectInitializer& ObjectInitializer = FObjectInitializer::Get());
+	GENERATED_UCLASS_BODY()
 
 	void ReallocatePins();
 public:
diff --git a/Engine/Source/Editor/UnrealEd/Classes/Niagara/NiagaraNodeOutputUpdate.h b/Engine/Source/Editor/UnrealEd/Classes/Niagara/NiagaraNodeOutputUpdate.h
--- a/Engine/Source/Editor/UnrealEd/Classes/Niagara/NiagaraNodeOutputUpdate.h
+++ b/Engine/Source/Editor/UnrealEd/Classes/Niagara/NiagaraNodeOutputUpdate.h
@@ -6,9 +6,7 @@
 UCLASS(MinimalAPI, Deprecated)
 class UDEPRECATED_NiagaraNodeOutputUpdate : public UNiagaraNode
 {
-	GENERATED_BODY()
-public:
-	UNREALED_API UDEPRECATED_NiagaraNodeOutputUpdate(const FObjectInitializer& ObjectInitializer = FObjectInitializer::Get());
+	GENERATED_UCLASS_BODY()
 
 public:
 
diff --git a/Engine/Source/Editor/UnrealEd/Classes/Preferences/CascadeOptions.h b/Engine/Source/Editor/UnrealEd/Classes/Preferences/CascadeOptions.h
--- a/Engine/Source/Editor/UnrealEd/Classes/Preferences/CascadeOptions.h
+++ b/Engine/Source/Editor/UnrealEd/Classes/Preferences/CascadeOptions.h
@@ -13,9 +13,7 @@
 UCLASS(hidecategories=Object, config=EditorUserSettings)
 class UNREALED_API UCascadeOptions : public UObject
 {
-	GENERATED_BODY()
-public:
-	UCascadeOptions(const FObjectInitializer& ObjectInitializer = FObjectInitializer::Get());
+	GENERATED_UCLASS_BODY()
 
 	UPROPERTY(EditAnywhere, config, Category=Options)
 	uint32 bShowModuleDump:1;
diff --git a/Engine/Source/Editor/UnrealEd/Classes/Preferences/CurveEdOptions.h b/Engine/Source/Editor/UnrealEd/Classes/Preferences/CurveEdOptions.h
--- a/Engine/Source/Editor/UnrealEd/Classes/Preferences/CurveEdOptions.h
+++ b/Engine/Source/Editor/UnrealEd/Classes/Preferences/CurveEdOptions.h
@@ -13,9 +13,7 @@
 UCLASS(hidecategories=Object, config=EditorUserSettings)
 class UNREALED_API UCurveEdOptions : public UObject
 {
-	GENERATED_BODY()
-public:
-	UCurveEdOptions(const FObjectInitializer& ObjectInitializer = FObjectInitializer::Get());
+	GENERATED_UCLASS_BODY()
 
 	UPROPERTY(EditAnywhere, config, Category=Options)
 	float MinViewRange;
diff --git a/Engine/Source/Editor/UnrealEd/Classes/Preferences/MaterialEditorOptions.h b/Engine/Source/Editor/UnrealEd/Classes/Preferences/MaterialEditorOptions.h
--- a/Engine/Source/Editor/UnrealEd/Classes/Preferences/MaterialEditorOptions.h
+++ b/Engine/Source/Editor/UnrealEd/Classes/Preferences/MaterialEditorOptions.h
@@ -12,9 +12,7 @@
 UCLASS(hidecategories=Object, config=EditorUserSettings)
 class UNREALED_API UMaterialEditorOptions : public UObject
 {
-	GENERATED_BODY()
-public:
-	UMaterialEditorOptions(const FObjectInitializer& ObjectInitializer = FObjectInitializer::Get());
+	GENERATED_UCLASS_BODY()
 
 	/** If true, render grid the preview scene. */
 	UPROPERTY(EditAnywhere, config, Category=Options)
diff --git a/Engine/Source/Editor/UnrealEd/Classes/Preferences/PersonaOptions.h b/Engine/Source/Editor/UnrealEd/Classes/Preferences/PersonaOptions.h
--- a/Engine/Source/Editor/UnrealEd/Classes/Preferences/PersonaOptions.h
+++ b/Engine/Source/Editor/UnrealEd/Classes/Preferences/PersonaOptions.h
@@ -13,9 +13,7 @@
 UCLASS(hidecategories=Object, config=EditorUserSettings)
 class UNREALED_API UPersonaOptions : public UObject
 {
-	GENERATED_BODY()
-public:
-	UPersonaOptions(const FObjectInitializer& ObjectInitializer = FObjectInitializer::Get());
+	GENERATED_UCLASS_BODY()
 
 	UPROPERTY(EditAnywhere, config, Category=Options)
 	uint32 bShowFloor:1;
diff --git a/Engine/Source/Editor/UnrealEd/Classes/Preferences/PhATSimOptions.h b/Engine/Source/Editor/UnrealEd/Classes/Preferences/PhATSimOptions.h
--- a/Engine/Source/Editor/UnrealEd/Classes/Preferences/PhATSimOptions.h
+++ b/Engine/Source/Editor/UnrealEd/Classes/Preferences/PhATSimOptions.h
@@ -7,9 +7,7 @@
 UCLASS(hidecategories=Object, config=EditorUserSettings)
 class UNREALED_API UPhATSimOptions : public UObject
 {
-	GENERATED_BODY()
-public:
-	UPhATSimOptions(const FObjectInitializer& ObjectInitializer = FObjectInitializer::Get());
+	GENERATED_UCLASS_BODY()
 
 	/** Lets you manually control the physics/animation */
 	UPROPERTY(EditAnywhere, transient, Category=Anim)
diff --git a/Engine/Source/Editor/UnrealEd/Classes/Preferences/UnrealEdOptions.h b/Engine/Source/Editor/UnrealEd/Classes/Preferences/UnrealEdOptions.h
--- a/Engine/Source/Editor/UnrealEd/Classes/Preferences/UnrealEdOptions.h
+++ b/Engine/Source/Editor/UnrealEd/Classes/Preferences/UnrealEdOptions.h
@@ -12,7 +12,7 @@
 USTRUCT()
 struct FEditorCommandCategory
 {
-	GENERATED_BODY()
+	GENERATED_USTRUCT_BODY()
 
 	UPROPERTY()
 	FName Parent;
@@ -26,7 +26,7 @@ struct FEditorCommandCategory
 USTRUCT()
 struct FEditorCommand
 {
-	GENERATED_BODY()
+	GENERATED_USTRUCT_BODY()
 
 	UPROPERTY()
 	FName Parent;
@@ -46,7 +46,7 @@ struct FEditorCommand
 USTRUCT()
 struct FClassPickerDefaults
 {
-	GENERATED_BODY()
+	GENERATED_USTRUCT_BODY()
 
 	/** The name of the class to select */
 	UPROPERTY()
@@ -66,9 +66,7 @@ struct FClassPickerDefaults
 UCLASS(Config=Editor)
 class UNREALED_API UUnrealEdOptions : public UObject
 {
-	GENERATED_BODY()
-public:
-	UUnrealEdOptions(const FObjectInitializer& ObjectInitializer = FObjectInitializer::Get());
+	GENERATED_UCLASS_BODY()
 
 	/** Categories of commands. */
 	UPROPERTY(config)
diff --git a/Engine/Source/Editor/UnrealEd/Classes/Settings/ContentBrowserSettings.h b/Engine/Source/Editor/UnrealEd/Classes/Settings/ContentBrowserSettings.h
--- a/Engine/Source/Editor/UnrealEd/Classes/Settings/ContentBrowserSettings.h
+++ b/Engine/Source/Editor/UnrealEd/Classes/Settings/ContentBrowserSettings.h
@@ -17,9 +17,7 @@ UCLASS(config=EditorUserSettings)
 class UNREALED_API UContentBrowserSettings
 	: public UObject
 {
-	GENERATED_BODY()
-public:
-	UContentBrowserSettings(const FObjectInitializer& ObjectInitializer = FObjectInitializer::Get());
+	GENERATED_UCLASS_BODY()
 
 public:
 
diff --git a/Engine/Source/Editor/UnrealEd/Classes/Settings/DestructableMeshEditorSettings.h b/Engine/Source/Editor/UnrealEd/Classes/Settings/DestructableMeshEditorSettings.h
--- a/Engine/Source/Editor/UnrealEd/Classes/Settings/DestructableMeshEditorSettings.h
+++ b/Engine/Source/Editor/UnrealEd/Classes/Settings/DestructableMeshEditorSettings.h
@@ -17,9 +17,7 @@ UCLASS(config=EditorUserSettings)
 class UNREALED_API UDestructableMeshEditorSettings
 	: public UObject
 {
-	GENERATED_BODY()
-public:
-	UDestructableMeshEditorSettings(const FObjectInitializer& ObjectInitializer = FObjectInitializer::Get());
+	GENERATED_UCLASS_BODY()
 
 public:
 
diff --git a/Engine/Source/Editor/UnrealEd/Classes/Settings/EditorExperimentalSettings.h b/Engine/Source/Editor/UnrealEd/Classes/Settings/EditorExperimentalSettings.h
--- a/Engine/Source/Editor/UnrealEd/Classes/Settings/EditorExperimentalSettings.h
+++ b/Engine/Source/Editor/UnrealEd/Classes/Settings/EditorExperimentalSettings.h
@@ -17,9 +17,7 @@ UCLASS(config=EditorUserSettings)
 class UNREALED_API UEditorExperimentalSettings
 	: public UObject
 {
-	GENERATED_BODY()
-public:
-	UEditorExperimentalSettings(const FObjectInitializer& ObjectInitializer = FObjectInitializer::Get());
+	GENERATED_UCLASS_BODY()
 
 public:
 
diff --git a/Engine/Source/Editor/UnrealEd/Classes/Settings/EditorLoadingSavingSettings.h b/Engine/Source/Editor/UnrealEd/Classes/Settings/EditorLoadingSavingSettings.h
--- a/Engine/Source/Editor/UnrealEd/Classes/Settings/EditorLoadingSavingSettings.h
+++ b/Engine/Source/Editor/UnrealEd/Classes/Settings/EditorLoadingSavingSettings.h
@@ -23,9 +23,7 @@ UCLASS(config=EditorUserSettings, autoexpandcategories=(AutoSave, AutoReimport,
 class UNREALED_API UEditorLoadingSavingSettings
 	: public UObject
 {
-	GENERATED_BODY()
-public:
-	UEditorLoadingSavingSettings(const FObjectInitializer& ObjectInitializer = FObjectInitializer::Get());
+	GENERATED_UCLASS_BODY()
 
 public:
 
diff --git a/Engine/Source/Editor/UnrealEd/Classes/Settings/EditorMiscSettings.h b/Engine/Source/Editor/UnrealEd/Classes/Settings/EditorMiscSettings.h
--- a/Engine/Source/Editor/UnrealEd/Classes/Settings/EditorMiscSettings.h
+++ b/Engine/Source/Editor/UnrealEd/Classes/Settings/EditorMiscSettings.h
@@ -13,9 +13,7 @@ UCLASS(config=EditorUserSettings)
 class UNREALED_API UEditorMiscSettings
 	: public UObject
 {
-	GENERATED_BODY()
-public:
-	UEditorMiscSettings(const FObjectInitializer& ObjectInitializer = FObjectInitializer::Get());
+	GENERATED_UCLASS_BODY()
 
 public:
 
diff --git a/Engine/Source/Editor/UnrealEd/Classes/Settings/LevelEditorMiscSettings.h b/Engine/Source/Editor/UnrealEd/Classes/Settings/LevelEditorMiscSettings.h
--- a/Engine/Source/Editor/UnrealEd/Classes/Settings/LevelEditorMiscSettings.h
+++ b/Engine/Source/Editor/UnrealEd/Classes/Settings/LevelEditorMiscSettings.h
@@ -12,9 +12,7 @@ UCLASS(config=EditorUserSettings)
 class UNREALED_API ULevelEditorMiscSettings
 	: public UObject
 {
-	GENERATED_BODY()
-public:
-	ULevelEditorMiscSettings(const FObjectInitializer& ObjectInitializer = FObjectInitializer::Get());
+	GENERATED_UCLASS_BODY()
 
 public:
 
diff --git a/Engine/Source/Editor/UnrealEd/Classes/Settings/LevelEditorViewportSettings.h b/Engine/Source/Editor/UnrealEd/Classes/Settings/LevelEditorViewportSettings.h
--- a/Engine/Source/Editor/UnrealEd/Classes/Settings/LevelEditorViewportSettings.h
+++ b/Engine/Source/Editor/UnrealEd/Classes/Settings/LevelEditorViewportSettings.h
@@ -63,7 +63,7 @@ enum EMeasuringToolUnits
 USTRUCT()
 struct UNREALED_API FLevelEditorViewportInstanceSettings
 {
-	GENERATED_BODY()
+	GENERATED_USTRUCT_BODY()
 
 	FLevelEditorViewportInstanceSettings()
 		: ViewportType(LVT_Perspective)
@@ -141,7 +141,7 @@ struct UNREALED_API FLevelEditorViewportInstanceSettings
 USTRUCT()
 struct UNREALED_API FLevelEditorViewportInstanceSettingsKeyValuePair
 {
-	GENERATED_BODY()
+	GENERATED_USTRUCT_BODY()
 
 	/*  Name identifying this config. */
 	UPROPERTY(config)
@@ -159,7 +159,7 @@ struct UNREALED_API FLevelEditorViewportInstanceSettingsKeyValuePair
 USTRUCT()
 struct UNREALED_API FSnapToSurfaceSettings
 {
-	GENERATED_BODY()
+	GENERATED_USTRUCT_BODY()
 
 	FSnapToSurfaceSettings()
 		: bEnabled(false)
@@ -188,9 +188,7 @@ UCLASS(config=EditorUserSettings)
 class UNREALED_API ULevelEditorViewportSettings
 	: public UObject
 {
-	GENERATED_BODY()
-public:
-	ULevelEditorViewportSettings(const FObjectInitializer& ObjectInitializer = FObjectInitializer::Get());
+	GENERATED_UCLASS_BODY()
 
 	/** Enable the use of flight camera controls under various circumstances. */
 	UPROPERTY(EditAnywhere, config, Category=Controls)
diff --git a/Engine/Source/Editor/UnrealEd/Classes/Settings/ProjectPackagingSettings.h b/Engine/Source/Editor/UnrealEd/Classes/Settings/ProjectPackagingSettings.h
--- a/Engine/Source/Editor/UnrealEd/Classes/Settings/ProjectPackagingSettings.h
+++ b/Engine/Source/Editor/UnrealEd/Classes/Settings/ProjectPackagingSettings.h
@@ -50,9 +50,7 @@ UCLASS(config=Game, defaultconfig)
 class UNREALED_API UProjectPackagingSettings
 	: public UObject
 {
-	GENERATED_BODY()
-public:
-	UProjectPackagingSettings(const FObjectInitializer& ObjectInitializer = FObjectInitializer::Get());
+	GENERATED_UCLASS_BODY()
 
 public:
 
diff --git a/Engine/Source/Editor/UnrealEd/Classes/SoundClassGraph/SoundClassGraph.h b/Engine/Source/Editor/UnrealEd/Classes/SoundClassGraph/SoundClassGraph.h
--- a/Engine/Source/Editor/UnrealEd/Classes/SoundClassGraph/SoundClassGraph.h
+++ b/Engine/Source/Editor/UnrealEd/Classes/SoundClassGraph/SoundClassGraph.h
@@ -7,9 +7,7 @@
 UCLASS(MinimalAPI)
 class USoundClassGraph : public UEdGraph
 {
-	GENERATED_BODY()
-public:
-	UNREALED_API USoundClassGraph(const FObjectInitializer& ObjectInitializer = FObjectInitializer::Get());
+	GENERATED_UCLASS_BODY()
 
 	/**
 	 * Set the SoundClass which forms the root of this graph
diff --git a/Engine/Source/Editor/UnrealEd/Classes/SoundClassGraph/SoundClassGraphNode.h b/Engine/Source/Editor/UnrealEd/Classes/SoundClassGraph/SoundClassGraphNode.h
--- a/Engine/Source/Editor/UnrealEd/Classes/SoundClassGraph/SoundClassGraphNode.h
+++ b/Engine/Source/Editor/UnrealEd/Classes/SoundClassGraph/SoundClassGraphNode.h
@@ -6,9 +6,7 @@
 UCLASS(MinimalAPI)
 class USoundClassGraphNode : public UEdGraphNode
 {
-	GENERATED_BODY()
-public:
-	UNREALED_API USoundClassGraphNode(const FObjectInitializer& ObjectInitializer = FObjectInitializer::Get());
+	GENERATED_UCLASS_BODY()
 
 	/** The SoundNode this represents */
 	UPROPERTY(VisibleAnywhere, instanced, Category=Sound)
diff --git a/Engine/Source/Editor/UnrealEd/Classes/SoundCueGraph/SoundCueGraph.h b/Engine/Source/Editor/UnrealEd/Classes/SoundCueGraph/SoundCueGraph.h
--- a/Engine/Source/Editor/UnrealEd/Classes/SoundCueGraph/SoundCueGraph.h
+++ b/Engine/Source/Editor/UnrealEd/Classes/SoundCueGraph/SoundCueGraph.h
@@ -7,9 +7,7 @@
 UCLASS(MinimalAPI)
 class USoundCueGraph : public UEdGraph
 {
-	GENERATED_BODY()
-public:
-	UNREALED_API USoundCueGraph(const FObjectInitializer& ObjectInitializer = FObjectInitializer::Get());
+	GENERATED_UCLASS_BODY()
 
 	/** Returns the SoundCue that contains this graph */
 	UNREALED_API class USoundCue* GetSoundCue() const;
diff --git a/Engine/Source/Editor/UnrealEd/Classes/SoundCueGraph/SoundCueGraphNode.h b/Engine/Source/Editor/UnrealEd/Classes/SoundCueGraph/SoundCueGraphNode.h
--- a/Engine/Source/Editor/UnrealEd/Classes/SoundCueGraph/SoundCueGraphNode.h
+++ b/Engine/Source/Editor/UnrealEd/Classes/SoundCueGraph/SoundCueGraphNode.h
@@ -8,9 +8,7 @@
 UCLASS(MinimalAPI)
 class USoundCueGraphNode : public USoundCueGraphNode_Base
 {
-	GENERATED_BODY()
-public:
-	UNREALED_API USoundCueGraphNode(const FObjectInitializer& ObjectInitializer = FObjectInitializer::Get());
+	GENERATED_UCLASS_BODY()
 
 	/** The SoundNode this represents */
 	UPROPERTY(VisibleAnywhere, instanced, Category=Sound)
diff --git a/Engine/Source/Editor/UnrealEd/Classes/SoundCueGraph/SoundCueGraphNode_Base.h b/Engine/Source/Editor/UnrealEd/Classes/SoundCueGraph/SoundCueGraphNode_Base.h
--- a/Engine/Source/Editor/UnrealEd/Classes/SoundCueGraph/SoundCueGraphNode_Base.h
+++ b/Engine/Source/Editor/UnrealEd/Classes/SoundCueGraph/SoundCueGraphNode_Base.h
@@ -7,9 +7,7 @@
 UCLASS(MinimalAPI)
 class USoundCueGraphNode_Base : public UEdGraphNode
 {
-	GENERATED_BODY()
-public:
-	UNREALED_API USoundCueGraphNode_Base(const FObjectInitializer& ObjectInitializer = FObjectInitializer::Get());
+	GENERATED_UCLASS_BODY()
 
 	/** Create all of the input pins required */
 	virtual void CreateInputPins() {};
diff --git a/Engine/Source/Editor/UnrealEd/Classes/SoundCueGraph/SoundCueGraphNode_Root.h b/Engine/Source/Editor/UnrealEd/Classes/SoundCueGraph/SoundCueGraphNode_Root.h
--- a/Engine/Source/Editor/UnrealEd/Classes/SoundCueGraph/SoundCueGraphNode_Root.h
+++ b/Engine/Source/Editor/UnrealEd/Classes/SoundCueGraph/SoundCueGraphNode_Root.h
@@ -6,9 +6,7 @@
 UCLASS(MinimalAPI)
 class USoundCueGraphNode_Root : public USoundCueGraphNode_Base
 {
-	GENERATED_BODY()
-public:
-	UNREALED_API USoundCueGraphNode_Root(const FObjectInitializer& ObjectInitializer = FObjectInitializer::Get());
+	GENERATED_UCLASS_BODY()
 
 	// UEdGraphNode interface
 	virtual FLinearColor GetNodeTitleColor() const override;
diff --git a/Engine/Source/Editor/UnrealEd/Classes/SoundCueGraph/SoundCueGraphSchema.h b/Engine/Source/Editor/UnrealEd/Classes/SoundCueGraph/SoundCueGraphSchema.h
--- a/Engine/Source/Editor/UnrealEd/Classes/SoundCueGraph/SoundCueGraphSchema.h
+++ b/Engine/Source/Editor/UnrealEd/Classes/SoundCueGraph/SoundCueGraphSchema.h
@@ -11,7 +11,7 @@ class USoundNodeWavePlayer;
 USTRUCT()
 struct UNREALED_API FSoundCueGraphSchemaAction_NewNode : public FEdGraphSchemaAction
 {
-	GENERATED_BODY()
+	GENERATED_USTRUCT_BODY();
 
 	/** Class of node we want to create */
 	UPROPERTY()
@@ -41,7 +41,7 @@ private:
 USTRUCT()
 struct UNREALED_API FSoundCueGraphSchemaAction_NewFromSelected : public FSoundCueGraphSchemaAction_NewNode
 {
-	GENERATED_BODY()
+	GENERATED_USTRUCT_BODY();
 
 	FSoundCueGraphSchemaAction_NewFromSelected() 
 		: FSoundCueGraphSchemaAction_NewNode()
@@ -60,7 +60,7 @@ struct UNREALED_API FSoundCueGraphSchemaAction_NewFromSelected : public FSoundCu
 USTRUCT()
 struct UNREALED_API FSoundCueGraphSchemaAction_NewComment : public FEdGraphSchemaAction
 {
-	GENERATED_BODY()
+	GENERATED_USTRUCT_BODY();
 
 	FSoundCueGraphSchemaAction_NewComment() 
 		: FEdGraphSchemaAction()
@@ -79,7 +79,7 @@ struct UNREALED_API FSoundCueGraphSchemaAction_NewComment : public FEdGraphSchem
 USTRUCT()
 struct UNREALED_API FSoundCueGraphSchemaAction_Paste : public FEdGraphSchemaAction
 {
-	GENERATED_BODY()
+	GENERATED_USTRUCT_BODY();
 
 	FSoundCueGraphSchemaAction_Paste() 
 		: FEdGraphSchemaAction()
@@ -97,9 +97,7 @@ struct UNREALED_API FSoundCueGraphSchemaAction_Paste : public FEdGraphSchemaActi
 UCLASS(MinimalAPI)
 class USoundCueGraphSchema : public UEdGraphSchema
 {
-	GENERATED_BODY()
-public:
-	UNREALED_API USoundCueGraphSchema(const FObjectInitializer& ObjectInitializer = FObjectInitializer::Get());
+	GENERATED_UCLASS_BODY()
 
 	/** Check whether connecting these pins would cause a loop */
 	bool ConnectionCausesLoop(const UEdGraphPin* InputPin, const UEdGraphPin* OutputPin) const;
diff --git a/Engine/Source/Editor/UnrealEd/Classes/TexAligner/TexAligner.h b/Engine/Source/Editor/UnrealEd/Classes/TexAligner/TexAligner.h
--- a/Engine/Source/Editor/UnrealEd/Classes/TexAligner/TexAligner.h
+++ b/Engine/Source/Editor/UnrealEd/Classes/TexAligner/TexAligner.h
@@ -38,9 +38,7 @@ enum ETexAlign
 UCLASS()
 class UTexAligner : public UObject
 {
-	GENERATED_BODY()
-public:
-	UTexAligner(const FObjectInitializer& ObjectInitializer = FObjectInitializer::Get());
+	GENERATED_UCLASS_BODY()
 
 	/** The default alignment this aligner represents. */
 	UPROPERTY()
diff --git a/Engine/Source/Editor/UnrealEd/Classes/TexAligner/TexAlignerBox.h b/Engine/Source/Editor/UnrealEd/Classes/TexAligner/TexAlignerBox.h
--- a/Engine/Source/Editor/UnrealEd/Classes/TexAligner/TexAlignerBox.h
+++ b/Engine/Source/Editor/UnrealEd/Classes/TexAligner/TexAlignerBox.h
@@ -12,9 +12,7 @@
 UCLASS(hidecategories=Object)
 class UTexAlignerBox : public UTexAligner
 {
-	GENERATED_BODY()
-public:
-	UTexAlignerBox(const FObjectInitializer& ObjectInitializer = FObjectInitializer::Get());
+	GENERATED_UCLASS_BODY()
 
 
 	// Begin UObject Interface
diff --git a/Engine/Source/Editor/UnrealEd/Classes/TexAligner/TexAlignerDefault.h b/Engine/Source/Editor/UnrealEd/Classes/TexAligner/TexAlignerDefault.h
--- a/Engine/Source/Editor/UnrealEd/Classes/TexAligner/TexAlignerDefault.h
+++ b/Engine/Source/Editor/UnrealEd/Classes/TexAligner/TexAlignerDefault.h
@@ -12,9 +12,7 @@
 UCLASS(hidecategories=Object)
 class UTexAlignerDefault : public UTexAligner
 {
-	GENERATED_BODY()
-public:
-	UTexAlignerDefault(const FObjectInitializer& ObjectInitializer = FObjectInitializer::Get());
+	GENERATED_UCLASS_BODY()
 
 
 	// Begin UObject Interface
diff --git a/Engine/Source/Editor/UnrealEd/Classes/TexAligner/TexAlignerFit.h b/Engine/Source/Editor/UnrealEd/Classes/TexAligner/TexAlignerFit.h
--- a/Engine/Source/Editor/UnrealEd/Classes/TexAligner/TexAlignerFit.h
+++ b/Engine/Source/Editor/UnrealEd/Classes/TexAligner/TexAlignerFit.h
@@ -12,9 +12,7 @@
 UCLASS(hidecategories=Object)
 class UTexAlignerFit : public UTexAligner
 {
-	GENERATED_BODY()
-public:
-	UTexAlignerFit(const FObjectInitializer& ObjectInitializer = FObjectInitializer::Get());
+	GENERATED_UCLASS_BODY()
 
 
 	// Begin UObject Interface
diff --git a/Engine/Source/Editor/UnrealEd/Classes/TexAligner/TexAlignerPlanar.h b/Engine/Source/Editor/UnrealEd/Classes/TexAligner/TexAlignerPlanar.h
--- a/Engine/Source/Editor/UnrealEd/Classes/TexAligner/TexAlignerPlanar.h
+++ b/Engine/Source/Editor/UnrealEd/Classes/TexAligner/TexAlignerPlanar.h
@@ -12,9 +12,7 @@
 UCLASS(hidecategories=Object)
 class UTexAlignerPlanar : public UTexAligner
 {
-	GENERATED_BODY()
-public:
-	UTexAlignerPlanar(const FObjectInitializer& ObjectInitializer = FObjectInitializer::Get());
+	GENERATED_UCLASS_BODY()
 
 
 	// Begin UObject Interface
diff --git a/Engine/Source/Editor/UnrealEd/Classes/ThumbnailRendering/AnimBlueprintThumbnailRenderer.h b/Engine/Source/Editor/UnrealEd/Classes/ThumbnailRendering/AnimBlueprintThumbnailRenderer.h
--- a/Engine/Source/Editor/UnrealEd/Classes/ThumbnailRendering/AnimBlueprintThumbnailRenderer.h
+++ b/Engine/Source/Editor/UnrealEd/Classes/ThumbnailRendering/AnimBlueprintThumbnailRenderer.h
@@ -11,9 +11,7 @@
 UCLASS(config=Editor, MinimalAPI)
 class UAnimBlueprintThumbnailRenderer : public UDefaultSizedThumbnailRenderer
 {
-	GENERATED_BODY()
-public:
-	UNREALED_API UAnimBlueprintThumbnailRenderer(const FObjectInitializer& ObjectInitializer = FObjectInitializer::Get());
+	GENERATED_UCLASS_BODY()
 
 
 	// Begin UThumbnailRenderer Object
diff --git a/Engine/Source/Editor/UnrealEd/Classes/ThumbnailRendering/BlueprintThumbnailRenderer.h b/Engine/Source/Editor/UnrealEd/Classes/ThumbnailRendering/BlueprintThumbnailRenderer.h
--- a/Engine/Source/Editor/UnrealEd/Classes/ThumbnailRendering/BlueprintThumbnailRenderer.h
+++ b/Engine/Source/Editor/UnrealEd/Classes/ThumbnailRendering/BlueprintThumbnailRenderer.h
@@ -11,9 +11,7 @@
 UCLASS(config=Editor,MinimalAPI)
 class UBlueprintThumbnailRenderer : public UDefaultSizedThumbnailRenderer
 {
-	GENERATED_BODY()
-public:
-	UNREALED_API UBlueprintThumbnailRenderer(const FObjectInitializer& ObjectInitializer = FObjectInitializer::Get());
+	GENERATED_UCLASS_BODY()
 
 
 	// Begin UThumbnailRenderer Object
diff --git a/Engine/Source/Editor/UnrealEd/Classes/ThumbnailRendering/DefaultSizedThumbnailRenderer.h b/Engine/Source/Editor/UnrealEd/Classes/ThumbnailRendering/DefaultSizedThumbnailRenderer.h
--- a/Engine/Source/Editor/UnrealEd/Classes/ThumbnailRendering/DefaultSizedThumbnailRenderer.h
+++ b/Engine/Source/Editor/UnrealEd/Classes/ThumbnailRendering/DefaultSizedThumbnailRenderer.h
@@ -11,9 +11,7 @@
 UCLASS(abstract, config=Editor, MinimalAPI)
 class UDefaultSizedThumbnailRenderer : public UThumbnailRenderer
 {
-	GENERATED_BODY()
-public:
-	UNREALED_API UDefaultSizedThumbnailRenderer(const FObjectInitializer& ObjectInitializer = FObjectInitializer::Get());
+	GENERATED_UCLASS_BODY()
 
 	/**
 	 * The default width of this thumbnail
diff --git a/Engine/Source/Editor/UnrealEd/Classes/ThumbnailRendering/FontThumbnailRenderer.h b/Engine/Source/Editor/UnrealEd/Classes/ThumbnailRendering/FontThumbnailRenderer.h
--- a/Engine/Source/Editor/UnrealEd/Classes/ThumbnailRendering/FontThumbnailRenderer.h
+++ b/Engine/Source/Editor/UnrealEd/Classes/ThumbnailRendering/FontThumbnailRenderer.h
@@ -11,9 +11,7 @@
 UCLASS()
 class UFontThumbnailRenderer : public UTextureThumbnailRenderer
 {
-	GENERATED_BODY()
-public:
-	UFontThumbnailRenderer(const FObjectInitializer& ObjectInitializer = FObjectInitializer::Get());
+	GENERATED_UCLASS_BODY()
 
 
 	// Begin UThumbnailRenderer Object
diff --git a/Engine/Source/Editor/UnrealEd/Classes/ThumbnailRendering/LevelThumbnailRenderer.h b/Engine/Source/Editor/UnrealEd/Classes/ThumbnailRendering/LevelThumbnailRenderer.h
--- a/Engine/Source/Editor/UnrealEd/Classes/ThumbnailRendering/LevelThumbnailRenderer.h
+++ b/Engine/Source/Editor/UnrealEd/Classes/ThumbnailRendering/LevelThumbnailRenderer.h
@@ -11,9 +11,7 @@
 UCLASS(config=Editor)
 class ULevelThumbnailRenderer : public UDefaultSizedThumbnailRenderer
 {
-	GENERATED_BODY()
-public:
-	ULevelThumbnailRenderer(const FObjectInitializer& ObjectInitializer = FObjectInitializer::Get());
+	GENERATED_UCLASS_BODY()
 
 
 	// Begin UThumbnailRenderer Object
diff --git a/Engine/Source/Editor/UnrealEd/Classes/ThumbnailRendering/MaterialFunctionThumbnailRenderer.h b/Engine/Source/Editor/UnrealEd/Classes/ThumbnailRendering/MaterialFunctionThumbnailRenderer.h
--- a/Engine/Source/Editor/UnrealEd/Classes/ThumbnailRendering/MaterialFunctionThumbnailRenderer.h
+++ b/Engine/Source/Editor/UnrealEd/Classes/ThumbnailRendering/MaterialFunctionThumbnailRenderer.h
@@ -11,9 +11,7 @@
 UCLASS(config=Editor)
 class UMaterialFunctionThumbnailRenderer : public UDefaultSizedThumbnailRenderer
 {
-	GENERATED_BODY()
-public:
-	UMaterialFunctionThumbnailRenderer(const FObjectInitializer& ObjectInitializer = FObjectInitializer::Get());
+	GENERATED_UCLASS_BODY()
 
 
 	// Begin UThumbnailRenderer Object
diff --git a/Engine/Source/Editor/UnrealEd/Classes/ThumbnailRendering/MaterialInstanceThumbnailRenderer.h b/Engine/Source/Editor/UnrealEd/Classes/ThumbnailRendering/MaterialInstanceThumbnailRenderer.h
--- a/Engine/Source/Editor/UnrealEd/Classes/ThumbnailRendering/MaterialInstanceThumbnailRenderer.h
+++ b/Engine/Source/Editor/UnrealEd/Classes/ThumbnailRendering/MaterialInstanceThumbnailRenderer.h
@@ -11,9 +11,7 @@
 UCLASS(config=Editor)
 class UMaterialInstanceThumbnailRenderer : public UDefaultSizedThumbnailRenderer
 {
-	GENERATED_BODY()
-public:
-	UMaterialInstanceThumbnailRenderer(const FObjectInitializer& ObjectInitializer = FObjectInitializer::Get());
+	GENERATED_UCLASS_BODY()
 
 
 	// Begin UThumbnailRenderer Object
diff --git a/Engine/Source/Editor/UnrealEd/Classes/ThumbnailRendering/ParticleSystemThumbnailRenderer.h b/Engine/Source/Editor/UnrealEd/Classes/ThumbnailRendering/ParticleSystemThumbnailRenderer.h
--- a/Engine/Source/Editor/UnrealEd/Classes/ThumbnailRendering/ParticleSystemThumbnailRenderer.h
+++ b/Engine/Source/Editor/UnrealEd/Classes/ThumbnailRendering/ParticleSystemThumbnailRenderer.h
@@ -11,9 +11,7 @@
 UCLASS(config=Editor)
 class UParticleSystemThumbnailRenderer : public UTextureThumbnailRenderer
 {
-	GENERATED_BODY()
-public:
-	UParticleSystemThumbnailRenderer(const FObjectInitializer& ObjectInitializer = FObjectInitializer::Get());
+	GENERATED_UCLASS_BODY()
 
 	UPROPERTY()
 	class UTexture2D* NoImage;
diff --git a/Engine/Source/Editor/UnrealEd/Classes/ThumbnailRendering/SceneThumbnailInfo.h b/Engine/Source/Editor/UnrealEd/Classes/ThumbnailRendering/SceneThumbnailInfo.h
--- a/Engine/Source/Editor/UnrealEd/Classes/ThumbnailRendering/SceneThumbnailInfo.h
+++ b/Engine/Source/Editor/UnrealEd/Classes/ThumbnailRendering/SceneThumbnailInfo.h
@@ -13,9 +13,7 @@
 UCLASS(MinimalAPI)
 class USceneThumbnailInfo : public UThumbnailInfo
 {
-	GENERATED_BODY()
-public:
-	UNREALED_API USceneThumbnailInfo(const FObjectInitializer& ObjectInitializer = FObjectInitializer::Get());
+	GENERATED_UCLASS_BODY()
 
 	/** The pitch of the orbit camera around the asset */
 	UPROPERTY(EditAnywhere, Category=Thumbnail)
diff --git a/Engine/Source/Editor/UnrealEd/Classes/ThumbnailRendering/SkeletalMeshThumbnailRenderer.h b/Engine/Source/Editor/UnrealEd/Classes/ThumbnailRendering/SkeletalMeshThumbnailRenderer.h
--- a/Engine/Source/Editor/UnrealEd/Classes/ThumbnailRendering/SkeletalMeshThumbnailRenderer.h
+++ b/Engine/Source/Editor/UnrealEd/Classes/ThumbnailRendering/SkeletalMeshThumbnailRenderer.h
@@ -11,9 +11,7 @@
 UCLASS(config=Editor, MinimalAPI)
 class USkeletalMeshThumbnailRenderer : public UDefaultSizedThumbnailRenderer
 {
-	GENERATED_BODY()
-public:
-	UNREALED_API USkeletalMeshThumbnailRenderer(const FObjectInitializer& ObjectInitializer = FObjectInitializer::Get());
+	GENERATED_UCLASS_BODY()
 
 
 	// Begin UThumbnailRenderer Object
diff --git a/Engine/Source/Editor/UnrealEd/Classes/ThumbnailRendering/SlateBrushThumbnailRenderer.h b/Engine/Source/Editor/UnrealEd/Classes/ThumbnailRendering/SlateBrushThumbnailRenderer.h
--- a/Engine/Source/Editor/UnrealEd/Classes/ThumbnailRendering/SlateBrushThumbnailRenderer.h
+++ b/Engine/Source/Editor/UnrealEd/Classes/ThumbnailRendering/SlateBrushThumbnailRenderer.h
@@ -11,9 +11,7 @@
 UCLASS(config=Editor,MinimalAPI)
 class USlateBrushThumbnailRenderer : public UDefaultSizedThumbnailRenderer
 {
-	GENERATED_BODY()
-public:
-	UNREALED_API USlateBrushThumbnailRenderer(const FObjectInitializer& ObjectInitializer = FObjectInitializer::Get());
+	GENERATED_UCLASS_BODY()
 
 public:
 
diff --git a/Engine/Source/Editor/UnrealEd/Classes/ThumbnailRendering/SoundWaveThumbnailRenderer.h b/Engine/Source/Editor/UnrealEd/Classes/ThumbnailRendering/SoundWaveThumbnailRenderer.h
--- a/Engine/Source/Editor/UnrealEd/Classes/ThumbnailRendering/SoundWaveThumbnailRenderer.h
+++ b/Engine/Source/Editor/UnrealEd/Classes/ThumbnailRendering/SoundWaveThumbnailRenderer.h
@@ -11,9 +11,7 @@
 UCLASS()
 class USoundWaveThumbnailRenderer : public UThumbnailRenderer
 {
-	GENERATED_BODY()
-public:
-	USoundWaveThumbnailRenderer(const FObjectInitializer& ObjectInitializer = FObjectInitializer::Get());
+	GENERATED_UCLASS_BODY()
 
 	// Begin UThumbnailRenderer Object
 	virtual void Draw(UObject* Object, int32 X, int32 Y, uint32 Width, uint32 Height, FRenderTarget*, FCanvas* Canvas) override;
diff --git a/Engine/Source/Editor/UnrealEd/Classes/ThumbnailRendering/StaticMeshThumbnailRenderer.h b/Engine/Source/Editor/UnrealEd/Classes/ThumbnailRendering/StaticMeshThumbnailRenderer.h
--- a/Engine/Source/Editor/UnrealEd/Classes/ThumbnailRendering/StaticMeshThumbnailRenderer.h
+++ b/Engine/Source/Editor/UnrealEd/Classes/ThumbnailRendering/StaticMeshThumbnailRenderer.h
@@ -11,9 +11,7 @@
 UCLASS(config=Editor,MinimalAPI)
 class UStaticMeshThumbnailRenderer : public UDefaultSizedThumbnailRenderer
 {
-	GENERATED_BODY()
-public:
-	UNREALED_API UStaticMeshThumbnailRenderer(const FObjectInitializer& ObjectInitializer = FObjectInitializer::Get());
+	GENERATED_UCLASS_BODY()
 
 
 	// Begin UThumbnailRenderer Object
diff --git a/Engine/Source/Editor/UnrealEd/Classes/ThumbnailRendering/SubsurfaceProfileRenderer.h b/Engine/Source/Editor/UnrealEd/Classes/ThumbnailRendering/SubsurfaceProfileRenderer.h
--- a/Engine/Source/Editor/UnrealEd/Classes/ThumbnailRendering/SubsurfaceProfileRenderer.h
+++ b/Engine/Source/Editor/UnrealEd/Classes/ThumbnailRendering/SubsurfaceProfileRenderer.h
@@ -11,9 +11,7 @@
 UCLASS()
 class USubsurfaceProfileRenderer : public UTextureThumbnailRenderer
 {
-	GENERATED_BODY()
-public:
-	USubsurfaceProfileRenderer(const FObjectInitializer& ObjectInitializer = FObjectInitializer::Get());
+	GENERATED_UCLASS_BODY()
 
 
 	// Begin UThumbnailRenderer Object
diff --git a/Engine/Source/Editor/UnrealEd/Classes/ThumbnailRendering/TextureThumbnailRenderer.h b/Engine/Source/Editor/UnrealEd/Classes/ThumbnailRendering/TextureThumbnailRenderer.h
--- a/Engine/Source/Editor/UnrealEd/Classes/ThumbnailRendering/TextureThumbnailRenderer.h
+++ b/Engine/Source/Editor/UnrealEd/Classes/ThumbnailRendering/TextureThumbnailRenderer.h
@@ -11,9 +11,7 @@
 UCLASS(MinimalAPI)
 class UTextureThumbnailRenderer : public UThumbnailRenderer
 {
-	GENERATED_BODY()
-public:
-	UNREALED_API UTextureThumbnailRenderer(const FObjectInitializer& ObjectInitializer = FObjectInitializer::Get());
+	GENERATED_UCLASS_BODY()
 
 
 	// Begin UThumbnailRenderer Object
diff --git a/Engine/Source/Editor/UnrealEd/Classes/ThumbnailRendering/ThumbnailManager.h b/Engine/Source/Editor/UnrealEd/Classes/ThumbnailRendering/ThumbnailManager.h
--- a/Engine/Source/Editor/UnrealEd/Classes/ThumbnailRendering/ThumbnailManager.h
+++ b/Engine/Source/Editor/UnrealEd/Classes/ThumbnailRendering/ThumbnailManager.h
@@ -33,7 +33,7 @@ enum EThumbnailPrimType
 USTRUCT()
 struct FThumbnailRenderingInfo
 {
-	GENERATED_BODY()
+	GENERATED_USTRUCT_BODY()
 
 	/**
 	 * The name of the class that this thumbnail is for (so we can lazy bind)
@@ -73,9 +73,7 @@ public:
 UCLASS(config=Editor)
 class UThumbnailManager : public UObject
 {
-	GENERATED_BODY()
-public:
-	UThumbnailManager(const FObjectInitializer& ObjectInitializer = FObjectInitializer::Get());
+	GENERATED_UCLASS_BODY()
 
 protected:
 	/**
diff --git a/Engine/Source/Editor/UnrealEd/Classes/ThumbnailRendering/ThumbnailRenderer.h b/Engine/Source/Editor/UnrealEd/Classes/ThumbnailRendering/ThumbnailRenderer.h
--- a/Engine/Source/Editor/UnrealEd/Classes/ThumbnailRendering/ThumbnailRenderer.h
+++ b/Engine/Source/Editor/UnrealEd/Classes/ThumbnailRendering/ThumbnailRenderer.h
@@ -13,9 +13,7 @@
 UCLASS(abstract, MinimalAPI)
 class UThumbnailRenderer : public UObject
 {
-	GENERATED_BODY()
-public:
-	UNREALED_API UThumbnailRenderer(const FObjectInitializer& ObjectInitializer = FObjectInitializer::Get());
+	GENERATED_UCLASS_BODY()
 
 
 public:
diff --git a/Engine/Source/Editor/UnrealEd/Classes/ThumbnailRendering/WorldThumbnailInfo.h b/Engine/Source/Editor/UnrealEd/Classes/ThumbnailRendering/WorldThumbnailInfo.h
--- a/Engine/Source/Editor/UnrealEd/Classes/ThumbnailRendering/WorldThumbnailInfo.h
+++ b/Engine/Source/Editor/UnrealEd/Classes/ThumbnailRendering/WorldThumbnailInfo.h
@@ -26,9 +26,7 @@ namespace EOrthoThumbnailDirection
 UCLASS(MinimalAPI)
 class UWorldThumbnailInfo : public USceneThumbnailInfo
 {
-	GENERATED_BODY()
-public:
-	UNREALED_API UWorldThumbnailInfo(const FObjectInitializer& ObjectInitializer = FObjectInitializer::Get());
+	GENERATED_UCLASS_BODY()
 
 	/** The type of projection to use */
 	UPROPERTY(EditAnywhere, Category=Thumbnail)
diff --git a/Engine/Source/Editor/UnrealEd/Classes/ThumbnailRendering/WorldThumbnailRenderer.h b/Engine/Source/Editor/UnrealEd/Classes/ThumbnailRendering/WorldThumbnailRenderer.h
--- a/Engine/Source/Editor/UnrealEd/Classes/ThumbnailRendering/WorldThumbnailRenderer.h
+++ b/Engine/Source/Editor/UnrealEd/Classes/ThumbnailRendering/WorldThumbnailRenderer.h
@@ -11,9 +11,7 @@
 UCLASS(config=Editor)
 class UWorldThumbnailRenderer : public UDefaultSizedThumbnailRenderer
 {
-	GENERATED_BODY()
-public:
-	UWorldThumbnailRenderer(const FObjectInitializer& ObjectInitializer = FObjectInitializer::Get());
+	GENERATED_UCLASS_BODY()
 
 	// Begin UThumbnailRenderer Object
 	virtual bool CanVisualizeAsset(UObject* Object) override;
diff --git a/Engine/Source/Editor/UnrealEd/Classes/UserDefinedStructure/UserDefinedStructEditorData.h b/Engine/Source/Editor/UnrealEd/Classes/UserDefinedStructure/UserDefinedStructEditorData.h
--- a/Engine/Source/Editor/UnrealEd/Classes/UserDefinedStructure/UserDefinedStructEditorData.h
+++ b/Engine/Source/Editor/UnrealEd/Classes/UserDefinedStructure/UserDefinedStructEditorData.h
@@ -7,7 +7,7 @@
 USTRUCT()
 struct FStructVariableDescription
 {
-	GENERATED_BODY()
+	GENERATED_USTRUCT_BODY()
 
 	UPROPERTY()
 	FName VarName;
@@ -65,9 +65,7 @@ struct FStructVariableDescription
 UCLASS()
 class UNREALED_API UUserDefinedStructEditorData : public UObject
 {
-	GENERATED_BODY()
-public:
-	UUserDefinedStructEditorData(const FObjectInitializer& ObjectInitializer = FObjectInitializer::Get());
+	GENERATED_UCLASS_BODY()
 
 private:
 	// the property is used to generate an uniqe name id for member variable
@@ -91,4 +89,4 @@ public:
 
 	uint32 GenerateUniqueNameIdForMemberVariable();
 	class UUserDefinedStruct* GetOwnerStruct() const;
-};
+};
\ No newline at end of file
diff --git a/Engine/Source/Editor/UnrealEd/Public/AutoReimport/AutoReimportManager.h b/Engine/Source/Editor/UnrealEd/Public/AutoReimport/AutoReimportManager.h
--- a/Engine/Source/Editor/UnrealEd/Public/AutoReimport/AutoReimportManager.h
+++ b/Engine/Source/Editor/UnrealEd/Public/AutoReimport/AutoReimportManager.h
@@ -22,9 +22,7 @@ struct FPathAndMountPoint
 UCLASS(config=Editor, transient)
 class UNREALED_API UAutoReimportManager : public UObject
 {
-	GENERATED_BODY()
-public:
-	UAutoReimportManager(const FObjectInitializer& ObjectInitializer = FObjectInitializer::Get());
+	GENERATED_UCLASS_BODY()
 public:
 
 	~UAutoReimportManager();
diff --git a/Engine/Source/Editor/UnrealEd/Public/CookerSettings.h b/Engine/Source/Editor/UnrealEd/Public/CookerSettings.h
--- a/Engine/Source/Editor/UnrealEd/Public/CookerSettings.h
+++ b/Engine/Source/Editor/UnrealEd/Public/CookerSettings.h
@@ -15,9 +15,7 @@ UCLASS(config=Engine, defaultconfig)
 class UNREALED_API UCookerSettings
 	: public UObject
 {
-	GENERATED_BODY()
-public:
-	UCookerSettings(const FObjectInitializer& ObjectInitializer = FObjectInitializer::Get());
+	GENERATED_UCLASS_BODY()
 
 public:
 
diff --git a/Engine/Source/Editor/UnrealEd/Public/EdGraphNode_Comment.h b/Engine/Source/Editor/UnrealEd/Public/EdGraphNode_Comment.h
--- a/Engine/Source/Editor/UnrealEd/Public/EdGraphNode_Comment.h
+++ b/Engine/Source/Editor/UnrealEd/Public/EdGraphNode_Comment.h
@@ -23,9 +23,11 @@ namespace ECommentBoxMode
 UCLASS(MinimalAPI)
 class UEdGraphNode_Comment : public UEdGraphNode
 {
+public:
 	GENERATED_BODY()
+
 public:
-	UNREALED_API UEdGraphNode_Comment(const FObjectInitializer& ObjectInitializer = FObjectInitializer::Get());
+	UEdGraphNode_Comment(const FObjectInitializer& ObjectInitializer = FObjectInitializer::Get());
 
 	/** Color to style comment with */
 	UPROPERTY(EditAnywhere, Category=Comment)
diff --git a/Engine/Source/Programs/UnrealHeaderTool/Private/HeaderParser.cpp b/Engine/Source/Programs/UnrealHeaderTool/Private/HeaderParser.cpp
--- a/Engine/Source/Programs/UnrealHeaderTool/Private/HeaderParser.cpp
+++ b/Engine/Source/Programs/UnrealHeaderTool/Private/HeaderParser.cpp
@@ -1432,17 +1432,17 @@ UScriptStruct* FHeaderParser::CompileStructDeclaration(FClasses& AllClasses, FUn
 		{
 			FError::Throwf(TEXT("USTRUCTs cannot contain UFUNCTIONs."));
 		}
-		else if (Token.Matches(TEXT("GENERATED_USTRUCT_BODY")) || Token.Matches(TEXT("GENERATED_BODY")))
+		else if (Token.Matches(TEXT("GENERATED_USTRUCT_BODY")))
 		{
-			// Match 'GENERATED_USTRUCT_BODY' '(' [StructName] ')' or 'GENERATED_BODY' '(' [StructName] ')'
+			// Match 'GENERATED_USTRUCT_BODY' '(' [StructName] ')'
 			if (CurrentAccessSpecifier != ACCESS_Public)
 			{
-				FError::Throwf(TEXT("%s must be in the public scope of '%s', not private or protected."), Token.Identifier, *StructNameInScript);
+				FError::Throwf(TEXT("GENERATED_USTRUCT_BODY must be in the public scope of '%s', not private or protected."), *StructNameInScript);
 			}
 
 			if (Struct->StructMacroDeclaredLineNumber != INDEX_NONE)
 			{
-				FError::Throwf(TEXT("Multiple %s declarations found in '%s'"), Token.Identifier, *StructNameInScript);
+				FError::Throwf(TEXT("Multiple GENERATED_USTRUCT_BODY declarations found in '%s'"), *StructNameInScript);
 			}
 
 			Struct->StructMacroDeclaredLineNumber = InputLine;
@@ -1453,7 +1453,7 @@ UScriptStruct* FHeaderParser::CompileStructDeclaration(FClasses& AllClasses, FUn
 			{
 				if (!DuplicateStructName.Matches(*StructNameInScript))
 				{
-					FError::Throwf(TEXT("The argument to %s must match the struct name '%s' if present.  However, the argument can be omitted entirely."), Token.Identifier, *StructNameInScript);
+					FError::Throwf(TEXT("The argument to GENERATED_USTRUCT_BODY must match the struct name '%s' if present.  However, the argument can be omitted entirely."), *StructNameInScript);
 				}
 			}
 
@@ -1567,7 +1567,7 @@ UScriptStruct* FHeaderParser::CompileStructDeclaration(FClasses& AllClasses, FUn
 	{
 		// Roll the line number back to the start of the struct body and error out
 		InputLine = SavedLineNumber;
-		FError::Throwf(TEXT("Expected a GENERATED_BODY() at the start of struct"));
+		FError::Throwf(TEXT("Expected a GENERATED_USTRUCT_BODY() at the start of struct"));
 	}
 
 	// Link the properties within the struct
@@ -1625,7 +1625,6 @@ const TCHAR *FHeaderParser::NestTypeName( ENestType NestType )
 			return TEXT("Global Scope");
 		case NEST_Class:
 			return TEXT("Class");
-		case NEST_NativeInterface:
 		case NEST_Interface:
 			return TEXT("Interface");
 		case NEST_FunctionDeclaration:
@@ -1706,7 +1705,6 @@ void FHeaderParser::PushNest(ENestType NestType, UStruct* InNode, FUnrealSourceF
 		TopNest->Allow = ALLOW_VarDecl | ALLOW_Function | ALLOW_TypeDecl;
 		break;
 
-	case NEST_NativeInterface:
 	case NEST_Interface:
 		TopNest->Allow = ALLOW_Function | ALLOW_TypeDecl;
 		break;
@@ -1740,7 +1738,7 @@ void FHeaderParser::PopNest(ENestType NestType, const TCHAR* Descr)
 		FError::Throwf(TEXT("Unexpected end of %s in '%s' block"), Descr, NestTypeName(TopNest->NestType));
 	}
 
-	if (NestType != NEST_GlobalScope && NestType != NEST_Class && NestType != NEST_Interface && NestType != NEST_NativeInterface && NestType != NEST_FunctionDeclaration)
+	if (NestType != NEST_GlobalScope && NestType != NEST_Class && NestType != NEST_Interface && NestType != NEST_FunctionDeclaration)
 	{
 		FError::Throwf(TEXT("Bad first pass NestType %i"), (uint8)NestType);
 	}
@@ -3629,7 +3627,7 @@ bool FHeaderParser::CompileDeclaration(FClasses& AllClasses, FUnrealSourceFile&
 		{
 			FError::Throwf(TEXT("Access specifier %s not allowed here."), Token.Identifier);
 		}
-		check(TopNest->NestType == NEST_Class || TopNest->NestType == NEST_Interface || TopNest->NestType == NEST_NativeInterface);
+		check(TopNest->NestType == NEST_Class || TopNest->NestType == NEST_Interface);
 		CurrentAccessSpecifier = AccessSpecifier;
 	}
 	else if (Token.Matches(TEXT("class")) && (TopNest->NestType == NEST_GlobalScope))
@@ -3642,7 +3640,6 @@ bool FHeaderParser::CompileDeclaration(FClasses& AllClasses, FUnrealSourceFile&
 
 		// Start parsing the second class
 		bEncounteredNewStyleClass_UnmatchedBrackets = true;
-		CurrentAccessSpecifier = ACCESS_Private;
 
 		if (!TryParseIInterfaceClass(AllClasses))
 		{
@@ -3651,79 +3648,75 @@ bool FHeaderParser::CompileDeclaration(FClasses& AllClasses, FUnrealSourceFile&
 			return SkipDeclaration(Token);
 		}
 	}
-	else if (Token.Matches(TEXT("GENERATED_IINTERFACE_BODY")) || (Token.Matches(TEXT("GENERATED_BODY")) && TopNest->NestType == NEST_NativeInterface))
+	else if (Token.Matches(TEXT("GENERATED_BODY")))
 	{
-		if (TopNest->NestType != NEST_NativeInterface)
+		if (TopNest->NestType != NEST_Class && TopNest->NestType != NEST_Interface)
 		{
-			FError::Throwf(TEXT("%s must occur inside the native interface definition"), Token.Identifier);
+			FError::Throwf(TEXT("%s must occur inside the class or interface definition"), Token.Identifier);
 		}
+		if (!ClassDefinitionRanges.Contains(GetCurrentClass()))
+		{
+			ClassDefinitionRanges.Add(GetCurrentClass(), ClassDefinitionRange());
+		}
+
+		ClassDefinitionRanges[GetCurrentClass()].bHasGeneratedBody = true;
+
 		RequireSymbol(TEXT("("), Token.Identifier);
 		RequireSymbol(TEXT(")"), Token.Identifier);
 
 		auto* ClassData = GetCurrentClassData();
 
 		ClassData->GeneratedBodyMacroAccessSpecifier = CurrentAccessSpecifier;
-		ClassData->SetInterfaceGeneratedBodyLine(InputLine);
-
-		bClassHasGeneratedIInterfaceBody = true;
+		ClassData->SetGeneratedBodyLine(InputLine);
 
-		if (Token.Matches(TEXT("GENERATED_IINTERFACE_BODY")))
+		if (TopNest->NestType == NEST_Class)
 		{
-			CurrentAccessSpecifier = ACCESS_Public;
+			bClassHasGeneratedBody = true;
 		}
 	}
-	else if (Token.Matches(TEXT("GENERATED_UINTERFACE_BODY")) || (Token.Matches(TEXT("GENERATED_BODY")) && TopNest->NestType == NEST_Interface))
+	else if (Token.Matches(TEXT("GENERATED_UCLASS_BODY")))
 	{
-		if (TopNest->NestType != NEST_Interface)
+		if (TopNest->NestType != NEST_Class)
 		{
-			FError::Throwf(TEXT("%s must occur inside the interface definition"), Token.Identifier);
+			FError::Throwf(TEXT("%s must occur inside the class definition"), Token.Identifier);
 		}
 		RequireSymbol(TEXT("("), Token.Identifier);
 		RequireSymbol(TEXT(")"), Token.Identifier);
 
-		auto* ClassData = GetCurrentClassData();
-
-		ClassData->GeneratedBodyMacroAccessSpecifier = CurrentAccessSpecifier;
-		ClassData->SetGeneratedBodyLine(InputLine);
+		GetCurrentClassData()->SetGeneratedBodyLine(InputLine);
 
-		bClassHasGeneratedUInterfaceBody = true;
+		// The body implementation macro always ends with a 'public:'
+		CurrentAccessSpecifier = ACCESS_Public;
 
-		if (Token.Matches(TEXT("GENERATED_UINTERFACE_BODY")))
-		{
-			CurrentAccessSpecifier = ACCESS_Public;
-		}
+		bClassHasGeneratedBody = true;
 	}
-	else if (Token.Matches(TEXT("GENERATED_UCLASS_BODY")) || (Token.Matches(TEXT("GENERATED_BODY")) && TopNest->NestType == NEST_Class))
+	else if (Token.Matches(TEXT("GENERATED_IINTERFACE_BODY")))
 	{
-		if (TopNest->NestType != NEST_Class)
-		{
-			FError::Throwf(TEXT("%s must occur inside the class definition"), Token.Identifier);
-		}
+		RequireSymbol(TEXT("("), Token.Identifier);
+		RequireSymbol(TEXT(")"), Token.Identifier);
 
-		auto* ClassData = GetCurrentClassData();
+		GetCurrentClassData()->SetInterfaceGeneratedBodyLine(InputLine);
 
-		if (Token.Matches(TEXT("GENERATED_BODY")))
-		{
-			if (!ClassDefinitionRanges.Contains(GetCurrentClass()))
-			{
-				ClassDefinitionRanges.Add(GetCurrentClass(), ClassDefinitionRange());
-			}
-
-			ClassDefinitionRanges[GetCurrentClass()].bHasGeneratedBody = true;
+		bClassHasGeneratedIInterfaceBody = true;
 
-			ClassData->GeneratedBodyMacroAccessSpecifier = CurrentAccessSpecifier;
-		}
-		else
+		// The body implementation macro always ends with a 'public:'
+		CurrentAccessSpecifier = ACCESS_Public;
+	}
+	else if (Token.Matches(TEXT("GENERATED_UINTERFACE_BODY")))
+	{
+		if (TopNest->NestType != NEST_Interface)
 		{
-			CurrentAccessSpecifier = ACCESS_Public;
+			FError::Throwf(TEXT("%s must occur inside the interface definition"), Token.Identifier);
 		}
-
 		RequireSymbol(TEXT("("), Token.Identifier);
 		RequireSymbol(TEXT(")"), Token.Identifier);
 
-		ClassData->SetGeneratedBodyLine(InputLine);
+		GetCurrentClassData()->SetGeneratedBodyLine(InputLine);
 
-		bClassHasGeneratedBody = true;
+		bClassHasGeneratedUInterfaceBody = true;
+
+		// The body implementation macro always ends with a 'public:'
+		CurrentAccessSpecifier = ACCESS_Public;
 	}
 	else if (Token.Matches(TEXT("UCLASS"), ESearchCase::CaseSensitive) && (TopNest->Allow & ALLOW_Class))
 	{
@@ -3789,8 +3782,7 @@ bool FHeaderParser::CompileDeclaration(FClasses& AllClasses, FUnrealSourceFile&
 		// Pop nesting here to allow other non UClass declarations in the header file.
 		if (CurrentClass->ClassFlags & CLASS_Interface)
 		{
-			checkf(TopNest->NestType == NEST_Interface || TopNest->NestType == NEST_NativeInterface, TEXT("Unexpected end of interface block."));
-			PopNest(TopNest->NestType, TEXT("'Interface'"));
+			PopNest(NEST_Interface, TEXT("'Interface'"));
 			PostPopNestInterface(AllClasses, CurrentClass);
 
 			// Ensure the UINTERFACE classes have a GENERATED_UINTERFACE_BODY declaration
@@ -3835,11 +3827,7 @@ bool FHeaderParser::CompileDeclaration(FClasses& AllClasses, FUnrealSourceFile&
 			FError::Throwf(TEXT("Extra ';' before end of file"));
 		}
 	}
-	else if (bEncounteredNewStyleClass_UnmatchedBrackets && IsInAClass() && GetCurrentClass() &&
-		(
-			Token.Matches(NameLookupCPP.GetNameCPP(GetCurrentClass())) || 
-			FString(Token.Identifier).EndsWith("_API") && GetToken(Token) && Token.Matches(NameLookupCPP.GetNameCPP(GetCurrentClass()))
-		))
+	else if (bEncounteredNewStyleClass_UnmatchedBrackets && IsInAClass() && GetCurrentClass() && Token.Matches(NameLookupCPP.GetNameCPP(GetCurrentClass())))
 	{
 		if (!TryToMatchConstructorParameterList(Token))
 		{
@@ -4702,7 +4690,7 @@ bool FHeaderParser::TryParseIInterfaceClass(FClasses& AllClasses)
 	RequireSymbol(TEXT("{"), *ErrorMsg);
 
 	// Push the interface class nesting again.
-	PushNest(NEST_NativeInterface, FoundClass);
+	PushNest(NEST_Interface, FoundClass);
 
 	return true;
 }
@@ -4731,7 +4719,7 @@ void FHeaderParser::CompileInterfaceDeclaration(FClasses& AllClasses)
 	// New style files have the interface name / extends afterwards
 	RequireIdentifier(TEXT("class"), TEXT("Interface declaration"));
 	FClass* InterfaceClass = ParseInterfaceNameDeclaration(AllClasses, /*out*/ DeclaredInterfaceName, /*out*/ RequiredAPIMacroIfPresent);
-	ClassDefinitionRanges.Add(InterfaceClass, ClassDefinitionRange(&Input[InputPos], nullptr));
+
 
 	// Record that this interface is RequiredAPI if the CORE_API style macro was present
 	if (!RequiredAPIMacroIfPresent.IsEmpty())
@@ -5049,7 +5037,7 @@ void FHeaderParser::CompileDelegateDeclaration(FUnrealSourceFile& SourceFile, FC
 		FCString::Strcpy( FuncInfo.Function.Identifier, *Name );
 	}
 
-	FuncInfo.MacroLine = InputLine;
+	FuncInfo.DelegateMacroLine = InputLine;
 	auto* DelegateSignatureFunction = CreateDelegateFunction(FuncInfo);
 
 	GScriptHelper.AddClassData(DelegateSignatureFunction);
@@ -5156,7 +5144,6 @@ void FHeaderParser::CompileFunctionDeclaration(FUnrealSourceFile& SourceFile, FC
 	bool bAutomaticallyFinal = true;
 
 	FFuncInfo FuncInfo;
-	FuncInfo.MacroLine = InputLine;
 	FuncInfo.FunctionFlags = FUNC_Native;
 
 	// Infer the function's access level from the currently declared C++ access level
diff --git a/Engine/Source/Programs/UnrealHeaderTool/Private/HeaderParser.h b/Engine/Source/Programs/UnrealHeaderTool/Private/HeaderParser.h
--- a/Engine/Source/Programs/UnrealHeaderTool/Private/HeaderParser.h
+++ b/Engine/Source/Programs/UnrealHeaderTool/Private/HeaderParser.h
@@ -22,8 +22,7 @@ enum ENestType
 	NEST_GlobalScope,
 	NEST_Class,
 	NEST_FunctionDeclaration,
-	NEST_Interface,
-	NEST_NativeInterface
+	NEST_Interface
 };
 
 /** Types of statements to allow within a particular nesting block. */
@@ -296,7 +295,7 @@ protected:
 	 */
 	FStructScope* GetCurrentClassScope() const
 	{
-		check(TopNest->NestType == NEST_Class || TopNest->NestType == NEST_Interface || TopNest->NestType == NEST_NativeInterface);
+		check(TopNest->NestType == NEST_Class || TopNest->NestType == NEST_Interface);
 
 		return (FStructScope*)TopNest->GetScope();
 	}
@@ -309,7 +308,7 @@ protected:
 		int32 Index = 0;
 		while (TopNest[Index].NestType != NEST_GlobalScope)
 		{
-			if (TopNest[Index].NestType == NEST_Class || TopNest->NestType == NEST_Interface || TopNest->NestType == NEST_NativeInterface)
+			if (TopNest[Index].NestType == NEST_Class || TopNest->NestType == NEST_Interface)
 			{
 				return true;
 			}
diff --git a/Engine/Source/Programs/UnrealHeaderTool/Private/NativeClassExporter.h b/Engine/Source/Programs/UnrealHeaderTool/Private/NativeClassExporter.h
--- a/Engine/Source/Programs/UnrealHeaderTool/Private/NativeClassExporter.h
+++ b/Engine/Source/Programs/UnrealHeaderTool/Private/NativeClassExporter.h
@@ -252,7 +252,7 @@ private:
 	void ExportGeneratedEnumsInitCode(const TArray<UEnum*>& Enums);
 
 	/**
-	 * Exports the macro declarations for GENERATED_BODY() for each Foo in the list of structs specified
+	 * Exports the macro declarations for GENERATED_USTRUCT_BODY() for each Foo in the list of structs specified
 	 * 
 	 * @param	SourceFile	the source file
 	 * @param	Structs		the structs to export
@@ -399,9 +399,8 @@ private:
 	 * @param FunctionData function data for the current function
 	 * @param Parameters list of parameters in the function
 	 * @param Return return parameter for the function
-	 * @param DeprecationWarningOutputDevice Device to output deprecation warnings for _Validate and _Implementation functions.
 	 */
-	void ExportFunctionThunk(FStringOutputDevice& RPCWrappers, UFunction* Function, const FFuncInfo& FunctionData, const TArray<UProperty*>& Parameters, UProperty* Return, FStringOutputDevice& DeprecationWarningOutputDevice);
+	void ExportFunctionThunk(FStringOutputDevice& RPCWrappers, UFunction* Function, const FFuncInfo& FunctionData, const TArray<UProperty*>& Parameters, UProperty* Return);
 
 	/** Exports the native function registration code for the given class. */
 	void ExportNatives(FClass* Class);
diff --git a/Engine/Source/Programs/UnrealHeaderTool/Private/ParserHelper.h b/Engine/Source/Programs/UnrealHeaderTool/Private/ParserHelper.h
--- a/Engine/Source/Programs/UnrealHeaderTool/Private/ParserHelper.h
+++ b/Engine/Source/Programs/UnrealHeaderTool/Private/ParserHelper.h
@@ -880,34 +880,34 @@ struct FFuncInfo
 	/** Whether this function represents a sealed event */
 	bool		bSealedEvent;
 	/** Delegate macro line in header. */
-	int32		MacroLine;
+	int32		DelegateMacroLine;
 	//@}
 
 	/** Constructor. */
 	FFuncInfo()
-		: Function()
-		, FunctionFlags(0)
-		, FunctionExportFlags(0)
-		, ExpectParms(0)
-		, FunctionReference(NULL)
-		, CppImplName(TEXT(""))
-		, CppValidationImplName(TEXT(""))
-		, RPCId(0)
-		, RPCResponseId(0)
-		, bSealedEvent(false)
-		, MacroLine(-1)
+	:	Function				()
+	,	FunctionFlags			(0)
+	,	FunctionExportFlags		(0)
+	,	ExpectParms				(0)
+	,	FunctionReference		(NULL)
+	,	CppImplName				(TEXT(""))
+	,	CppValidationImplName	(TEXT(""))
+	,	RPCId					(0)
+	,	RPCResponseId			(0)
+	,	bSealedEvent			(false)
+	,	DelegateMacroLine		(-1)
 	{}
 
-	FFuncInfo(const FFuncInfo& Other)
-		: FunctionFlags(Other.FunctionFlags)
-		, FunctionExportFlags(Other.FunctionExportFlags)
-		, ExpectParms(Other.ExpectParms)
-		, FunctionReference(Other.FunctionReference)
-		, CppImplName(Other.CppImplName)
-		, CppValidationImplName(Other.CppValidationImplName)
-		, RPCId(Other.RPCId)
-		, RPCResponseId(Other.RPCResponseId)
-		, MacroLine(Other.MacroLine)
+	FFuncInfo( const FFuncInfo& Other )
+	:	FunctionFlags   (Other.FunctionFlags)
+	,	FunctionExportFlags(Other.FunctionExportFlags)
+	,	ExpectParms		(Other.ExpectParms)
+	,	FunctionReference(Other.FunctionReference)
+	,	CppImplName(Other.CppImplName)
+	,	CppValidationImplName(Other.CppValidationImplName)
+	,	RPCId(Other.RPCId)
+	,	RPCResponseId(Other.RPCResponseId)
+	,	DelegateMacroLine(Other.DelegateMacroLine)
 	{
 		Function.Clone(Other.Function);
 		if (FunctionReference)
diff --git a/Engine/Source/Programs/UnrealHeaderTool/Resources/UHTDebugging/TestDeprecatedObject.h b/Engine/Source/Programs/UnrealHeaderTool/Resources/UHTDebugging/TestDeprecatedObject.h
--- a/Engine/Source/Programs/UnrealHeaderTool/Resources/UHTDebugging/TestDeprecatedObject.h
+++ b/Engine/Source/Programs/UnrealHeaderTool/Resources/UHTDebugging/TestDeprecatedObject.h
@@ -15,7 +15,5 @@ namespace EFoo
 UCLASS(deprecated)
 class UDEPRECATED_TestDeprecatedObject : public UObject
 {
-	GENERATED_BODY()
-public:
-	UDEPRECATED_TestDeprecatedObject(const FObjectInitializer& ObjectInitializer = FObjectInitializer::Get());
+	GENERATED_UCLASS_BODY()
 };
diff --git a/Engine/Source/Programs/UnrealHeaderTool/Resources/UHTDebugging/TestInterface.h b/Engine/Source/Programs/UnrealHeaderTool/Resources/UHTDebugging/TestInterface.h
--- a/Engine/Source/Programs/UnrealHeaderTool/Resources/UHTDebugging/TestInterface.h
+++ b/Engine/Source/Programs/UnrealHeaderTool/Resources/UHTDebugging/TestInterface.h
@@ -7,15 +7,12 @@
 UINTERFACE()
 class UTestInterface : public UInterface
 {
-	GENERATED_BODY()
-public:
-	UTestInterface(const FObjectInitializer& ObjectInitializer = FObjectInitializer::Get());
+	GENERATED_UINTERFACE_BODY()
 };
 
 class ITestInterface
 {
-	GENERATED_BODY()
-public:
+	GENERATED_IINTERFACE_BODY()
 
 	UFUNCTION(BlueprintNativeEvent)
 	FString SomeFunction(int32 Val) const;
diff --git a/Engine/Source/Programs/UnrealHeaderTool/Resources/UHTDebugging/TestInterfaceObject.h b/Engine/Source/Programs/UnrealHeaderTool/Resources/UHTDebugging/TestInterfaceObject.h
--- a/Engine/Source/Programs/UnrealHeaderTool/Resources/UHTDebugging/TestInterfaceObject.h
+++ b/Engine/Source/Programs/UnrealHeaderTool/Resources/UHTDebugging/TestInterfaceObject.h
@@ -7,9 +7,7 @@
 UCLASS()
 class UTestInterfaceObject : public UObject, public ITestInterface
 {
-	GENERATED_BODY()
-public:
-	UTestInterfaceObject(const FObjectInitializer& ObjectInitializer = FObjectInitializer::Get());
+	GENERATED_UCLASS_BODY()
 
 //	UFUNCTION(BlueprintNativeEvent)
 	FString SomeFunction(int32 Val) const;
diff --git a/Engine/Source/Programs/UnrealHeaderTool/Resources/UHTDebugging/TestObject.h b/Engine/Source/Programs/UnrealHeaderTool/Resources/UHTDebugging/TestObject.h
--- a/Engine/Source/Programs/UnrealHeaderTool/Resources/UHTDebugging/TestObject.h
+++ b/Engine/Source/Programs/UnrealHeaderTool/Resources/UHTDebugging/TestObject.h
@@ -8,9 +8,7 @@
 UCLASS()
 class UTestObject : public UObject
 {
-	GENERATED_BODY()
-public:
-	UTestObject(const FObjectInitializer& ObjectInitializer = FObjectInitializer::Get());
+	GENERATED_UCLASS_BODY()
 
 	UFUNCTION(BlueprintCallable, Category="Random")
 	void TestForNullPtrDefaults(UObject* Obj1 = NULL, UObject* Obj2 = nullptr, UObject* Obj3 = 0);
diff --git a/Engine/Source/Runtime/AIModule/Classes/AISystem.h b/Engine/Source/Runtime/AIModule/Classes/AISystem.h
--- a/Engine/Source/Runtime/AIModule/Classes/AISystem.h
+++ b/Engine/Source/Runtime/AIModule/Classes/AISystem.h
@@ -18,9 +18,7 @@ class UBlackboardComponent;
 UCLASS(config=Engine, defaultconfig)
 class AIMODULE_API UAISystem : public UAISystemBase
 {
-	GENERATED_BODY()
-public:
-	UAISystem(const FObjectInitializer& ObjectInitializer = FObjectInitializer::Get());
+	GENERATED_UCLASS_BODY()
 
 	UPROPERTY(globalconfig, EditAnywhere, Category = "AISystem", meta = (MetaClass = "AIPerceptionSystem", DisplayName = "Perception System Class"))
 	FStringClassReference PerceptionSystemClassName;
diff --git a/Engine/Source/Runtime/AIModule/Classes/Actions/PawnAction.h b/Engine/Source/Runtime/AIModule/Classes/Actions/PawnAction.h
--- a/Engine/Source/Runtime/AIModule/Classes/Actions/PawnAction.h
+++ b/Engine/Source/Runtime/AIModule/Classes/Actions/PawnAction.h
@@ -43,9 +43,7 @@ namespace EPawnActionFailHandling
 UCLASS(abstract, EditInlineNew)
 class AIMODULE_API UPawnAction : public UObject
 {
-	GENERATED_BODY()
-public:
-	UPawnAction(const FObjectInitializer& ObjectInitializer = FObjectInitializer::Get());
+	GENERATED_UCLASS_BODY()
 
 	friend UPawnActionsComponent;
 	friend FPawnActionStack;
diff --git a/Engine/Source/Runtime/AIModule/Classes/Actions/PawnAction_BlueprintBase.h b/Engine/Source/Runtime/AIModule/Classes/Actions/PawnAction_BlueprintBase.h
--- a/Engine/Source/Runtime/AIModule/Classes/Actions/PawnAction_BlueprintBase.h
+++ b/Engine/Source/Runtime/AIModule/Classes/Actions/PawnAction_BlueprintBase.h
@@ -7,9 +7,7 @@
 UCLASS(abstract, Blueprintable)
 class AIMODULE_API UPawnAction_BlueprintBase : public UPawnAction
 {
-	GENERATED_BODY()
-public:
-	UPawnAction_BlueprintBase(const FObjectInitializer& ObjectInitializer = FObjectInitializer::Get());
+	GENERATED_UCLASS_BODY()
 
 public:
 
diff --git a/Engine/Source/Runtime/AIModule/Classes/Actions/PawnAction_Move.h b/Engine/Source/Runtime/AIModule/Classes/Actions/PawnAction_Move.h
--- a/Engine/Source/Runtime/AIModule/Classes/Actions/PawnAction_Move.h
+++ b/Engine/Source/Runtime/AIModule/Classes/Actions/PawnAction_Move.h
@@ -22,9 +22,7 @@ namespace EPawnActionMoveMode
 UCLASS()
 class AIMODULE_API UPawnAction_Move : public UPawnAction
 {
-	GENERATED_BODY()
-public:
-	UPawnAction_Move(const FObjectInitializer& ObjectInitializer = FObjectInitializer::Get());
+	GENERATED_UCLASS_BODY()
 protected:
 	UPROPERTY(Category = PawnAction, EditAnywhere, BlueprintReadWrite)
 	AActor* GoalActor;
diff --git a/Engine/Source/Runtime/AIModule/Classes/Actions/PawnAction_Repeat.h b/Engine/Source/Runtime/AIModule/Classes/Actions/PawnAction_Repeat.h
--- a/Engine/Source/Runtime/AIModule/Classes/Actions/PawnAction_Repeat.h
+++ b/Engine/Source/Runtime/AIModule/Classes/Actions/PawnAction_Repeat.h
@@ -7,9 +7,7 @@
 UCLASS()
 class AIMODULE_API UPawnAction_Repeat : public UPawnAction
 {
-	GENERATED_BODY()
-public:
-	UPawnAction_Repeat(const FObjectInitializer& ObjectInitializer = FObjectInitializer::Get());
+	GENERATED_UCLASS_BODY()
 
 	enum
 	{
@@ -39,4 +37,4 @@ protected:
 	virtual void OnChildFinished(UPawnAction& Action, EPawnActionResult::Type WithResult) override;
 
 	bool PushSubAction();
-};
+};
\ No newline at end of file
diff --git a/Engine/Source/Runtime/AIModule/Classes/Actions/PawnAction_Sequence.h b/Engine/Source/Runtime/AIModule/Classes/Actions/PawnAction_Sequence.h
--- a/Engine/Source/Runtime/AIModule/Classes/Actions/PawnAction_Sequence.h
+++ b/Engine/Source/Runtime/AIModule/Classes/Actions/PawnAction_Sequence.h
@@ -7,9 +7,7 @@
 UCLASS()
 class AIMODULE_API UPawnAction_Sequence : public UPawnAction
 {
-	GENERATED_BODY()
-public:
-	UPawnAction_Sequence(const FObjectInitializer& ObjectInitializer = FObjectInitializer::Get());
+	GENERATED_UCLASS_BODY()
 
 	UPROPERTY()
 	TArray<UPawnAction*> ActionSequence;
@@ -32,4 +30,4 @@ protected:
 	virtual void OnChildFinished(UPawnAction& Action, EPawnActionResult::Type WithResult) override;
 
 	bool PushNextActionCopy();
-};
+};
\ No newline at end of file
diff --git a/Engine/Source/Runtime/AIModule/Classes/Actions/PawnAction_Wait.h b/Engine/Source/Runtime/AIModule/Classes/Actions/PawnAction_Wait.h
--- a/Engine/Source/Runtime/AIModule/Classes/Actions/PawnAction_Wait.h
+++ b/Engine/Source/Runtime/AIModule/Classes/Actions/PawnAction_Wait.h
@@ -9,9 +9,7 @@
 UCLASS()
 class AIMODULE_API UPawnAction_Wait : public UPawnAction
 {
-	GENERATED_BODY()
-public:
-	UPawnAction_Wait(const FObjectInitializer& ObjectInitializer = FObjectInitializer::Get());
+	GENERATED_UCLASS_BODY()
 		
 	UPROPERTY()
 	float TimeToWait;
@@ -29,4 +27,4 @@ public:
 	virtual EPawnActionAbortState::Type PerformAbort(EAIForceParam::Type ShouldForce);
 
 	void TimerDone();
-};
+};
\ No newline at end of file
diff --git a/Engine/Source/Runtime/AIModule/Classes/Actions/PawnActionsComponent.h b/Engine/Source/Runtime/AIModule/Classes/Actions/PawnActionsComponent.h
--- a/Engine/Source/Runtime/AIModule/Classes/Actions/PawnActionsComponent.h
+++ b/Engine/Source/Runtime/AIModule/Classes/Actions/PawnActionsComponent.h
@@ -11,7 +11,7 @@ class UPawnAction;
 USTRUCT()
 struct AIMODULE_API FPawnActionEvent
 {
-	GENERATED_BODY()
+	GENERATED_USTRUCT_BODY()
 
 	// used for marking FPawnActionEvent instances created solely for comparisons uses
 	static const int32 FakeActionIndex = INDEX_NONE;
@@ -37,7 +37,7 @@ struct AIMODULE_API FPawnActionEvent
 USTRUCT()
 struct AIMODULE_API FPawnActionStack
 {
-	GENERATED_BODY()
+	GENERATED_USTRUCT_BODY()
 
 private:
 	UPROPERTY()
@@ -68,9 +68,7 @@ public:
 UCLASS()
 class AIMODULE_API UPawnActionsComponent : public UActorComponent
 {
-	GENERATED_BODY()
-public:
-	UPawnActionsComponent(const FObjectInitializer& ObjectInitializer = FObjectInitializer::Get());
+	GENERATED_UCLASS_BODY()
 
 protected:
 	UPROPERTY(BlueprintReadOnly, Category="PawnActions")
@@ -163,4 +161,4 @@ protected:
 private:
 	/** Removed all pending action events associated with PawnAction. Private to make sure it's called only in special cases */
 	void RemoveEventsForAction(UPawnAction& PawnAction);
-};
+};
\ No newline at end of file
diff --git a/Engine/Source/Runtime/AIModule/Classes/BehaviorTree/BTAuxiliaryNode.h b/Engine/Source/Runtime/AIModule/Classes/BehaviorTree/BTAuxiliaryNode.h
--- a/Engine/Source/Runtime/AIModule/Classes/BehaviorTree/BTAuxiliaryNode.h
+++ b/Engine/Source/Runtime/AIModule/Classes/BehaviorTree/BTAuxiliaryNode.h
@@ -30,9 +30,7 @@ struct FBTAuxiliaryMemory : public FBTInstancedNodeMemory
 UCLASS(Abstract)
 class AIMODULE_API UBTAuxiliaryNode : public UBTNode
 {
-	GENERATED_BODY()
-public:
-	UBTAuxiliaryNode(const FObjectInitializer& ObjectInitializer = FObjectInitializer::Get());
+	GENERATED_UCLASS_BODY()
 
 	/** wrapper for node instancing: OnBecomeRelevant */
 	void WrappedOnBecomeRelevant(UBehaviorTreeComponent& OwnerComp, uint8* NodeMemory) const;
diff --git a/Engine/Source/Runtime/AIModule/Classes/BehaviorTree/BTCompositeNode.h b/Engine/Source/Runtime/AIModule/Classes/BehaviorTree/BTCompositeNode.h
--- a/Engine/Source/Runtime/AIModule/Classes/BehaviorTree/BTCompositeNode.h
+++ b/Engine/Source/Runtime/AIModule/Classes/BehaviorTree/BTCompositeNode.h
@@ -49,7 +49,7 @@ namespace EBTDecoratorLogic
 USTRUCT()
 struct FBTDecoratorLogic
 {
-	GENERATED_BODY()
+	GENERATED_USTRUCT_BODY()
 
 	UPROPERTY()
 	TEnumAsByte<EBTDecoratorLogic::Type> Operation;
@@ -64,7 +64,7 @@ struct FBTDecoratorLogic
 USTRUCT()
 struct FBTCompositeChild
 {
-	GENERATED_BODY()
+	GENERATED_USTRUCT_BODY()
 
 	/** child node */
 	UPROPERTY()
@@ -85,9 +85,7 @@ struct FBTCompositeChild
 UCLASS(Abstract)
 class AIMODULE_API UBTCompositeNode : public UBTNode
 {
-	GENERATED_BODY()
-public:
-	UBTCompositeNode(const FObjectInitializer& ObjectInitializer = FObjectInitializer::Get());
+	GENERATED_UCLASS_BODY()
 
 	/** child nodes */
 	UPROPERTY()
diff --git a/Engine/Source/Runtime/AIModule/Classes/BehaviorTree/BTDecorator.h b/Engine/Source/Runtime/AIModule/Classes/BehaviorTree/BTDecorator.h
--- a/Engine/Source/Runtime/AIModule/Classes/BehaviorTree/BTDecorator.h
+++ b/Engine/Source/Runtime/AIModule/Classes/BehaviorTree/BTDecorator.h
@@ -30,9 +30,7 @@ struct FBehaviorTreeSearchData;
 UCLASS(Abstract)
 class AIMODULE_API UBTDecorator : public UBTAuxiliaryNode
 {
-	GENERATED_BODY()
-public:
-	UBTDecorator(const FObjectInitializer& ObjectInitializer = FObjectInitializer::Get());
+	GENERATED_UCLASS_BODY()
 
 	/** fill in data about tree structure */
 	void InitializeDecorator(uint8 InChildIndex);
diff --git a/Engine/Source/Runtime/AIModule/Classes/BehaviorTree/BTTaskNode.h b/Engine/Source/Runtime/AIModule/Classes/BehaviorTree/BTTaskNode.h
--- a/Engine/Source/Runtime/AIModule/Classes/BehaviorTree/BTTaskNode.h
+++ b/Engine/Source/Runtime/AIModule/Classes/BehaviorTree/BTTaskNode.h
@@ -26,9 +26,7 @@ struct FAIMessage;
 UCLASS(Abstract)
 class AIMODULE_API UBTTaskNode : public UBTNode
 {
-	GENERATED_BODY()
-public:
-	UBTTaskNode(const FObjectInitializer& ObjectInitializer = FObjectInitializer::Get());
+	GENERATED_UCLASS_BODY()
 
 	/** starts this task, should return Succeeded, Failed or InProgress
 	 *  (use FinishLatentTask() when returning InProgress)
diff --git a/Engine/Source/Runtime/AIModule/Classes/BehaviorTree/BehaviorTree.h b/Engine/Source/Runtime/AIModule/Classes/BehaviorTree/BehaviorTree.h
--- a/Engine/Source/Runtime/AIModule/Classes/BehaviorTree/BehaviorTree.h
+++ b/Engine/Source/Runtime/AIModule/Classes/BehaviorTree/BehaviorTree.h
@@ -13,9 +13,7 @@ class UBTDecorator;
 UCLASS(BlueprintType)
 class AIMODULE_API UBehaviorTree : public UObject
 {
-	GENERATED_BODY()
-public:
-	UBehaviorTree(const FObjectInitializer& ObjectInitializer = FObjectInitializer::Get());
+	GENERATED_UCLASS_BODY()
 
 	/** root node of loaded tree */
 	UPROPERTY()
diff --git a/Engine/Source/Runtime/AIModule/Classes/BehaviorTree/BehaviorTreeComponent.h b/Engine/Source/Runtime/AIModule/Classes/BehaviorTree/BehaviorTreeComponent.h
--- a/Engine/Source/Runtime/AIModule/Classes/BehaviorTree/BehaviorTreeComponent.h
+++ b/Engine/Source/Runtime/AIModule/Classes/BehaviorTree/BehaviorTreeComponent.h
@@ -78,9 +78,7 @@ struct FBTPendingInitializeInfo
 UCLASS()
 class AIMODULE_API UBehaviorTreeComponent : public UBrainComponent
 {
-	GENERATED_BODY()
-public:
-	UBehaviorTreeComponent(const FObjectInitializer& ObjectInitializer = FObjectInitializer::Get());
+	GENERATED_UCLASS_BODY()
 
 	// Begin UBrainComponent overrides
 	virtual void RestartLogic() override;
diff --git a/Engine/Source/Runtime/AIModule/Classes/BehaviorTree/BehaviorTreeManager.h b/Engine/Source/Runtime/AIModule/Classes/BehaviorTree/BehaviorTreeManager.h
--- a/Engine/Source/Runtime/AIModule/Classes/BehaviorTree/BehaviorTreeManager.h
+++ b/Engine/Source/Runtime/AIModule/Classes/BehaviorTree/BehaviorTreeManager.h
@@ -11,7 +11,7 @@ class UBehaviorTree;
 USTRUCT()
 struct FBehaviorTreeTemplateInfo
 {
-	GENERATED_BODY()
+	GENERATED_USTRUCT_BODY()
 
 	/** behavior tree asset */
 	UPROPERTY()
@@ -28,9 +28,7 @@ struct FBehaviorTreeTemplateInfo
 UCLASS(config=Engine)
 class AIMODULE_API UBehaviorTreeManager : public UObject
 {
-	GENERATED_BODY()
-public:
-	UBehaviorTreeManager(const FObjectInitializer& ObjectInitializer = FObjectInitializer::Get());
+	GENERATED_UCLASS_BODY()
 
 	/** limit for recording execution steps for debugger */
 	UPROPERTY(config)
