Enabling UHT to digest GENERATED_BODY instead of GENERATED_*_BODY.

diff --git a/Engine/Source/Programs/UnrealHeaderTool/Private/HeaderParser.cpp b/Engine/Source/Programs/UnrealHeaderTool/Private/HeaderParser.cpp
--- a/Engine/Source/Programs/UnrealHeaderTool/Private/HeaderParser.cpp
+++ b/Engine/Source/Programs/UnrealHeaderTool/Private/HeaderParser.cpp
@@ -1432,17 +1432,17 @@ UScriptStruct* FHeaderParser::CompileStructDeclaration(FClasses& AllClasses, FUn
 		{
 			FError::Throwf(TEXT("USTRUCTs cannot contain UFUNCTIONs."));
 		}
-		else if (Token.Matches(TEXT("GENERATED_USTRUCT_BODY")))
+		else if (Token.Matches(TEXT("GENERATED_USTRUCT_BODY")) || Token.Matches(TEXT("GENERATED_BODY")))
 		{
-			// Match 'GENERATED_USTRUCT_BODY' '(' [StructName] ')'
+			// Match 'GENERATED_USTRUCT_BODY' '(' [StructName] ')' or 'GENERATED_BODY' '(' [StructName] ')'
 			if (CurrentAccessSpecifier != ACCESS_Public)
 			{
-				FError::Throwf(TEXT("GENERATED_USTRUCT_BODY must be in the public scope of '%s', not private or protected."), *StructNameInScript);
+				FError::Throwf(TEXT("%s must be in the public scope of '%s', not private or protected."), Token.Identifier, *StructNameInScript);
 			}
 
 			if (Struct->StructMacroDeclaredLineNumber != INDEX_NONE)
 			{
-				FError::Throwf(TEXT("Multiple GENERATED_USTRUCT_BODY declarations found in '%s'"), *StructNameInScript);
+				FError::Throwf(TEXT("Multiple %s declarations found in '%s'"), Token.Identifier, *StructNameInScript);
 			}
 
 			Struct->StructMacroDeclaredLineNumber = InputLine;
@@ -1453,7 +1453,7 @@ UScriptStruct* FHeaderParser::CompileStructDeclaration(FClasses& AllClasses, FUn
 			{
 				if (!DuplicateStructName.Matches(*StructNameInScript))
 				{
-					FError::Throwf(TEXT("The argument to GENERATED_USTRUCT_BODY must match the struct name '%s' if present.  However, the argument can be omitted entirely."), *StructNameInScript);
+					FError::Throwf(TEXT("The argument to %s must match the struct name '%s' if present.  However, the argument can be omitted entirely."), Token.Identifier, *StructNameInScript);
 				}
 			}
 
@@ -1567,7 +1567,7 @@ UScriptStruct* FHeaderParser::CompileStructDeclaration(FClasses& AllClasses, FUn
 	{
 		// Roll the line number back to the start of the struct body and error out
 		InputLine = SavedLineNumber;
-		FError::Throwf(TEXT("Expected a GENERATED_USTRUCT_BODY() at the start of struct"));
+		FError::Throwf(TEXT("Expected a GENERATED_BODY() at the start of struct"));
 	}
 
 	// Link the properties within the struct
@@ -1625,6 +1625,7 @@ const TCHAR *FHeaderParser::NestTypeName( ENestType NestType )
 			return TEXT("Global Scope");
 		case NEST_Class:
 			return TEXT("Class");
+		case NEST_NativeInterface:
 		case NEST_Interface:
 			return TEXT("Interface");
 		case NEST_FunctionDeclaration:
@@ -1705,6 +1706,7 @@ void FHeaderParser::PushNest(ENestType NestType, UStruct* InNode, FUnrealSourceF
 		TopNest->Allow = ALLOW_VarDecl | ALLOW_Function | ALLOW_TypeDecl;
 		break;
 
+	case NEST_NativeInterface:
 	case NEST_Interface:
 		TopNest->Allow = ALLOW_Function | ALLOW_TypeDecl;
 		break;
@@ -1738,7 +1740,7 @@ void FHeaderParser::PopNest(ENestType NestType, const TCHAR* Descr)
 		FError::Throwf(TEXT("Unexpected end of %s in '%s' block"), Descr, NestTypeName(TopNest->NestType));
 	}
 
-	if (NestType != NEST_GlobalScope && NestType != NEST_Class && NestType != NEST_Interface && NestType != NEST_FunctionDeclaration)
+	if (NestType != NEST_GlobalScope && NestType != NEST_Class && NestType != NEST_Interface && NestType != NEST_NativeInterface && NestType != NEST_FunctionDeclaration)
 	{
 		FError::Throwf(TEXT("Bad first pass NestType %i"), (uint8)NestType);
 	}
@@ -3627,7 +3629,7 @@ bool FHeaderParser::CompileDeclaration(FClasses& AllClasses, FUnrealSourceFile&
 		{
 			FError::Throwf(TEXT("Access specifier %s not allowed here."), Token.Identifier);
 		}
-		check(TopNest->NestType == NEST_Class || TopNest->NestType == NEST_Interface);
+		check(TopNest->NestType == NEST_Class || TopNest->NestType == NEST_Interface || TopNest->NestType == NEST_NativeInterface);
 		CurrentAccessSpecifier = AccessSpecifier;
 	}
 	else if (Token.Matches(TEXT("class")) && (TopNest->NestType == NEST_GlobalScope))
@@ -3640,6 +3642,7 @@ bool FHeaderParser::CompileDeclaration(FClasses& AllClasses, FUnrealSourceFile&
 
 		// Start parsing the second class
 		bEncounteredNewStyleClass_UnmatchedBrackets = true;
+		CurrentAccessSpecifier = ACCESS_Private;
 
 		if (!TryParseIInterfaceClass(AllClasses))
 		{
@@ -3648,75 +3651,79 @@ bool FHeaderParser::CompileDeclaration(FClasses& AllClasses, FUnrealSourceFile&
 			return SkipDeclaration(Token);
 		}
 	}
-	else if (Token.Matches(TEXT("GENERATED_BODY")))
+	else if (Token.Matches(TEXT("GENERATED_IINTERFACE_BODY")) || (Token.Matches(TEXT("GENERATED_BODY")) && TopNest->NestType == NEST_NativeInterface))
 	{
-		if (TopNest->NestType != NEST_Class && TopNest->NestType != NEST_Interface)
+		if (TopNest->NestType != NEST_NativeInterface)
 		{
-			FError::Throwf(TEXT("%s must occur inside the class or interface definition"), Token.Identifier);
+			FError::Throwf(TEXT("%s must occur inside the native interface definition"), Token.Identifier);
 		}
-		if (!ClassDefinitionRanges.Contains(GetCurrentClass()))
-		{
-			ClassDefinitionRanges.Add(GetCurrentClass(), ClassDefinitionRange());
-		}
-
-		ClassDefinitionRanges[GetCurrentClass()].bHasGeneratedBody = true;
-
 		RequireSymbol(TEXT("("), Token.Identifier);
 		RequireSymbol(TEXT(")"), Token.Identifier);
 
 		auto* ClassData = GetCurrentClassData();
 
 		ClassData->GeneratedBodyMacroAccessSpecifier = CurrentAccessSpecifier;
-		ClassData->SetGeneratedBodyLine(InputLine);
+		ClassData->SetInterfaceGeneratedBodyLine(InputLine);
 
-		if (TopNest->NestType == NEST_Class)
+		bClassHasGeneratedIInterfaceBody = true;
+
+		if (Token.Matches(TEXT("GENERATED_IINTERFACE_BODY")))
 		{
-			bClassHasGeneratedBody = true;
+			CurrentAccessSpecifier = ACCESS_Public;
 		}
 	}
-	else if (Token.Matches(TEXT("GENERATED_UCLASS_BODY")))
+	else if (Token.Matches(TEXT("GENERATED_UINTERFACE_BODY")) || (Token.Matches(TEXT("GENERATED_BODY")) && TopNest->NestType == NEST_Interface))
 	{
-		if (TopNest->NestType != NEST_Class)
+		if (TopNest->NestType != NEST_Interface)
 		{
-			FError::Throwf(TEXT("%s must occur inside the class definition"), Token.Identifier);
+			FError::Throwf(TEXT("%s must occur inside the interface definition"), Token.Identifier);
 		}
 		RequireSymbol(TEXT("("), Token.Identifier);
 		RequireSymbol(TEXT(")"), Token.Identifier);
 
-		GetCurrentClassData()->SetGeneratedBodyLine(InputLine);
+		auto* ClassData = GetCurrentClassData();
+
+		ClassData->GeneratedBodyMacroAccessSpecifier = CurrentAccessSpecifier;
+		ClassData->SetGeneratedBodyLine(InputLine);
 
-		// The body implementation macro always ends with a 'public:'
-		CurrentAccessSpecifier = ACCESS_Public;
+		bClassHasGeneratedUInterfaceBody = true;
 
-		bClassHasGeneratedBody = true;
+		if (Token.Matches(TEXT("GENERATED_UINTERFACE_BODY")))
+		{
+			CurrentAccessSpecifier = ACCESS_Public;
+		}
 	}
-	else if (Token.Matches(TEXT("GENERATED_IINTERFACE_BODY")))
+	else if (Token.Matches(TEXT("GENERATED_UCLASS_BODY")) || (Token.Matches(TEXT("GENERATED_BODY")) && TopNest->NestType == NEST_Class))
 	{
-		RequireSymbol(TEXT("("), Token.Identifier);
-		RequireSymbol(TEXT(")"), Token.Identifier);
+		if (TopNest->NestType != NEST_Class)
+		{
+			FError::Throwf(TEXT("%s must occur inside the class definition"), Token.Identifier);
+		}
 
-		GetCurrentClassData()->SetInterfaceGeneratedBodyLine(InputLine);
+		auto* ClassData = GetCurrentClassData();
 
-		bClassHasGeneratedIInterfaceBody = true;
+		if (Token.Matches(TEXT("GENERATED_BODY")))
+		{
+			if (!ClassDefinitionRanges.Contains(GetCurrentClass()))
+			{
+				ClassDefinitionRanges.Add(GetCurrentClass(), ClassDefinitionRange());
+			}
 
-		// The body implementation macro always ends with a 'public:'
-		CurrentAccessSpecifier = ACCESS_Public;
-	}
-	else if (Token.Matches(TEXT("GENERATED_UINTERFACE_BODY")))
-	{
-		if (TopNest->NestType != NEST_Interface)
+			ClassDefinitionRanges[GetCurrentClass()].bHasGeneratedBody = true;
+
+			ClassData->GeneratedBodyMacroAccessSpecifier = CurrentAccessSpecifier;
+		}
+		else
 		{
-			FError::Throwf(TEXT("%s must occur inside the interface definition"), Token.Identifier);
+			CurrentAccessSpecifier = ACCESS_Public;
 		}
+
 		RequireSymbol(TEXT("("), Token.Identifier);
 		RequireSymbol(TEXT(")"), Token.Identifier);
 
-		GetCurrentClassData()->SetGeneratedBodyLine(InputLine);
-
-		bClassHasGeneratedUInterfaceBody = true;
+		ClassData->SetGeneratedBodyLine(InputLine);
 
-		// The body implementation macro always ends with a 'public:'
-		CurrentAccessSpecifier = ACCESS_Public;
+		bClassHasGeneratedBody = true;
 	}
 	else if (Token.Matches(TEXT("UCLASS"), ESearchCase::CaseSensitive) && (TopNest->Allow & ALLOW_Class))
 	{
@@ -3782,7 +3789,8 @@ bool FHeaderParser::CompileDeclaration(FClasses& AllClasses, FUnrealSourceFile&
 		// Pop nesting here to allow other non UClass declarations in the header file.
 		if (CurrentClass->ClassFlags & CLASS_Interface)
 		{
-			PopNest(NEST_Interface, TEXT("'Interface'"));
+			checkf(TopNest->NestType == NEST_Interface || TopNest->NestType == NEST_NativeInterface, TEXT("Unexpected end of interface block."));
+			PopNest(TopNest->NestType, TEXT("'Interface'"));
 			PostPopNestInterface(AllClasses, CurrentClass);
 
 			// Ensure the UINTERFACE classes have a GENERATED_UINTERFACE_BODY declaration
@@ -3827,7 +3835,11 @@ bool FHeaderParser::CompileDeclaration(FClasses& AllClasses, FUnrealSourceFile&
 			FError::Throwf(TEXT("Extra ';' before end of file"));
 		}
 	}
-	else if (bEncounteredNewStyleClass_UnmatchedBrackets && IsInAClass() && GetCurrentClass() && Token.Matches(NameLookupCPP.GetNameCPP(GetCurrentClass())))
+	else if (bEncounteredNewStyleClass_UnmatchedBrackets && IsInAClass() && GetCurrentClass() &&
+		(
+			Token.Matches(NameLookupCPP.GetNameCPP(GetCurrentClass())) || 
+			FString(Token.Identifier).EndsWith("_API") && GetToken(Token) && Token.Matches(NameLookupCPP.GetNameCPP(GetCurrentClass()))
+		))
 	{
 		if (!TryToMatchConstructorParameterList(Token))
 		{
@@ -4690,7 +4702,7 @@ bool FHeaderParser::TryParseIInterfaceClass(FClasses& AllClasses)
 	RequireSymbol(TEXT("{"), *ErrorMsg);
 
 	// Push the interface class nesting again.
-	PushNest(NEST_Interface, FoundClass);
+	PushNest(NEST_NativeInterface, FoundClass);
 
 	return true;
 }
@@ -4719,7 +4731,7 @@ void FHeaderParser::CompileInterfaceDeclaration(FClasses& AllClasses)
 	// New style files have the interface name / extends afterwards
 	RequireIdentifier(TEXT("class"), TEXT("Interface declaration"));
 	FClass* InterfaceClass = ParseInterfaceNameDeclaration(AllClasses, /*out*/ DeclaredInterfaceName, /*out*/ RequiredAPIMacroIfPresent);
-
+	ClassDefinitionRanges.Add(InterfaceClass, ClassDefinitionRange(&Input[InputPos], nullptr));
 
 	// Record that this interface is RequiredAPI if the CORE_API style macro was present
 	if (!RequiredAPIMacroIfPresent.IsEmpty())
@@ -5037,7 +5049,7 @@ void FHeaderParser::CompileDelegateDeclaration(FUnrealSourceFile& SourceFile, FC
 		FCString::Strcpy( FuncInfo.Function.Identifier, *Name );
 	}
 
-	FuncInfo.DelegateMacroLine = InputLine;
+	FuncInfo.MacroLine = InputLine;
 	auto* DelegateSignatureFunction = CreateDelegateFunction(FuncInfo);
 
 	GScriptHelper.AddClassData(DelegateSignatureFunction);
@@ -5144,6 +5156,7 @@ void FHeaderParser::CompileFunctionDeclaration(FUnrealSourceFile& SourceFile, FC
 	bool bAutomaticallyFinal = true;
 
 	FFuncInfo FuncInfo;
+	FuncInfo.MacroLine = InputLine;
 	FuncInfo.FunctionFlags = FUNC_Native;
 
 	// Infer the function's access level from the currently declared C++ access level
diff --git a/Engine/Source/Programs/UnrealHeaderTool/Private/HeaderParser.h b/Engine/Source/Programs/UnrealHeaderTool/Private/HeaderParser.h
--- a/Engine/Source/Programs/UnrealHeaderTool/Private/HeaderParser.h
+++ b/Engine/Source/Programs/UnrealHeaderTool/Private/HeaderParser.h
@@ -22,7 +22,8 @@ enum ENestType
 	NEST_GlobalScope,
 	NEST_Class,
 	NEST_FunctionDeclaration,
-	NEST_Interface
+	NEST_Interface,
+	NEST_NativeInterface
 };
 
 /** Types of statements to allow within a particular nesting block. */
@@ -295,7 +296,7 @@ protected:
 	 */
 	FStructScope* GetCurrentClassScope() const
 	{
-		check(TopNest->NestType == NEST_Class || TopNest->NestType == NEST_Interface);
+		check(TopNest->NestType == NEST_Class || TopNest->NestType == NEST_Interface || TopNest->NestType == NEST_NativeInterface);
 
 		return (FStructScope*)TopNest->GetScope();
 	}
@@ -308,7 +309,7 @@ protected:
 		int32 Index = 0;
 		while (TopNest[Index].NestType != NEST_GlobalScope)
 		{
-			if (TopNest[Index].NestType == NEST_Class || TopNest->NestType == NEST_Interface)
+			if (TopNest[Index].NestType == NEST_Class || TopNest->NestType == NEST_Interface || TopNest->NestType == NEST_NativeInterface)
 			{
 				return true;
 			}
diff --git a/Engine/Source/Programs/UnrealHeaderTool/Private/NativeClassExporter.h b/Engine/Source/Programs/UnrealHeaderTool/Private/NativeClassExporter.h
--- a/Engine/Source/Programs/UnrealHeaderTool/Private/NativeClassExporter.h
+++ b/Engine/Source/Programs/UnrealHeaderTool/Private/NativeClassExporter.h
@@ -252,7 +252,7 @@ private:
 	void ExportGeneratedEnumsInitCode(const TArray<UEnum*>& Enums);
 
 	/**
-	 * Exports the macro declarations for GENERATED_USTRUCT_BODY() for each Foo in the list of structs specified
+	 * Exports the macro declarations for GENERATED_BODY() for each Foo in the list of structs specified
 	 * 
 	 * @param	SourceFile	the source file
 	 * @param	Structs		the structs to export
@@ -399,8 +399,9 @@ private:
 	 * @param FunctionData function data for the current function
 	 * @param Parameters list of parameters in the function
 	 * @param Return return parameter for the function
+	 * @param DeprecationWarningOutputDevice Device to output deprecation warnings for _Validate and _Implementation functions.
 	 */
-	void ExportFunctionThunk(FStringOutputDevice& RPCWrappers, UFunction* Function, const FFuncInfo& FunctionData, const TArray<UProperty*>& Parameters, UProperty* Return);
+	void ExportFunctionThunk(FStringOutputDevice& RPCWrappers, UFunction* Function, const FFuncInfo& FunctionData, const TArray<UProperty*>& Parameters, UProperty* Return, FStringOutputDevice& DeprecationWarningOutputDevice);
 
 	/** Exports the native function registration code for the given class. */
 	void ExportNatives(FClass* Class);
diff --git a/Engine/Source/Programs/UnrealHeaderTool/Private/ParserHelper.h b/Engine/Source/Programs/UnrealHeaderTool/Private/ParserHelper.h
--- a/Engine/Source/Programs/UnrealHeaderTool/Private/ParserHelper.h
+++ b/Engine/Source/Programs/UnrealHeaderTool/Private/ParserHelper.h
@@ -880,34 +880,34 @@ struct FFuncInfo
 	/** Whether this function represents a sealed event */
 	bool		bSealedEvent;
 	/** Delegate macro line in header. */
-	int32		DelegateMacroLine;
+	int32		MacroLine;
 	//@}
 
 	/** Constructor. */
 	FFuncInfo()
-	:	Function				()
-	,	FunctionFlags			(0)
-	,	FunctionExportFlags		(0)
-	,	ExpectParms				(0)
-	,	FunctionReference		(NULL)
-	,	CppImplName				(TEXT(""))
-	,	CppValidationImplName	(TEXT(""))
-	,	RPCId					(0)
-	,	RPCResponseId			(0)
-	,	bSealedEvent			(false)
-	,	DelegateMacroLine		(-1)
+		: Function()
+		, FunctionFlags(0)
+		, FunctionExportFlags(0)
+		, ExpectParms(0)
+		, FunctionReference(NULL)
+		, CppImplName(TEXT(""))
+		, CppValidationImplName(TEXT(""))
+		, RPCId(0)
+		, RPCResponseId(0)
+		, bSealedEvent(false)
+		, MacroLine(-1)
 	{}
 
-	FFuncInfo( const FFuncInfo& Other )
-	:	FunctionFlags   (Other.FunctionFlags)
-	,	FunctionExportFlags(Other.FunctionExportFlags)
-	,	ExpectParms		(Other.ExpectParms)
-	,	FunctionReference(Other.FunctionReference)
-	,	CppImplName(Other.CppImplName)
-	,	CppValidationImplName(Other.CppValidationImplName)
-	,	RPCId(Other.RPCId)
-	,	RPCResponseId(Other.RPCResponseId)
-	,	DelegateMacroLine(Other.DelegateMacroLine)
+	FFuncInfo(const FFuncInfo& Other)
+		: FunctionFlags(Other.FunctionFlags)
+		, FunctionExportFlags(Other.FunctionExportFlags)
+		, ExpectParms(Other.ExpectParms)
+		, FunctionReference(Other.FunctionReference)
+		, CppImplName(Other.CppImplName)
+		, CppValidationImplName(Other.CppValidationImplName)
+		, RPCId(Other.RPCId)
+		, RPCResponseId(Other.RPCResponseId)
+		, MacroLine(Other.MacroLine)
 	{
 		Function.Clone(Other.Function);
 		if (FunctionReference)
diff --git a/Engine/Source/Programs/UnrealHeaderTool/Resources/UHTDebugging/StructOnlyHeader.h b/Engine/Source/Programs/UnrealHeaderTool/Resources/UHTDebugging/StructOnlyHeader.h
--- a/Engine/Source/Programs/UnrealHeaderTool/Resources/UHTDebugging/StructOnlyHeader.h
+++ b/Engine/Source/Programs/UnrealHeaderTool/Resources/UHTDebugging/StructOnlyHeader.h
@@ -7,5 +7,5 @@
 USTRUCT()
 struct FSomeStruct
 {
-	GENERATED_USTRUCT_BODY()
+	GENERATED_BODY()
 };
diff --git a/Engine/Source/Programs/UnrealHeaderTool/Resources/UHTDebugging/TestDeprecatedObject.h b/Engine/Source/Programs/UnrealHeaderTool/Resources/UHTDebugging/TestDeprecatedObject.h
--- a/Engine/Source/Programs/UnrealHeaderTool/Resources/UHTDebugging/TestDeprecatedObject.h
+++ b/Engine/Source/Programs/UnrealHeaderTool/Resources/UHTDebugging/TestDeprecatedObject.h
@@ -15,5 +15,7 @@ namespace EFoo
 UCLASS(deprecated)
 class UDEPRECATED_TestDeprecatedObject : public UObject
 {
-	GENERATED_UCLASS_BODY()
+	GENERATED_BODY()
+public:
+	UDEPRECATED_TestDeprecatedObject(const FObjectInitializer& ObjectInitializer = FObjectInitializer::Get());
 };
diff --git a/Engine/Source/Programs/UnrealHeaderTool/Resources/UHTDebugging/TestInterface.h b/Engine/Source/Programs/UnrealHeaderTool/Resources/UHTDebugging/TestInterface.h
--- a/Engine/Source/Programs/UnrealHeaderTool/Resources/UHTDebugging/TestInterface.h
+++ b/Engine/Source/Programs/UnrealHeaderTool/Resources/UHTDebugging/TestInterface.h
@@ -7,12 +7,15 @@
 UINTERFACE()
 class UTestInterface : public UInterface
 {
-	GENERATED_UINTERFACE_BODY()
+	GENERATED_BODY()
+public:
+	UTestInterface(const FObjectInitializer& ObjectInitializer = FObjectInitializer::Get());
 };
 
 class ITestInterface
 {
-	GENERATED_IINTERFACE_BODY()
+	GENERATED_BODY()
+public:
 
 	UFUNCTION(BlueprintNativeEvent)
 	FString SomeFunction(int32 Val) const;
diff --git a/Engine/Source/Programs/UnrealHeaderTool/Resources/UHTDebugging/TestInterfaceObject.h b/Engine/Source/Programs/UnrealHeaderTool/Resources/UHTDebugging/TestInterfaceObject.h
--- a/Engine/Source/Programs/UnrealHeaderTool/Resources/UHTDebugging/TestInterfaceObject.h
+++ b/Engine/Source/Programs/UnrealHeaderTool/Resources/UHTDebugging/TestInterfaceObject.h
@@ -7,7 +7,9 @@
 UCLASS()
 class UTestInterfaceObject : public UObject, public ITestInterface
 {
-	GENERATED_UCLASS_BODY()
+	GENERATED_BODY()
+public:
+	UTestInterfaceObject(const FObjectInitializer& ObjectInitializer = FObjectInitializer::Get());
 
 //	UFUNCTION(BlueprintNativeEvent)
 	FString SomeFunction(int32 Val) const;
diff --git a/Engine/Source/Programs/UnrealHeaderTool/Resources/UHTDebugging/TestObject.h b/Engine/Source/Programs/UnrealHeaderTool/Resources/UHTDebugging/TestObject.h
--- a/Engine/Source/Programs/UnrealHeaderTool/Resources/UHTDebugging/TestObject.h
+++ b/Engine/Source/Programs/UnrealHeaderTool/Resources/UHTDebugging/TestObject.h
@@ -8,7 +8,9 @@
 UCLASS()
 class UTestObject : public UObject
 {
-	GENERATED_UCLASS_BODY()
+	GENERATED_BODY()
+public:
+	UTestObject(const FObjectInitializer& ObjectInitializer = FObjectInitializer::Get());
 
 	UFUNCTION(BlueprintCallable, Category="Random")
 	void TestForNullPtrDefaults(UObject* Obj1 = NULL, UObject* Obj2 = nullptr, UObject* Obj3 = 0);
