.generated.inl headers removed.

diff --git a/Engine/Source/Programs/UnrealHeaderTool/Private/CodeGenerator.cpp b/Engine/Source/Programs/UnrealHeaderTool/Private/CodeGenerator.cpp
--- a/Engine/Source/Programs/UnrealHeaderTool/Private/CodeGenerator.cpp
+++ b/Engine/Source/Programs/UnrealHeaderTool/Private/CodeGenerator.cpp
@@ -4686,7 +4686,6 @@ void FNativeClassHeaderGenerator::ExportGeneratedCPP()
 	FStringOutputDevice GeneratedCPPClassesIncludes;
 	FStringOutputDevice GeneratedCPPEpilogue;
 	FStringOutputDevice GeneratedCPPText;
-	FStringOutputDevice GeneratedINLText;
 	FStringOutputDevice GeneratedLinkerFixupFunction;
 
 	TArray<FStringOutputDevice> GeneratedCPPFiles;
@@ -4703,18 +4702,6 @@ void FNativeClassHeaderGenerator::ExportGeneratedCPP()
 																								LINE_TERMINATOR
 		);
 
-	FString DeprecationMessage = TEXT(".generated.inl files have been deprecated - please remove the associated #include from your code");
-	GeneratedINLText.Logf(
-		TEXT("#ifdef _MSC_VER")										LINE_TERMINATOR
-		TEXT("    #pragma message(__FILE__\"(9): warning : %s\")")	LINE_TERMINATOR
-		TEXT("#else")												LINE_TERMINATOR
-		TEXT("    #pragma message(\"%s\")")							LINE_TERMINATOR
-		TEXT("#endif")												LINE_TERMINATOR
-																	LINE_TERMINATOR,
-		*DeprecationMessage,
-		*DeprecationMessage
-	);
-
 	FString ModulePCHInclude;
 
 	const auto* ModuleInfo = GPackageToManifestModuleMap.FindChecked(Package);
@@ -4807,9 +4794,6 @@ void FNativeClassHeaderGenerator::ExportGeneratedCPP()
 		IFileManager::Get().Delete(*CppPath);
  	}
 
-	FString InlPath = ModuleInfo->GeneratedCPPFilenameBase + TEXT(".inl");
-	SaveHeaderIfChanged(*InlPath, *(GeneratedCPPPreamble + GeneratedINLText));
-
 	if( FParse::Param( FCommandLine::Get(), TEXT("GenerateConstructors") ) && AllConstructors.Len())
 	{
 		const FString PrivateLoc = PkgDir / PkgName / TEXT("Private");
