Fix && and || being in the same scope.

diff --git a/Engine/Source/Programs/UnrealHeaderTool/Private/HeaderParser.cpp b/Engine/Source/Programs/UnrealHeaderTool/Private/HeaderParser.cpp
--- a/Engine/Source/Programs/UnrealHeaderTool/Private/HeaderParser.cpp
+++ b/Engine/Source/Programs/UnrealHeaderTool/Private/HeaderParser.cpp
@@ -3838,7 +3838,7 @@ bool FHeaderParser::CompileDeclaration(FClasses& AllClasses, FUnrealSourceFile&
 	else if (bEncounteredNewStyleClass_UnmatchedBrackets && IsInAClass() && GetCurrentClass() &&
 		(
 			Token.Matches(NameLookupCPP.GetNameCPP(GetCurrentClass())) || 
-			FString(Token.Identifier).EndsWith("_API") && GetToken(Token) && Token.Matches(NameLookupCPP.GetNameCPP(GetCurrentClass()))
+			(FString(Token.Identifier).EndsWith("_API") && GetToken(Token) && Token.Matches(NameLookupCPP.GetNameCPP(GetCurrentClass())))
 		))
 	{
 		if (!TryToMatchConstructorParameterList(Token))
