GetVarType and GetVarNameAndDim no longer take ObjectFlags.

diff --git a/Engine/Source/Programs/UnrealHeaderTool/Private/HeaderParser.cpp b/Engine/Source/Programs/UnrealHeaderTool/Private/HeaderParser.cpp
--- a/Engine/Source/Programs/UnrealHeaderTool/Private/HeaderParser.cpp
+++ b/Engine/Source/Programs/UnrealHeaderTool/Private/HeaderParser.cpp
@@ -2135,7 +2135,6 @@ bool FHeaderParser::GetVarType
 	FClasses&                       AllClasses,
 	FScope*							Scope,
 	FPropertyBase&                  VarProperty,
-	EObjectFlags&                   ObjectFlags,
 	uint64                          Disallow,
 	const TCHAR*                    Thing,
 	FToken*                         OuterPropertyType,
@@ -2148,7 +2147,6 @@ bool FHeaderParser::GetVarType
 	FName RepCallbackName = FName(NAME_None);
 
 	// Get flags.
-	ObjectFlags = RF_Public;
 	uint64 Flags = 0;
 	// force members to be 'blueprint read only' if in a const class
 	if (VariableCategory == EVariableCategory::Member && (Cast<UClass>(OwnerStruct) != nullptr) && (((UClass*)OwnerStruct)->ClassFlags & CLASS_Const))
@@ -2467,21 +2465,18 @@ bool FHeaderParser::GetVarType
 
 	if (CurrentAccessSpecifier == ACCESS_Public || VariableCategory != EVariableCategory::Member)
 	{
-		ObjectFlags |= RF_Public;
 		Flags       &= ~CPF_Protected;
 		ExportFlags |= PROPEXPORT_Public;
 		ExportFlags &= ~(PROPEXPORT_Private|PROPEXPORT_Protected);
 	}
 	else if (CurrentAccessSpecifier == ACCESS_Protected)
 	{
-		ObjectFlags |= RF_Public;
 		Flags       |= CPF_Protected;
 		ExportFlags |= PROPEXPORT_Protected;
 		ExportFlags &= ~(PROPEXPORT_Public|PROPEXPORT_Private);
 	}
 	else if (CurrentAccessSpecifier == ACCESS_Private)
 	{
-		ObjectFlags &= ~RF_Public;
 		Flags       &= ~CPF_Protected;
 		ExportFlags |= PROPEXPORT_Private;
 		ExportFlags &= ~(PROPEXPORT_Public|PROPEXPORT_Protected);
@@ -2616,11 +2611,10 @@ bool FHeaderParser::GetVarType
 
 		// GetVarType() clears the property flags of the array var, so use dummy 
 		// flags when getting the inner property
-		EObjectFlags InnerFlags;
 		uint64 OriginalVarTypeFlags = VarType.PropertyFlags;
 		VarType.PropertyFlags |= Flags;
 
-		GetVarType(AllClasses, Scope, VarProperty, InnerFlags, Disallow, TEXT("'tarray'"), &VarType, EPropertyDeclarationStyle::None, VariableCategory);
+		GetVarType(AllClasses, Scope, VarProperty, Disallow, TEXT("'tarray'"), &VarType, EPropertyDeclarationStyle::None, VariableCategory);
 		if (VarProperty.ArrayType != EArrayType::None)
 		{
 			FError::Throwf(TEXT("Arrays within arrays not supported.") );
@@ -3217,17 +3211,23 @@ EFindName FHeaderParser::GetFindFlagForPropertyName(const TCHAR* PropertyName)
 
 UProperty* FHeaderParser::GetVarNameAndDim
 (
-	UStruct*     Scope,
-	FToken&      VarProperty,
-	EObjectFlags ObjectFlags,
-	bool         NoArrays,
-	bool         IsFunction,
-	const TCHAR* HardcodedName,
-	const TCHAR* HintText
+	UStruct*                Scope,
+	FToken&                 VarProperty,
+	bool                    NoArrays,
+	bool                    IsFunction,
+	const TCHAR*            HardcodedName,
+	const TCHAR*            HintText,
+	EVariableCategory::Type VariableCategory
 )
 {
 	check(Scope);
 
+	EObjectFlags ObjectFlags = RF_Public;
+	if (VariableCategory == EVariableCategory::Member && CurrentAccessSpecifier == ACCESS_Private)
+	{
+		ObjectFlags = RF_NoFlags;
+	}
+
 	AddModuleRelativePathToMetadata(Scope, VarProperty.MetaData);
 
 	// Get variable name.
@@ -4826,8 +4826,8 @@ void FHeaderParser::ParseParameterList(FClasses& AllClasses, UFunction* Function
 	{
 		// Get parameter type.
 		FToken Property(CPT_None);
-		EObjectFlags ObjectFlags;
-		GetVarType(AllClasses, GetCurrentScope(), Property, ObjectFlags, ~(CPF_ParmFlags | CPF_AutoWeak | CPF_RepSkip), TEXT("Function parameter"), NULL, EPropertyDeclarationStyle::None, (Function->FunctionFlags & FUNC_Net) ? EVariableCategory::ReplicatedParameter : EVariableCategory::RegularParameter);
+		EVariableCategory::Type VariableCategory = (Function->FunctionFlags & FUNC_Net) ? EVariableCategory::ReplicatedParameter : EVariableCategory::RegularParameter;
+		GetVarType(AllClasses, GetCurrentScope(), Property, ~(CPF_ParmFlags | CPF_AutoWeak | CPF_RepSkip), TEXT("Function parameter"), NULL, EPropertyDeclarationStyle::None, VariableCategory);
 		Property.PropertyFlags |= CPF_Parm;
 
 		if (bExpectCommaBeforeName)
@@ -4835,7 +4835,7 @@ void FHeaderParser::ParseParameterList(FClasses& AllClasses, UFunction* Function
 			RequireSymbol(TEXT(","), TEXT("Delegate definitions require a , between the parameter type and parameter name"));
 		}
 
-		UProperty* Prop = GetVarNameAndDim(Function, Property, ObjectFlags, /*NoArrays=*/ false, /*IsFunction=*/ true, NULL, TEXT("Function parameter"));
+		UProperty* Prop = GetVarNameAndDim(Function, Property, /*NoArrays=*/ false, /*IsFunction=*/ true, NULL, TEXT("Function parameter"), VariableCategory);
 
 		Function->NumParms++;
 
@@ -5012,12 +5012,11 @@ void FHeaderParser::CompileDelegateDeclaration(FUnrealSourceFile& SourceFile, FC
 	RequireSymbol(TEXT("("), CurrentScopeName);
 
 	// Parse the return value type
-	EObjectFlags ReturnValueFlags = RF_NoFlags;
 	FToken ReturnType( CPT_None );
 
 	if (bHasReturnValue)
 	{
-		GetVarType(AllClasses, GetCurrentScope(), ReturnType, ReturnValueFlags, 0, NULL, NULL, EPropertyDeclarationStyle::None, EVariableCategory::Return);
+		GetVarType(AllClasses, GetCurrentScope(), ReturnType, 0, NULL, NULL, EPropertyDeclarationStyle::None, EVariableCategory::Return);
 		RequireSymbol(TEXT(","), CurrentScopeName);
 	}
 
@@ -5094,7 +5093,7 @@ void FHeaderParser::CompileDelegateDeclaration(FUnrealSourceFile& SourceFile, FC
 	if (bHasReturnValue)
 	{
 		ReturnType.PropertyFlags |= CPF_Parm | CPF_OutParm | CPF_ReturnParm;
-		UProperty* ReturnProp = GetVarNameAndDim(DelegateSignatureFunction, ReturnType, ReturnValueFlags, /*NoArrays=*/ true, /*IsFunction=*/ true, TEXT("ReturnValue"), TEXT("Function return type"));
+		UProperty* ReturnProp = GetVarNameAndDim(DelegateSignatureFunction, ReturnType, /*NoArrays=*/ true, /*IsFunction=*/ true, TEXT("ReturnValue"), TEXT("Function return type"), EVariableCategory::Return);
 
 		DelegateSignatureFunction->NumParms++;
 	}
@@ -5341,14 +5340,13 @@ void FHeaderParser::CompileFunctionDeclaration(FUnrealSourceFile& SourceFile, FC
 	}
 
 	// Get return type.
-	EObjectFlags ReturnValueFlags = RF_NoFlags;
 	FToken ReturnType( CPT_None );
 	bool bHasReturnValue = false;
 
 	// C++ style functions always have a return value type, even if it's void
 	if (!MatchIdentifier(TEXT("void")))
 	{
-		bHasReturnValue = GetVarType(AllClasses, GetCurrentScope(), ReturnType, ReturnValueFlags, 0, NULL, NULL, EPropertyDeclarationStyle::None, EVariableCategory::Return);
+		bHasReturnValue = GetVarType(AllClasses, GetCurrentScope(), ReturnType, 0, NULL, NULL, EPropertyDeclarationStyle::None, EVariableCategory::Return);
 	}
 
 	// Get function or operator name.
@@ -5421,7 +5419,7 @@ void FHeaderParser::CompileFunctionDeclaration(FUnrealSourceFile& SourceFile, FC
 	if (bHasReturnValue)
 	{
 		ReturnType.PropertyFlags |= CPF_Parm | CPF_OutParm | CPF_ReturnParm;
-		UProperty* ReturnProp = GetVarNameAndDim(TopFunction, ReturnType, ReturnValueFlags, /*NoArrays=*/ true, /*IsFunction=*/ true, TEXT("ReturnValue"), TEXT("Function return type"));
+		UProperty* ReturnProp = GetVarNameAndDim(TopFunction, ReturnType, /*NoArrays=*/ true, /*IsFunction=*/ true, TEXT("ReturnValue"), TEXT("Function return type"), EVariableCategory::Return);
 
 		TopFunction->NumParms++;
 	}
@@ -5905,10 +5903,9 @@ void FHeaderParser::CompileVariableDeclaration(FClasses& AllClasses, UStruct* St
 
 	// Get variable type.
 	FPropertyBase OriginalProperty(CPT_None);
-	EObjectFlags ObjectFlags = RF_NoFlags;
 
 	FIndexRange TypeRange;
-	GetVarType( AllClasses, &FScope::GetTypeScope(Struct).Get(), OriginalProperty, ObjectFlags, DisallowFlags, TEXT("Member variable declaration"), /*OuterPropertyType=*/ NULL, EPropertyDeclarationStyle::UPROPERTY, EVariableCategory::Member, &TypeRange );
+	GetVarType( AllClasses, &FScope::GetTypeScope(Struct).Get(), OriginalProperty, DisallowFlags, TEXT("Member variable declaration"), /*OuterPropertyType=*/ NULL, EPropertyDeclarationStyle::UPROPERTY, EVariableCategory::Member, &TypeRange );
 	OriginalProperty.PropertyFlags |= EdFlags;
 
 	FString* Category = OriginalProperty.MetaData.Find("Category");
@@ -5977,7 +5974,7 @@ void FHeaderParser::CompileVariableDeclaration(FClasses& AllClasses, UStruct* St
 	do
 	{
 		FToken     Property    = OriginalProperty;
-		UProperty* NewProperty = GetVarNameAndDim(Struct, Property, ObjectFlags, /*NoArrays=*/ false, /*IsFunction=*/ false, NULL, TEXT("Member variable declaration"));
+		UProperty* NewProperty = GetVarNameAndDim(Struct, Property, /*NoArrays=*/ false, /*IsFunction=*/ false, NULL, TEXT("Member variable declaration"), EVariableCategory::Member);
 
 		// Optionally consume the :1 at the end of a bitfield boolean declaration
 		if (Property.IsBool() && MatchSymbol(TEXT(":")))
diff --git a/Engine/Source/Programs/UnrealHeaderTool/Private/HeaderParser.h b/Engine/Source/Programs/UnrealHeaderTool/Private/HeaderParser.h
--- a/Engine/Source/Programs/UnrealHeaderTool/Private/HeaderParser.h
+++ b/Engine/Source/Programs/UnrealHeaderTool/Private/HeaderParser.h
@@ -582,8 +582,6 @@ protected:
 	 * @param   AllClasses                the class tree for CurrentPackage
 	 * @param   Scope                     struct to create the property in
 	 * @param   VarProperty               will be filled in with type and property flag data for the property declaration that was parsed
-	 * @param   ObjectFlags               will contain the object flags that should be assigned to all UProperties which are part of this declaration (will not be determined
-	 *                                    until GetVarNameAndDim is called for this declaration).
 	 * @param   Disallow                  contains a mask of variable modifiers that are disallowed in this context
 	 * @param   Thing                     used for compiler errors to provide more information about the type of parsing that was occurring
 	 * @param   OuterPropertyType         only specified when compiling the inner properties for arrays or maps.  corresponds to the FToken for the outer property declaration.
@@ -597,7 +595,6 @@ protected:
 		FClasses&                       AllClasses,
 		FScope*							Scope,
 		FPropertyBase&                  VarProperty,
-		EObjectFlags&                   ObjectFlags,
 		uint64                          Disallow,
 		const TCHAR*                    Thing,
 		FToken*                         OuterPropertyType,
@@ -608,25 +605,25 @@ protected:
 	/**
 	 * Parses a variable name declaration and creates a new UProperty object.
 	 *
-	 * @param	Scope			struct to create the property in
-	 * @param	VarProperty		type and propertyflag info for the new property (inout)
-	 * @param	ObjectFlags		flags to pass on to the new property
-	 * @param	NoArrays		true if static arrays are disallowed
-	 * @param	IsFunction		true if the property is a function parameter or return value
-	 * @param	HardcodedName	name to assign to the new UProperty, if specified. primarily used for function return values,
-	 *							which are automatically called "ReturnValue"
-	 * @param	HintText		text to use in error message if error is encountered
+	 * @param	Scope				struct to create the property in
+	 * @param	VarProperty			type and propertyflag info for the new property (inout)
+	 * @param	NoArrays			true if static arrays are disallowed
+	 * @param	IsFunction			true if the property is a function parameter or return value
+	 * @param	HardcodedName		name to assign to the new UProperty, if specified. primarily used for function return values,
+	 *								which are automatically called "ReturnValue"
+	 * @param	HintText			text to use in error message if error is encountered
+	 * @param   VariableCategory	what kind of variable is being created
 	 *
 	 * @return	a pointer to the new UProperty if successful, or NULL if there was no property to parse
 	 */
 	UProperty* GetVarNameAndDim(
 		UStruct* Struct,
 		FToken& VarProperty,
-		EObjectFlags ObjectFlags,
 		bool NoArrays,
 		bool IsFunction,
 		const TCHAR* HardcodedName,
-		const TCHAR* Thing);
+		const TCHAR* Thing,
+		EVariableCategory::Type VariableCategory);
 	
 	/**
 	 * Returns whether the specified class can be referenced from the class currently being compiled.
