Error contexts determined from EVariableCategory.
HintText/Thing parameters removed from GetVarNameAndDim/GetVarType.

diff --git a/Engine/Source/Programs/UnrealHeaderTool/Private/HeaderParser.cpp b/Engine/Source/Programs/UnrealHeaderTool/Private/HeaderParser.cpp
--- a/Engine/Source/Programs/UnrealHeaderTool/Private/HeaderParser.cpp
+++ b/Engine/Source/Programs/UnrealHeaderTool/Private/HeaderParser.cpp
@@ -389,6 +389,28 @@ namespace
 	{
 		MetaData.Add(TEXT("EditInline"), TEXT("true"));
 	}
+
+	const TCHAR* GetHintText(EVariableCategory::Type VariableCategory)
+	{
+		switch (VariableCategory)
+		{
+			case EVariableCategory::ReplicatedParameter:
+			case EVariableCategory::RegularParameter:
+				return TEXT("Function parameter");
+
+			case EVariableCategory::Return:
+				return TEXT("Function return type");
+
+			case EVariableCategory::Member:
+				return TEXT("Member variable declaration");
+
+			default:
+				FError::Throwf(TEXT("Unknown variable category"));
+		}
+
+		// Unreachable
+		return nullptr;
+	}
 }
 	
 /////////////////////////////////////////////////////
@@ -2136,7 +2158,6 @@ bool FHeaderParser::GetVarType
 	FScope*							Scope,
 	FPropertyBase&                  VarProperty,
 	uint64                          Disallow,
-	const TCHAR*                    Thing,
 	FToken*                         OuterPropertyType,
 	EPropertyDeclarationStyle::Type PropertyDeclarationStyle,
 	EVariableCategory::Type         VariableCategory,
@@ -2515,7 +2536,7 @@ bool FHeaderParser::GetVarType
 	{
 		if (VariableCategory == EVariableCategory::Member)
 		{
-			FError::Throwf(TEXT("%s: Cannot declare enum at variable declaration"), Thing );
+			FError::Throwf(TEXT("%s: Cannot declare enum at variable declaration"), GetHintText(VariableCategory));
 		}
 
 		bUnconsumedEnumKeyword = true;
@@ -2525,12 +2546,7 @@ bool FHeaderParser::GetVarType
 	FToken VarType;
 	if ( !GetIdentifier(VarType,1) )
 	{
-		if ( !Thing )
-		{
-			check(!MetaDataFromNewStyle.Num());
-			return 0;
-		}
-		FError::Throwf(TEXT("%s: Missing variable type"), Thing );
+		FError::Throwf(TEXT("%s: Missing variable type"), GetHintText(VariableCategory));
 	}
 
 	if ( VarType.Matches(TEXT("int8")) )
@@ -2614,7 +2630,7 @@ bool FHeaderParser::GetVarType
 		uint64 OriginalVarTypeFlags = VarType.PropertyFlags;
 		VarType.PropertyFlags |= Flags;
 
-		GetVarType(AllClasses, Scope, VarProperty, Disallow, TEXT("'tarray'"), &VarType, EPropertyDeclarationStyle::None, VariableCategory);
+		GetVarType(AllClasses, Scope, VarProperty, Disallow, &VarType, EPropertyDeclarationStyle::None, VariableCategory);
 		if (VarProperty.ArrayType != EArrayType::None)
 		{
 			FError::Throwf(TEXT("Arrays within arrays not supported.") );
@@ -3216,7 +3232,6 @@ UProperty* FHeaderParser::GetVarNameAndDim
 	bool                    NoArrays,
 	bool                    IsFunction,
 	const TCHAR*            HardcodedName,
-	const TCHAR*            HintText,
 	EVariableCategory::Type VariableCategory
 )
 {
@@ -3228,6 +3243,8 @@ UProperty* FHeaderParser::GetVarNameAndDim
 		ObjectFlags = RF_NoFlags;
 	}
 
+	const TCHAR* HintText = GetHintText(VariableCategory);
+
 	AddModuleRelativePathToMetadata(Scope, VarProperty.MetaData);
 
 	// Get variable name.
@@ -4827,7 +4844,7 @@ void FHeaderParser::ParseParameterList(FClasses& AllClasses, UFunction* Function
 		// Get parameter type.
 		FToken Property(CPT_None);
 		EVariableCategory::Type VariableCategory = (Function->FunctionFlags & FUNC_Net) ? EVariableCategory::ReplicatedParameter : EVariableCategory::RegularParameter;
-		GetVarType(AllClasses, GetCurrentScope(), Property, ~(CPF_ParmFlags | CPF_AutoWeak | CPF_RepSkip), TEXT("Function parameter"), NULL, EPropertyDeclarationStyle::None, VariableCategory);
+		GetVarType(AllClasses, GetCurrentScope(), Property, ~(CPF_ParmFlags | CPF_AutoWeak | CPF_RepSkip), NULL, EPropertyDeclarationStyle::None, VariableCategory);
 		Property.PropertyFlags |= CPF_Parm;
 
 		if (bExpectCommaBeforeName)
@@ -4835,7 +4852,7 @@ void FHeaderParser::ParseParameterList(FClasses& AllClasses, UFunction* Function
 			RequireSymbol(TEXT(","), TEXT("Delegate definitions require a , between the parameter type and parameter name"));
 		}
 
-		UProperty* Prop = GetVarNameAndDim(Function, Property, /*NoArrays=*/ false, /*IsFunction=*/ true, NULL, TEXT("Function parameter"), VariableCategory);
+		UProperty* Prop = GetVarNameAndDim(Function, Property, /*NoArrays=*/ false, /*IsFunction=*/ true, NULL, VariableCategory);
 
 		Function->NumParms++;
 
@@ -5016,7 +5033,7 @@ void FHeaderParser::CompileDelegateDeclaration(FUnrealSourceFile& SourceFile, FC
 
 	if (bHasReturnValue)
 	{
-		GetVarType(AllClasses, GetCurrentScope(), ReturnType, 0, NULL, NULL, EPropertyDeclarationStyle::None, EVariableCategory::Return);
+		GetVarType(AllClasses, GetCurrentScope(), ReturnType, 0, NULL, EPropertyDeclarationStyle::None, EVariableCategory::Return);
 		RequireSymbol(TEXT(","), CurrentScopeName);
 	}
 
@@ -5093,7 +5110,7 @@ void FHeaderParser::CompileDelegateDeclaration(FUnrealSourceFile& SourceFile, FC
 	if (bHasReturnValue)
 	{
 		ReturnType.PropertyFlags |= CPF_Parm | CPF_OutParm | CPF_ReturnParm;
-		UProperty* ReturnProp = GetVarNameAndDim(DelegateSignatureFunction, ReturnType, /*NoArrays=*/ true, /*IsFunction=*/ true, TEXT("ReturnValue"), TEXT("Function return type"), EVariableCategory::Return);
+		UProperty* ReturnProp = GetVarNameAndDim(DelegateSignatureFunction, ReturnType, /*NoArrays=*/ true, /*IsFunction=*/ true, TEXT("ReturnValue"), EVariableCategory::Return);
 
 		DelegateSignatureFunction->NumParms++;
 	}
@@ -5346,7 +5363,7 @@ void FHeaderParser::CompileFunctionDeclaration(FUnrealSourceFile& SourceFile, FC
 	// C++ style functions always have a return value type, even if it's void
 	if (!MatchIdentifier(TEXT("void")))
 	{
-		bHasReturnValue = GetVarType(AllClasses, GetCurrentScope(), ReturnType, 0, NULL, NULL, EPropertyDeclarationStyle::None, EVariableCategory::Return);
+		bHasReturnValue = GetVarType(AllClasses, GetCurrentScope(), ReturnType, 0, NULL, EPropertyDeclarationStyle::None, EVariableCategory::Return);
 	}
 
 	// Get function or operator name.
@@ -5419,7 +5436,7 @@ void FHeaderParser::CompileFunctionDeclaration(FUnrealSourceFile& SourceFile, FC
 	if (bHasReturnValue)
 	{
 		ReturnType.PropertyFlags |= CPF_Parm | CPF_OutParm | CPF_ReturnParm;
-		UProperty* ReturnProp = GetVarNameAndDim(TopFunction, ReturnType, /*NoArrays=*/ true, /*IsFunction=*/ true, TEXT("ReturnValue"), TEXT("Function return type"), EVariableCategory::Return);
+		UProperty* ReturnProp = GetVarNameAndDim(TopFunction, ReturnType, /*NoArrays=*/ true, /*IsFunction=*/ true, TEXT("ReturnValue"), EVariableCategory::Return);
 
 		TopFunction->NumParms++;
 	}
@@ -5905,7 +5922,7 @@ void FHeaderParser::CompileVariableDeclaration(FClasses& AllClasses, UStruct* St
 	FPropertyBase OriginalProperty(CPT_None);
 
 	FIndexRange TypeRange;
-	GetVarType( AllClasses, &FScope::GetTypeScope(Struct).Get(), OriginalProperty, DisallowFlags, TEXT("Member variable declaration"), /*OuterPropertyType=*/ NULL, EPropertyDeclarationStyle::UPROPERTY, EVariableCategory::Member, &TypeRange );
+	GetVarType( AllClasses, &FScope::GetTypeScope(Struct).Get(), OriginalProperty, DisallowFlags, /*OuterPropertyType=*/ NULL, EPropertyDeclarationStyle::UPROPERTY, EVariableCategory::Member, &TypeRange );
 	OriginalProperty.PropertyFlags |= EdFlags;
 
 	FString* Category = OriginalProperty.MetaData.Find("Category");
@@ -5974,7 +5991,7 @@ void FHeaderParser::CompileVariableDeclaration(FClasses& AllClasses, UStruct* St
 	do
 	{
 		FToken     Property    = OriginalProperty;
-		UProperty* NewProperty = GetVarNameAndDim(Struct, Property, /*NoArrays=*/ false, /*IsFunction=*/ false, NULL, TEXT("Member variable declaration"), EVariableCategory::Member);
+		UProperty* NewProperty = GetVarNameAndDim(Struct, Property, /*NoArrays=*/ false, /*IsFunction=*/ false, NULL, EVariableCategory::Member);
 
 		// Optionally consume the :1 at the end of a bitfield boolean declaration
 		if (Property.IsBool() && MatchSymbol(TEXT(":")))
diff --git a/Engine/Source/Programs/UnrealHeaderTool/Private/HeaderParser.h b/Engine/Source/Programs/UnrealHeaderTool/Private/HeaderParser.h
--- a/Engine/Source/Programs/UnrealHeaderTool/Private/HeaderParser.h
+++ b/Engine/Source/Programs/UnrealHeaderTool/Private/HeaderParser.h
@@ -45,6 +45,18 @@ namespace EDelegateSpecifierAction
 	};
 }
 
+/** The category of variable declaration being parsed */
+namespace EVariableCategory
+{
+	enum Type
+	{
+		RegularParameter,
+		ReplicatedParameter,
+		Return,
+		Member
+	};
+}
+
 /** Information for a particular nesting level. */
 class FNestInfo
 {
@@ -564,18 +576,6 @@ protected:
 	UField* FindField( UStruct* InScope, const TCHAR* InIdentifier, bool bIncludeParents=true, UClass* FieldClass=UField::StaticClass(), const TCHAR* Thing=NULL );
 	void SkipStatements( int32 SubCount, const TCHAR* ErrorTag );
 
-	/** The category of variable declaration being parsed */
-	struct EVariableCategory
-	{
-		enum Type
-		{
-			RegularParameter,
-			ReplicatedParameter,
-			Return,
-			Member
-		};
-	};
-
 	/**
 	 * Parses a variable or return value declaration and determines the variable type and property flags.
 	 *
@@ -583,7 +583,6 @@ protected:
 	 * @param   Scope                     struct to create the property in
 	 * @param   VarProperty               will be filled in with type and property flag data for the property declaration that was parsed
 	 * @param   Disallow                  contains a mask of variable modifiers that are disallowed in this context
-	 * @param   Thing                     used for compiler errors to provide more information about the type of parsing that was occurring
 	 * @param   OuterPropertyType         only specified when compiling the inner properties for arrays or maps.  corresponds to the FToken for the outer property declaration.
 	 * @param   PropertyDeclarationStyle  if the variable is defined with a UPROPERTY
 	 * @param   VariableCategory          what kind of variable is being parsed
@@ -596,7 +595,6 @@ protected:
 		FScope*							Scope,
 		FPropertyBase&                  VarProperty,
 		uint64                          Disallow,
-		const TCHAR*                    Thing,
 		FToken*                         OuterPropertyType,
 		EPropertyDeclarationStyle::Type PropertyDeclarationStyle,
 		EVariableCategory::Type         VariableCategory,
@@ -611,7 +609,6 @@ protected:
 	 * @param	IsFunction			true if the property is a function parameter or return value
 	 * @param	HardcodedName		name to assign to the new UProperty, if specified. primarily used for function return values,
 	 *								which are automatically called "ReturnValue"
-	 * @param	HintText			text to use in error message if error is encountered
 	 * @param   VariableCategory	what kind of variable is being created
 	 *
 	 * @return	a pointer to the new UProperty if successful, or NULL if there was no property to parse
@@ -622,7 +619,6 @@ protected:
 		bool NoArrays,
 		bool IsFunction,
 		const TCHAR* HardcodedName,
-		const TCHAR* Thing,
 		EVariableCategory::Type VariableCategory);
 	
 	/**
