Redundant GetVarType return value removed.

diff --git a/Engine/Source/Programs/UnrealHeaderTool/Private/HeaderParser.cpp b/Engine/Source/Programs/UnrealHeaderTool/Private/HeaderParser.cpp
--- a/Engine/Source/Programs/UnrealHeaderTool/Private/HeaderParser.cpp
+++ b/Engine/Source/Programs/UnrealHeaderTool/Private/HeaderParser.cpp
@@ -2152,7 +2152,7 @@ void FHeaderParser::CompileDirective(FClasses& AllClasses, FUnrealSourceFile& So
 	Variable declaration parser.
 -----------------------------------------------------------------------------*/
 
-bool FHeaderParser::GetVarType
+void FHeaderParser::GetVarType
 (
 	FClasses&                       AllClasses,
 	FScope*							Scope,
@@ -3189,7 +3189,6 @@ bool FHeaderParser::GetVarType
 	{
 		ParsedVarIndexRange->Count = InputPos - ParsedVarIndexRange->StartIndex;
 	}
-	return 1;
 }
 
 /**
@@ -5358,12 +5357,12 @@ void FHeaderParser::CompileFunctionDeclaration(FUnrealSourceFile& SourceFile, FC
 
 	// Get return type.
 	FToken ReturnType( CPT_None );
-	bool bHasReturnValue = false;
 
 	// C++ style functions always have a return value type, even if it's void
-	if (!MatchIdentifier(TEXT("void")))
+	bool bHasReturnValue = !MatchIdentifier(TEXT("void"));
+	if (bHasReturnValue)
 	{
-		bHasReturnValue = GetVarType(AllClasses, GetCurrentScope(), ReturnType, 0, NULL, EPropertyDeclarationStyle::None, EVariableCategory::Return);
+		GetVarType(AllClasses, GetCurrentScope(), ReturnType, 0, NULL, EPropertyDeclarationStyle::None, EVariableCategory::Return);
 	}
 
 	// Get function or operator name.
diff --git a/Engine/Source/Programs/UnrealHeaderTool/Private/HeaderParser.h b/Engine/Source/Programs/UnrealHeaderTool/Private/HeaderParser.h
--- a/Engine/Source/Programs/UnrealHeaderTool/Private/HeaderParser.h
+++ b/Engine/Source/Programs/UnrealHeaderTool/Private/HeaderParser.h
@@ -587,10 +587,8 @@ protected:
 	 * @param   PropertyDeclarationStyle  if the variable is defined with a UPROPERTY
 	 * @param   VariableCategory          what kind of variable is being parsed
 	 * @param   ParsedVarIndexRange       The source text [Start, End) index range for the parsed type.
-	 *
-	 * @return  true if the variable type was parsed
 	 */
-	bool GetVarType(
+	void GetVarType(
 		FClasses&                       AllClasses,
 		FScope*							Scope,
 		FPropertyBase&                  VarProperty,
