GetVarAndDim's HardcodedName parameter removed.

diff --git a/Engine/Source/Programs/UnrealHeaderTool/Private/HeaderParser.cpp b/Engine/Source/Programs/UnrealHeaderTool/Private/HeaderParser.cpp
--- a/Engine/Source/Programs/UnrealHeaderTool/Private/HeaderParser.cpp
+++ b/Engine/Source/Programs/UnrealHeaderTool/Private/HeaderParser.cpp
@@ -3230,7 +3230,6 @@ UProperty* FHeaderParser::GetVarNameAndDim
 	FToken&                 VarProperty,
 	bool                    NoArrays,
 	bool                    IsFunction,
-	const TCHAR*            HardcodedName,
 	EVariableCategory::Type VariableCategory
 )
 {
@@ -3247,11 +3246,11 @@ UProperty* FHeaderParser::GetVarNameAndDim
 	AddModuleRelativePathToMetadata(Scope, VarProperty.MetaData);
 
 	// Get variable name.
-	if (HardcodedName != NULL)
+	if (VariableCategory == EVariableCategory::Return)
 	{
 		// Hard-coded variable name, such as with return value.
 		VarProperty.TokenType = TOKEN_Identifier;
-		FCString::Strcpy( VarProperty.Identifier, HardcodedName );
+		FCString::Strcpy( VarProperty.Identifier, TEXT("ReturnValue") );
 	}
 	else
 	{
@@ -3456,7 +3455,7 @@ UProperty* FHeaderParser::GetVarNameAndDim
 		{
 			UBoolProperty* NewBoolProperty = new(EC_InternalUseOnlyConstructor, NewScope, PropertyName, ObjectFlags)UBoolProperty(FObjectInitializer());
 			NewProperty = NewBoolProperty;
-			if (HardcodedName && FCString::Stricmp(HardcodedName, TEXT("ReturnValue")) == 0)
+			if (VariableCategory == EVariableCategory::Return)
 			{
 				NewBoolProperty->SetBoolSize(sizeof(bool), true);
 			}
@@ -4851,7 +4850,7 @@ void FHeaderParser::ParseParameterList(FClasses& AllClasses, UFunction* Function
 			RequireSymbol(TEXT(","), TEXT("Delegate definitions require a , between the parameter type and parameter name"));
 		}
 
-		UProperty* Prop = GetVarNameAndDim(Function, Property, /*NoArrays=*/ false, /*IsFunction=*/ true, NULL, VariableCategory);
+		UProperty* Prop = GetVarNameAndDim(Function, Property, /*NoArrays=*/ false, /*IsFunction=*/ true, VariableCategory);
 
 		Function->NumParms++;
 
@@ -5109,7 +5108,7 @@ void FHeaderParser::CompileDelegateDeclaration(FUnrealSourceFile& SourceFile, FC
 	if (bHasReturnValue)
 	{
 		ReturnType.PropertyFlags |= CPF_Parm | CPF_OutParm | CPF_ReturnParm;
-		UProperty* ReturnProp = GetVarNameAndDim(DelegateSignatureFunction, ReturnType, /*NoArrays=*/ true, /*IsFunction=*/ true, TEXT("ReturnValue"), EVariableCategory::Return);
+		UProperty* ReturnProp = GetVarNameAndDim(DelegateSignatureFunction, ReturnType, /*NoArrays=*/ true, /*IsFunction=*/ true, EVariableCategory::Return);
 
 		DelegateSignatureFunction->NumParms++;
 	}
@@ -5435,7 +5434,7 @@ void FHeaderParser::CompileFunctionDeclaration(FUnrealSourceFile& SourceFile, FC
 	if (bHasReturnValue)
 	{
 		ReturnType.PropertyFlags |= CPF_Parm | CPF_OutParm | CPF_ReturnParm;
-		UProperty* ReturnProp = GetVarNameAndDim(TopFunction, ReturnType, /*NoArrays=*/ true, /*IsFunction=*/ true, TEXT("ReturnValue"), EVariableCategory::Return);
+		UProperty* ReturnProp = GetVarNameAndDim(TopFunction, ReturnType, /*NoArrays=*/ true, /*IsFunction=*/ true, EVariableCategory::Return);
 
 		TopFunction->NumParms++;
 	}
@@ -5990,7 +5989,7 @@ void FHeaderParser::CompileVariableDeclaration(FClasses& AllClasses, UStruct* St
 	do
 	{
 		FToken     Property    = OriginalProperty;
-		UProperty* NewProperty = GetVarNameAndDim(Struct, Property, /*NoArrays=*/ false, /*IsFunction=*/ false, NULL, EVariableCategory::Member);
+		UProperty* NewProperty = GetVarNameAndDim(Struct, Property, /*NoArrays=*/ false, /*IsFunction=*/ false, EVariableCategory::Member);
 
 		// Optionally consume the :1 at the end of a bitfield boolean declaration
 		if (Property.IsBool() && MatchSymbol(TEXT(":")))
diff --git a/Engine/Source/Programs/UnrealHeaderTool/Private/HeaderParser.h b/Engine/Source/Programs/UnrealHeaderTool/Private/HeaderParser.h
--- a/Engine/Source/Programs/UnrealHeaderTool/Private/HeaderParser.h
+++ b/Engine/Source/Programs/UnrealHeaderTool/Private/HeaderParser.h
@@ -616,7 +616,6 @@ protected:
 		FToken& VarProperty,
 		bool NoArrays,
 		bool IsFunction,
-		const TCHAR* HardcodedName,
 		EVariableCategory::Type VariableCategory);
 	
 	/**
