IsFunction parameter removed from GetVarNameAndDim.

diff --git a/Engine/Source/Programs/UnrealHeaderTool/Private/HeaderParser.cpp b/Engine/Source/Programs/UnrealHeaderTool/Private/HeaderParser.cpp
--- a/Engine/Source/Programs/UnrealHeaderTool/Private/HeaderParser.cpp
+++ b/Engine/Source/Programs/UnrealHeaderTool/Private/HeaderParser.cpp
@@ -3229,7 +3229,6 @@ UProperty* FHeaderParser::GetVarNameAndDim
 	UStruct*                Scope,
 	FToken&                 VarProperty,
 	bool                    NoArrays,
-	bool                    IsFunction,
 	EVariableCategory::Type VariableCategory
 )
 {
@@ -3371,14 +3370,13 @@ UProperty* FHeaderParser::GetVarNameAndDim
 	}
 
 	// Try gathering metadata for member fields
-	if (!IsFunction)
+	if (VariableCategory == EVariableCategory::Member)
 	{
 		ParseFieldMetaData(VarProperty.MetaData, VarProperty.Identifier);
 		AddFormattedPrevCommentAsTooltipMetaData(VarProperty.MetaData);
 	}
-
 	// validate UFunction parameters
-	if (IsFunction)
+	else
 	{
 		// UFunctions with a smart pointer as input parameter wont compile anyway, because of missing P_GET_... macro.
 		// UFunctions with a smart pointer as return type will crash when called via blueprint, because they are not supported in VM.
@@ -4850,7 +4848,7 @@ void FHeaderParser::ParseParameterList(FClasses& AllClasses, UFunction* Function
 			RequireSymbol(TEXT(","), TEXT("Delegate definitions require a , between the parameter type and parameter name"));
 		}
 
-		UProperty* Prop = GetVarNameAndDim(Function, Property, /*NoArrays=*/ false, /*IsFunction=*/ true, VariableCategory);
+		UProperty* Prop = GetVarNameAndDim(Function, Property, /*NoArrays=*/ false, VariableCategory);
 
 		Function->NumParms++;
 
@@ -5108,7 +5106,7 @@ void FHeaderParser::CompileDelegateDeclaration(FUnrealSourceFile& SourceFile, FC
 	if (bHasReturnValue)
 	{
 		ReturnType.PropertyFlags |= CPF_Parm | CPF_OutParm | CPF_ReturnParm;
-		UProperty* ReturnProp = GetVarNameAndDim(DelegateSignatureFunction, ReturnType, /*NoArrays=*/ true, /*IsFunction=*/ true, EVariableCategory::Return);
+		UProperty* ReturnProp = GetVarNameAndDim(DelegateSignatureFunction, ReturnType, /*NoArrays=*/ true, EVariableCategory::Return);
 
 		DelegateSignatureFunction->NumParms++;
 	}
@@ -5434,7 +5432,7 @@ void FHeaderParser::CompileFunctionDeclaration(FUnrealSourceFile& SourceFile, FC
 	if (bHasReturnValue)
 	{
 		ReturnType.PropertyFlags |= CPF_Parm | CPF_OutParm | CPF_ReturnParm;
-		UProperty* ReturnProp = GetVarNameAndDim(TopFunction, ReturnType, /*NoArrays=*/ true, /*IsFunction=*/ true, EVariableCategory::Return);
+		UProperty* ReturnProp = GetVarNameAndDim(TopFunction, ReturnType, /*NoArrays=*/ true, EVariableCategory::Return);
 
 		TopFunction->NumParms++;
 	}
@@ -5989,7 +5987,7 @@ void FHeaderParser::CompileVariableDeclaration(FClasses& AllClasses, UStruct* St
 	do
 	{
 		FToken     Property    = OriginalProperty;
-		UProperty* NewProperty = GetVarNameAndDim(Struct, Property, /*NoArrays=*/ false, /*IsFunction=*/ false, EVariableCategory::Member);
+		UProperty* NewProperty = GetVarNameAndDim(Struct, Property, /*NoArrays=*/ false, EVariableCategory::Member);
 
 		// Optionally consume the :1 at the end of a bitfield boolean declaration
 		if (Property.IsBool() && MatchSymbol(TEXT(":")))
diff --git a/Engine/Source/Programs/UnrealHeaderTool/Private/HeaderParser.h b/Engine/Source/Programs/UnrealHeaderTool/Private/HeaderParser.h
--- a/Engine/Source/Programs/UnrealHeaderTool/Private/HeaderParser.h
+++ b/Engine/Source/Programs/UnrealHeaderTool/Private/HeaderParser.h
@@ -604,9 +604,6 @@ protected:
 	 * @param	Scope				struct to create the property in
 	 * @param	VarProperty			type and propertyflag info for the new property (inout)
 	 * @param	NoArrays			true if static arrays are disallowed
-	 * @param	IsFunction			true if the property is a function parameter or return value
-	 * @param	HardcodedName		name to assign to the new UProperty, if specified. primarily used for function return values,
-	 *								which are automatically called "ReturnValue"
 	 * @param   VariableCategory	what kind of variable is being created
 	 *
 	 * @return	a pointer to the new UProperty if successful, or NULL if there was no property to parse
@@ -615,7 +612,6 @@ protected:
 		UStruct* Struct,
 		FToken& VarProperty,
 		bool NoArrays,
-		bool IsFunction,
 		EVariableCategory::Type VariableCategory);
 	
 	/**
