NoArrays parameter removed from GetVarNameAndDim.

diff --git a/Engine/Source/Programs/UnrealHeaderTool/Private/HeaderParser.cpp b/Engine/Source/Programs/UnrealHeaderTool/Private/HeaderParser.cpp
--- a/Engine/Source/Programs/UnrealHeaderTool/Private/HeaderParser.cpp
+++ b/Engine/Source/Programs/UnrealHeaderTool/Private/HeaderParser.cpp
@@ -3228,7 +3228,6 @@ UProperty* FHeaderParser::GetVarNameAndDim
 (
 	UStruct*                Scope,
 	FToken&                 VarProperty,
-	bool                    NoArrays,
 	EVariableCategory::Type VariableCategory
 )
 {
@@ -3323,7 +3322,7 @@ UProperty* FHeaderParser::GetVarNameAndDim
 	FToken Dimensions;
 	if (MatchSymbol(TEXT("[")))
 	{
-		if (NoArrays)
+		if (VariableCategory == EVariableCategory::Return)
 		{
 			FError::Throwf(TEXT("Arrays aren't allowed in this context") );
 		}
@@ -4848,7 +4847,7 @@ void FHeaderParser::ParseParameterList(FClasses& AllClasses, UFunction* Function
 			RequireSymbol(TEXT(","), TEXT("Delegate definitions require a , between the parameter type and parameter name"));
 		}
 
-		UProperty* Prop = GetVarNameAndDim(Function, Property, /*NoArrays=*/ false, VariableCategory);
+		UProperty* Prop = GetVarNameAndDim(Function, Property, VariableCategory);
 
 		Function->NumParms++;
 
@@ -5106,7 +5105,7 @@ void FHeaderParser::CompileDelegateDeclaration(FUnrealSourceFile& SourceFile, FC
 	if (bHasReturnValue)
 	{
 		ReturnType.PropertyFlags |= CPF_Parm | CPF_OutParm | CPF_ReturnParm;
-		UProperty* ReturnProp = GetVarNameAndDim(DelegateSignatureFunction, ReturnType, /*NoArrays=*/ true, EVariableCategory::Return);
+		UProperty* ReturnProp = GetVarNameAndDim(DelegateSignatureFunction, ReturnType, EVariableCategory::Return);
 
 		DelegateSignatureFunction->NumParms++;
 	}
@@ -5432,7 +5431,7 @@ void FHeaderParser::CompileFunctionDeclaration(FUnrealSourceFile& SourceFile, FC
 	if (bHasReturnValue)
 	{
 		ReturnType.PropertyFlags |= CPF_Parm | CPF_OutParm | CPF_ReturnParm;
-		UProperty* ReturnProp = GetVarNameAndDim(TopFunction, ReturnType, /*NoArrays=*/ true, EVariableCategory::Return);
+		UProperty* ReturnProp = GetVarNameAndDim(TopFunction, ReturnType, EVariableCategory::Return);
 
 		TopFunction->NumParms++;
 	}
@@ -5987,7 +5986,7 @@ void FHeaderParser::CompileVariableDeclaration(FClasses& AllClasses, UStruct* St
 	do
 	{
 		FToken     Property    = OriginalProperty;
-		UProperty* NewProperty = GetVarNameAndDim(Struct, Property, /*NoArrays=*/ false, EVariableCategory::Member);
+		UProperty* NewProperty = GetVarNameAndDim(Struct, Property, EVariableCategory::Member);
 
 		// Optionally consume the :1 at the end of a bitfield boolean declaration
 		if (Property.IsBool() && MatchSymbol(TEXT(":")))
diff --git a/Engine/Source/Programs/UnrealHeaderTool/Private/HeaderParser.h b/Engine/Source/Programs/UnrealHeaderTool/Private/HeaderParser.h
--- a/Engine/Source/Programs/UnrealHeaderTool/Private/HeaderParser.h
+++ b/Engine/Source/Programs/UnrealHeaderTool/Private/HeaderParser.h
@@ -603,7 +603,6 @@ protected:
 	 *
 	 * @param	Scope				struct to create the property in
 	 * @param	VarProperty			type and propertyflag info for the new property (inout)
-	 * @param	NoArrays			true if static arrays are disallowed
 	 * @param   VariableCategory	what kind of variable is being created
 	 *
 	 * @return	a pointer to the new UProperty if successful, or NULL if there was no property to parse
@@ -611,7 +610,6 @@ protected:
 	UProperty* GetVarNameAndDim(
 		UStruct* Struct,
 		FToken& VarProperty,
-		bool NoArrays,
 		EVariableCategory::Type VariableCategory);
 	
 	/**
