Nested switch broken up.

diff --git a/Engine/Source/Programs/UnrealHeaderTool/Private/HeaderParser.cpp b/Engine/Source/Programs/UnrealHeaderTool/Private/HeaderParser.cpp
--- a/Engine/Source/Programs/UnrealHeaderTool/Private/HeaderParser.cpp
+++ b/Engine/Source/Programs/UnrealHeaderTool/Private/HeaderParser.cpp
@@ -3415,165 +3415,159 @@ UProperty* FHeaderParser::GetVarNameAndDim
 			ArrayDim = 2;
 		}
 
-		if (VarProperty.Type == CPT_Byte)
+		switch (VarProperty.Type)
 		{
-			NewProperty = new(EC_InternalUseOnlyConstructor, NewScope, PropertyName, ObjectFlags)UByteProperty(FObjectInitializer());
-			((UByteProperty*)NewProperty)->Enum = VarProperty.Enum;
-		}
-		else if (VarProperty.Type == CPT_Int8)
-		{
-			NewProperty = new(EC_InternalUseOnlyConstructor, NewScope, PropertyName, ObjectFlags)UInt8Property(FObjectInitializer());
-		}
-		else if (VarProperty.Type == CPT_Int16)
-		{
-			NewProperty = new(EC_InternalUseOnlyConstructor, NewScope, PropertyName, ObjectFlags)UInt16Property(FObjectInitializer());
-		}
-		else if (VarProperty.Type == CPT_Int)
-		{
-			NewProperty = new(EC_InternalUseOnlyConstructor, NewScope, PropertyName, ObjectFlags)UIntProperty(FObjectInitializer());
-		}
-		else if (VarProperty.Type == CPT_Int64)
-		{
-			NewProperty = new(EC_InternalUseOnlyConstructor, NewScope, PropertyName, ObjectFlags)UInt64Property(FObjectInitializer());
-		}
-		else if (VarProperty.Type == CPT_UInt16)
-		{
-			NewProperty = new(EC_InternalUseOnlyConstructor, NewScope, PropertyName, ObjectFlags)UUInt16Property(FObjectInitializer());
-		}
-		else if (VarProperty.Type == CPT_UInt32)
-		{
-			NewProperty = new(EC_InternalUseOnlyConstructor, NewScope, PropertyName, ObjectFlags)UUInt32Property(FObjectInitializer());
-		}
-		else if (VarProperty.Type == CPT_UInt64)
-		{
-			NewProperty = new(EC_InternalUseOnlyConstructor, NewScope, PropertyName, ObjectFlags)UUInt64Property(FObjectInitializer());
-		}
-		else if (VarProperty.IsBool())
-		{
-			UBoolProperty* NewBoolProperty = new(EC_InternalUseOnlyConstructor, NewScope, PropertyName, ObjectFlags)UBoolProperty(FObjectInitializer());
-			NewProperty = NewBoolProperty;
-			if (VariableCategory == EVariableCategory::Return)
-			{
-				NewBoolProperty->SetBoolSize(sizeof(bool), true);
-			}
-			else
-			{
-				switch( VarProperty.Type )
+			case CPT_Byte:
+				NewProperty = new(EC_InternalUseOnlyConstructor, NewScope, PropertyName, ObjectFlags)UByteProperty(FObjectInitializer());
+				((UByteProperty*)NewProperty)->Enum = VarProperty.Enum;
+				break;
+
+			case CPT_Int8:
+				NewProperty = new(EC_InternalUseOnlyConstructor, NewScope, PropertyName, ObjectFlags)UInt8Property(FObjectInitializer());
+				break;
+
+			case CPT_Int16:
+				NewProperty = new(EC_InternalUseOnlyConstructor, NewScope, PropertyName, ObjectFlags)UInt16Property(FObjectInitializer());
+				break;
+
+			case CPT_Int:
+				NewProperty = new(EC_InternalUseOnlyConstructor, NewScope, PropertyName, ObjectFlags)UIntProperty(FObjectInitializer());
+				break;
+
+			case CPT_Int64:
+				NewProperty = new(EC_InternalUseOnlyConstructor, NewScope, PropertyName, ObjectFlags)UInt64Property(FObjectInitializer());
+				break;
+
+			case CPT_UInt16:
+				NewProperty = new(EC_InternalUseOnlyConstructor, NewScope, PropertyName, ObjectFlags)UUInt16Property(FObjectInitializer());
+				break;
+
+			case CPT_UInt32:
+				NewProperty = new(EC_InternalUseOnlyConstructor, NewScope, PropertyName, ObjectFlags)UUInt32Property(FObjectInitializer());
+				break;
+
+			case CPT_UInt64:
+				NewProperty = new(EC_InternalUseOnlyConstructor, NewScope, PropertyName, ObjectFlags)UUInt64Property(FObjectInitializer());
+				break;
+
+			case CPT_Bool:
+				NewProperty = new(EC_InternalUseOnlyConstructor, NewScope, PropertyName, ObjectFlags)UBoolProperty(FObjectInitializer());
+				((UBoolProperty*)NewProperty)->SetBoolSize(sizeof(bool), true);
+				break;
+
+			case CPT_Bool8:
+				NewProperty = new(EC_InternalUseOnlyConstructor, NewScope, PropertyName, ObjectFlags)UBoolProperty(FObjectInitializer());
+				((UBoolProperty*)NewProperty)->SetBoolSize((VariableCategory == EVariableCategory::Return) ? sizeof(bool) : sizeof(uint8), VariableCategory == EVariableCategory::Return);
+				break;
+
+			case CPT_Bool16:
+				NewProperty = new(EC_InternalUseOnlyConstructor, NewScope, PropertyName, ObjectFlags)UBoolProperty(FObjectInitializer());
+				((UBoolProperty*)NewProperty)->SetBoolSize((VariableCategory == EVariableCategory::Return) ? sizeof(bool) : sizeof(uint16), VariableCategory == EVariableCategory::Return);
+				break;
+
+			case CPT_Bool32:
+				NewProperty = new(EC_InternalUseOnlyConstructor, NewScope, PropertyName, ObjectFlags)UBoolProperty(FObjectInitializer());
+				((UBoolProperty*)NewProperty)->SetBoolSize((VariableCategory == EVariableCategory::Return) ? sizeof(bool) : sizeof(uint32), VariableCategory == EVariableCategory::Return);
+				break;
+
+			case CPT_Bool64:
+				NewProperty = new(EC_InternalUseOnlyConstructor, NewScope, PropertyName, ObjectFlags)UBoolProperty(FObjectInitializer());
+				((UBoolProperty*)NewProperty)->SetBoolSize((VariableCategory == EVariableCategory::Return) ? sizeof(bool) : sizeof(uint64), VariableCategory == EVariableCategory::Return);
+				break;
+
+			case CPT_Float:
+				NewProperty = new(EC_InternalUseOnlyConstructor, NewScope, PropertyName, ObjectFlags)UFloatProperty(FObjectInitializer());
+				break;
+
+			case CPT_Double:
+				NewProperty = new(EC_InternalUseOnlyConstructor, NewScope, PropertyName, ObjectFlags)UDoubleProperty(FObjectInitializer());
+				break;
+
+			case CPT_ObjectReference:
+				check(VarProperty.PropertyClass);
+				if (VarProperty.PropertyClass->IsChildOf(UClass::StaticClass()))
 				{
-				case CPT_Bool:
-					NewBoolProperty->SetBoolSize(sizeof(bool), true);
-					break;
-				case CPT_Bool8:
-					NewBoolProperty->SetBoolSize(sizeof(uint8));
-					break;
-				case CPT_Bool16:
-					NewBoolProperty->SetBoolSize(sizeof(uint16));
-					break;
-				case CPT_Bool32:
-					NewBoolProperty->SetBoolSize(sizeof(uint32));
-					break;
-				case CPT_Bool64:
-					NewBoolProperty->SetBoolSize(sizeof(uint64));
-					break;
+					NewProperty = new(EC_InternalUseOnlyConstructor, NewScope, PropertyName, ObjectFlags)UClassProperty(FObjectInitializer());
+					((UClassProperty*)NewProperty)->MetaClass = VarProperty.MetaClass;
 				}
-			}
-		}
-		else if (VarProperty.Type == CPT_Float)
-		{
-			NewProperty = new(EC_InternalUseOnlyConstructor, NewScope, PropertyName, ObjectFlags)UFloatProperty(FObjectInitializer());
-		}
-		else if (VarProperty.Type == CPT_Double)
-		{
-			NewProperty = new(EC_InternalUseOnlyConstructor, NewScope, PropertyName, ObjectFlags)UDoubleProperty(FObjectInitializer());
-		}
-		else if (VarProperty.Type == CPT_ObjectReference)
-		{
-			check(VarProperty.PropertyClass);
-			if (VarProperty.PropertyClass->IsChildOf(UClass::StaticClass()))
-			{
-				NewProperty = new(EC_InternalUseOnlyConstructor, NewScope, PropertyName, ObjectFlags)UClassProperty(FObjectInitializer());
-				((UClassProperty*)NewProperty)->MetaClass = VarProperty.MetaClass;
-			}
-			else
-			{
-				if ( DoesAnythingInHierarchyHaveDefaultToInstanced( VarProperty.PropertyClass ) )
+				else
 				{
-					VarProperty.PropertyFlags |= CPF_InstancedReference;
-					AddEditInlineMetaData(VarProperty.MetaData);
+					if ( DoesAnythingInHierarchyHaveDefaultToInstanced( VarProperty.PropertyClass ) )
+					{
+						VarProperty.PropertyFlags |= CPF_InstancedReference;
+						AddEditInlineMetaData(VarProperty.MetaData);
+					}
+					NewProperty = new(EC_InternalUseOnlyConstructor, NewScope, PropertyName, ObjectFlags)UObjectProperty(FObjectInitializer());
 				}
-				NewProperty = new(EC_InternalUseOnlyConstructor, NewScope, PropertyName, ObjectFlags)UObjectProperty(FObjectInitializer());
-			}
 
-			((UObjectPropertyBase*)NewProperty)->PropertyClass = VarProperty.PropertyClass;
-		}
-		else if (VarProperty.Type == CPT_WeakObjectReference)
-		{
-			check(VarProperty.PropertyClass);
-			NewProperty = new(EC_InternalUseOnlyConstructor, NewScope, PropertyName, ObjectFlags)UWeakObjectProperty(FObjectInitializer());
-			((UObjectPropertyBase*)NewProperty)->PropertyClass = VarProperty.PropertyClass;
-		}
-		else if( VarProperty.Type == CPT_LazyObjectReference )
-		{
-			check(VarProperty.PropertyClass);
-			NewProperty = new(EC_InternalUseOnlyConstructor, NewScope, PropertyName, ObjectFlags)ULazyObjectProperty(FObjectInitializer());
-			((UObjectPropertyBase*)NewProperty)->PropertyClass = VarProperty.PropertyClass;
-		}
-		else if( VarProperty.Type == CPT_AssetObjectReference )
-		{
-			check(VarProperty.PropertyClass);
-			if (VarProperty.PropertyClass->IsChildOf(UClass::StaticClass()))
-			{
-				NewProperty = new(EC_InternalUseOnlyConstructor, NewScope, PropertyName, ObjectFlags)UAssetClassProperty(FObjectInitializer());
-				((UAssetClassProperty*)NewProperty)->MetaClass = VarProperty.MetaClass;
-			}
-			else
-			{
-				NewProperty = new(EC_InternalUseOnlyConstructor, NewScope, PropertyName, ObjectFlags)UAssetObjectProperty(FObjectInitializer());
-			}
-			((UObjectPropertyBase*)NewProperty)->PropertyClass = VarProperty.PropertyClass;
-		}
-		else if (VarProperty.Type == CPT_Interface)
-		{
-			check(VarProperty.PropertyClass);
-			check(VarProperty.PropertyClass->HasAnyClassFlags(CLASS_Interface));
+				((UObjectPropertyBase*)NewProperty)->PropertyClass = VarProperty.PropertyClass;
+				break;
 
-			UInterfaceProperty* InterfaceProperty = new(EC_InternalUseOnlyConstructor, NewScope, PropertyName, ObjectFlags) UInterfaceProperty(FObjectInitializer());
-			InterfaceProperty->InterfaceClass = VarProperty.PropertyClass;
+			case CPT_WeakObjectReference:
+				check(VarProperty.PropertyClass);
+				NewProperty = new(EC_InternalUseOnlyConstructor, NewScope, PropertyName, ObjectFlags)UWeakObjectProperty(FObjectInitializer());
+				((UObjectPropertyBase*)NewProperty)->PropertyClass = VarProperty.PropertyClass;
+				break;
 
-			NewProperty = InterfaceProperty;
-		}
-		else if (VarProperty.Type == CPT_Name)
-		{
-			NewProperty = new(EC_InternalUseOnlyConstructor, NewScope, PropertyName, ObjectFlags)UNameProperty(FObjectInitializer());
-		}
-		else if (VarProperty.Type == CPT_String)
-		{
-			NewProperty = new(EC_InternalUseOnlyConstructor, NewScope, PropertyName, ObjectFlags)UStrProperty(FObjectInitializer());
-		}
-		else if (VarProperty.Type == CPT_Text)
-		{
-			NewProperty = new(EC_InternalUseOnlyConstructor, NewScope, PropertyName, ObjectFlags)UTextProperty(FObjectInitializer());
-		}
-		else if (VarProperty.Type == CPT_Struct)
-		{
-			NewProperty = new(EC_InternalUseOnlyConstructor, NewScope, PropertyName, ObjectFlags)UStructProperty(FObjectInitializer());
-			if (VarProperty.Struct->StructFlags & STRUCT_HasInstancedReference)
-			{
-				VarProperty.PropertyFlags |= CPF_ContainsInstancedReference;
-			}
-			((UStructProperty*)NewProperty)->Struct = VarProperty.Struct;
-		}
-		else if (VarProperty.Type == CPT_Delegate)
-		{
-			NewProperty = new(EC_InternalUseOnlyConstructor, NewScope, PropertyName, ObjectFlags)UDelegateProperty(FObjectInitializer());
-		}
-		else if (VarProperty.Type == CPT_MulticastDelegate)
-		{
-			NewProperty = new(EC_InternalUseOnlyConstructor, NewScope, PropertyName, ObjectFlags)UMulticastDelegateProperty(FObjectInitializer());
-		}
-		else
-		{
-			FError::Throwf(TEXT("Unknown property type %i"), (uint8)VarProperty.Type );
+			case CPT_LazyObjectReference:
+				check(VarProperty.PropertyClass);
+				NewProperty = new(EC_InternalUseOnlyConstructor, NewScope, PropertyName, ObjectFlags)ULazyObjectProperty(FObjectInitializer());
+				((UObjectPropertyBase*)NewProperty)->PropertyClass = VarProperty.PropertyClass;
+				break;
+
+			case CPT_AssetObjectReference:
+				check(VarProperty.PropertyClass);
+				if (VarProperty.PropertyClass->IsChildOf(UClass::StaticClass()))
+				{
+					NewProperty = new(EC_InternalUseOnlyConstructor, NewScope, PropertyName, ObjectFlags)UAssetClassProperty(FObjectInitializer());
+					((UAssetClassProperty*)NewProperty)->MetaClass = VarProperty.MetaClass;
+				}
+				else
+				{
+					NewProperty = new(EC_InternalUseOnlyConstructor, NewScope, PropertyName, ObjectFlags)UAssetObjectProperty(FObjectInitializer());
+				}
+				((UObjectPropertyBase*)NewProperty)->PropertyClass = VarProperty.PropertyClass;
+				break;
+
+			case CPT_Interface:
+				check(VarProperty.PropertyClass);
+				check(VarProperty.PropertyClass->HasAnyClassFlags(CLASS_Interface));
+
+				NewProperty = new(EC_InternalUseOnlyConstructor, NewScope, PropertyName, ObjectFlags) UInterfaceProperty(FObjectInitializer());
+				((UInterfaceProperty*)NewProperty)->InterfaceClass = VarProperty.PropertyClass;
+				break;
+
+			case CPT_Name:
+				NewProperty = new(EC_InternalUseOnlyConstructor, NewScope, PropertyName, ObjectFlags)UNameProperty(FObjectInitializer());
+				break;
+
+			case CPT_String:
+				NewProperty = new(EC_InternalUseOnlyConstructor, NewScope, PropertyName, ObjectFlags)UStrProperty(FObjectInitializer());
+				break;
+
+			case CPT_Text:
+				NewProperty = new(EC_InternalUseOnlyConstructor, NewScope, PropertyName, ObjectFlags)UTextProperty(FObjectInitializer());
+				break;
+
+			case CPT_Struct:
+				NewProperty = new(EC_InternalUseOnlyConstructor, NewScope, PropertyName, ObjectFlags)UStructProperty(FObjectInitializer());
+				if (VarProperty.Struct->StructFlags & STRUCT_HasInstancedReference)
+				{
+					VarProperty.PropertyFlags |= CPF_ContainsInstancedReference;
+				}
+				((UStructProperty*)NewProperty)->Struct = VarProperty.Struct;
+				break;
+
+			case CPT_Delegate:
+				NewProperty = new(EC_InternalUseOnlyConstructor, NewScope, PropertyName, ObjectFlags)UDelegateProperty(FObjectInitializer());
+				break;
+
+			case CPT_MulticastDelegate:
+				NewProperty = new(EC_InternalUseOnlyConstructor, NewScope, PropertyName, ObjectFlags)UMulticastDelegateProperty(FObjectInitializer());
+				break;
+
+			default:
+				FError::Throwf(TEXT("Unknown property type %i"), (uint8)VarProperty.Type);
 		}
 
 		if( Array )
