UE-11425: GitHub 896 : Update UnrealSourceFile Macros to work correctly for project files.

GitHub PR #896

diff --git a/Engine/Source/Programs/UnrealHeaderTool/Private/UnrealSourceFile.cpp b/Engine/Source/Programs/UnrealHeaderTool/Private/UnrealSourceFile.cpp
--- a/Engine/Source/Programs/UnrealHeaderTool/Private/UnrealSourceFile.cpp
+++ b/Engine/Source/Programs/UnrealHeaderTool/Private/UnrealSourceFile.cpp
@@ -15,10 +15,24 @@ FString FUnrealSourceFile::GetFileId() const
 
 	FPaths::MakeStandardFilename(StdFilename);
 
+	bool RelativePath = FPaths::IsRelative(StdFilename);
+
+	if (!RelativePath)
+	{
+		// If path is still absolute that means MakeStandardFilename has failed
+		// In this case make it relative to the current project. 
+		RelativePath = FPaths::MakePathRelativeTo(StdFilename, *FPaths::GetPath(FPaths::GetProjectFilePath()));
+	}
+
+	// If the path has passed either MakeStandardFilename or MakePathRelativeTo it should be using internal path seperators
+	if (RelativePath)
+	{
+		// Remove any preceding parent directory paths
+		while (StdFilename.RemoveFromStart(TEXT("../")));
+	}
+
 	FStringOutputDevice Out;
 
-	// Standard filename will always start with "..\..\..\", so we are removing it.
-	StdFilename = StdFilename.Mid(9);
 	for (auto Char : StdFilename)
 	{
 		if (FChar::IsAlnum(Char))
