UHT: Fixed a large number of missing 'override' warnings on Linux

This will not fix warnings for UObject classes that implement legacy C++ interfaces for which we generate code (GENERATED_IINTERFACE_BODY), because, as far as I can tell, the code generator is currently not able to detect those cases. For example: class UFont : public UObject, public IFontProviderInterface --- IFontProviderInterface gets a _getUObject() injected, but we cannot detect this in the code generation for UFont. Some more work is needed here. It might be better to simply always generate _getUObject() for all classes, so that there are no weird special cases.

diff --git a/Engine/Source/Programs/UnrealHeaderTool/Private/CodeGenerator.cpp b/Engine/Source/Programs/UnrealHeaderTool/Private/CodeGenerator.cpp
--- a/Engine/Source/Programs/UnrealHeaderTool/Private/CodeGenerator.cpp
+++ b/Engine/Source/Programs/UnrealHeaderTool/Private/CodeGenerator.cpp
@@ -1800,7 +1800,15 @@ void FNativeClassHeaderGenerator::ExportClassesFromSourceFileInner(FUnrealSource
 					ClassBoilerplate.Logf(TEXT("    static const TCHAR* StaticConfigName() {return TEXT(\"%s\");}\r\n\r\n"), *Class->ClassConfigName.ToString());
 				}
 
-				ClassBoilerplate.Logf(TEXT("%sUObject* _getUObject() const { return const_cast<%s*>(this); }\r\n"), FCString::Spc(4), ClassCPPName);
+				// @todo srobb: clean up _getUObject() code generation
+				if (((SuperClass == UClass::StaticClass()) || (SuperClass == UObject::StaticClass()) || (SuperClass && (SuperClass->ClassFlags & CLASS_Intrinsic))) && !Class->IsChildOf(UInterface::StaticClass()))
+				{
+					ClassBoilerplate.Logf(TEXT("%svirtual UObject* _getUObject() const { return const_cast<%s*>(this); }\r\n"), FCString::Spc(4), ClassCPPName);
+				}
+				else
+				{
+					ClassBoilerplate.Logf(TEXT("%svirtual UObject* _getUObject() const override { return const_cast<%s*>(this); }\r\n"), FCString::Spc(4), ClassCPPName);
+				}
 
 				// Replication, add in the declaration for GetLifetimeReplicatedProps() automatically if there are any net flagged properties
 				bool bHasNetFlaggedProperties = false;
