UHT error when an property/parameter exposed to BP is not supported by blueprint.

diff --git a/Engine/Source/Editor/BlueprintGraph/Private/EdGraphSchema_K2.cpp b/Engine/Source/Editor/BlueprintGraph/Private/EdGraphSchema_K2.cpp
--- a/Engine/Source/Editor/BlueprintGraph/Private/EdGraphSchema_K2.cpp
+++ b/Engine/Source/Editor/BlueprintGraph/Private/EdGraphSchema_K2.cpp
@@ -3212,7 +3212,6 @@ FText UEdGraphSchema_K2::GetCategoryText(const FString& Category, const bool bFo
 	}
 	else
 	{
-		ensureMsgf(false, TEXT("Unhandled blueprint type ('%s'), failed to find type description."), *Category);
 		return FText::FromString(Category);
 	}
 }
diff --git a/Engine/Source/Programs/UnrealHeaderTool/Private/HeaderParser.cpp b/Engine/Source/Programs/UnrealHeaderTool/Private/HeaderParser.cpp
--- a/Engine/Source/Programs/UnrealHeaderTool/Private/HeaderParser.cpp
+++ b/Engine/Source/Programs/UnrealHeaderTool/Private/HeaderParser.cpp
@@ -411,6 +411,35 @@ namespace
 		// Unreachable
 		return nullptr;
 	}
+
+	bool IsPropertySupportedByBlueprint(const UProperty* Property)
+	{
+		if (Property == NULL)
+		{
+			return false;
+		}
+		if (auto ArrayProperty = Cast<const UArrayProperty>(Property))
+		{
+			return IsPropertySupportedByBlueprint(ArrayProperty->Inner);
+		}
+
+		const bool bSupportedType = Property->IsA<UInterfaceProperty>()
+			|| Property->IsA<UClassProperty>()
+			|| Property->IsA<UAssetClassProperty>() //?
+			|| Property->IsA<UObjectPropertyBase>()
+			|| Property->IsA<UStructProperty>()
+			|| Property->IsA<UFloatProperty>()
+			|| Property->IsA<UIntProperty>()
+			|| Property->IsA<UByteProperty>()
+			|| Property->IsA<UNameProperty>()
+			|| Property->IsA<UBoolProperty>()
+			|| Property->IsA<UStrProperty>()
+			|| Property->IsA<UTextProperty>()
+			|| Property->IsA<UMulticastDelegateProperty>()
+			|| Property->IsA<UDelegateProperty>();
+
+		return bSupportedType;
+	}
 }
 	
 /////////////////////////////////////////////////////
@@ -5651,6 +5680,11 @@ void FHeaderParser::CompileFunctionDeclaration(FUnrealSourceFile& SourceFile, FC
 			{
 				FError::Throwf(TEXT("Static array cannot be exposed to blueprint. Function: %s Parameter %s\n"), *TopFunction->GetName(), *Param->GetName());
 			}
+
+			if (!IsPropertySupportedByBlueprint(Param))
+			{
+				FError::Throwf(TEXT("Type '%s' is not supported by blueprint. Function: %s Parameter %s\n"), *Param->GetCPPType(), *TopFunction->GetName(), *Param->GetName());
+			}
 		}
 	}
 
@@ -6025,6 +6059,11 @@ void FHeaderParser::CompileVariableDeclaration(FClasses& AllClasses, UStruct* St
 			FError::Throwf(TEXT("Static array cannot be exposed to blueprint %s.%s"), *Struct->GetName(), *NewProperty->GetName());
 		}
 
+		if (NewProperty->HasAnyPropertyFlags(CPF_BlueprintVisible) && !IsPropertySupportedByBlueprint(NewProperty))
+		{
+			FError::Throwf(TEXT("Type '%s' is not supported by blueprint. %s.%s"), *NewProperty->GetCPPType(), *Struct->GetName(), *NewProperty->GetName());
+		}
+
 	} while( MatchSymbol(TEXT(",")) );
 
 	// Optional member initializer.
