UnrealHeaderTool: Added time logging for various phases and clarified code/comments surrounding the phases
[INTEGRATE] Change 2493166 by [REDACTED_PERSON][REDACTED_HANDLE] on 2015/03/26 23:49:36

diff --git a/Engine/Source/Programs/UnrealHeaderTool/Private/CodeGenerator.cpp b/Engine/Source/Programs/UnrealHeaderTool/Private/CodeGenerator.cpp
--- a/Engine/Source/Programs/UnrealHeaderTool/Private/CodeGenerator.cpp
+++ b/Engine/Source/Programs/UnrealHeaderTool/Private/CodeGenerator.cpp
@@ -1,6 +1,5 @@
 // Copyright 1998-2015 Epic Games, Inc. All Rights Reserved.
 
-
 #include "UnrealHeaderTool.h"
 
 #include "UniquePtr.h"
@@ -18,6 +17,9 @@
 
 FManifest GManifest;
 
+double GMacroizeTime = 0.0;
+double GTabifyTime = 0.0;
+
 static TArray<FString> ChangeMessages;
 static bool bWriteContents = false;
 static bool bVerifyContents = false;
@@ -25,7 +27,7 @@ static bool bVerifyContents = false;
 static const bool bMultiLineUFUNCTION = true;
 static const bool bMultiLineUPROPERTY = true;
 
-static TSharedRef<FUnrealSourceFile> GenerateCodeForHeader(UPackage* InParent, const FString& FileName, EObjectFlags Flags, const TCHAR* Buffer);
+static TSharedRef<FUnrealSourceFile> PerformInitialParseOnHeader(UPackage* InParent, const FString& FileName, EObjectFlags Flags, const TCHAR* Buffer);
 
 FCompilerMetadataManager GScriptHelper;
 
@@ -121,8 +123,10 @@ static struct FFlagAudit
 	}
 } TheFlagAudit;
 
-static FString Tabify(const TCHAR *Input)
+static FString Tabify(const TCHAR* Input)
 {
+	FScopedDurationTimer Tracker(GTabifyTime);
+
 	FString Result(Input); 
 
 	const int32 SpacesPerTab = 4;
@@ -208,8 +212,10 @@ bool FindPackageLocation(const TCHAR* InPackage, FString& OutLocation, FString&
 }
 
 
-FString Macroize(const TCHAR *MacroName, const TCHAR *StringToMacroize)
+FString Macroize(const TCHAR* MacroName, const TCHAR* StringToMacroize)
 {
+	FScopedDurationTimer Tracker(GMacroizeTime);
+
 	FString Result = StringToMacroize;
 	if (Result.Len())
 	{
@@ -4467,21 +4473,25 @@ bool FNativeClassHeaderGenerator::SaveHeaderIfChanged(const TCHAR* HeaderPath, c
 	if (!bTestedCmdLine)
 	{
 		bTestedCmdLine = true;
-		if( FParse::Param( FCommandLine::Get(), TEXT("WRITEREF") ) )
+
+		const FString ReferenceGeneratedCodePath = FString(FPaths::GameSavedDir()) / TEXT("ReferenceGeneratedCode/");
+		const FString VerifyGeneratedCodePath = FString(FPaths::GameSavedDir()) / TEXT("VerifyGeneratedCode/");
+
+		if (FParse::Param(FCommandLine::Get(), TEXT("WRITEREF")))
 		{
 			bWriteContents = true;
-			UE_LOG(LogCompile, Log, TEXT("********************************* Writing reference generated code to ReferenceGeneratedCode."));
+			UE_LOG(LogCompile, Log, TEXT("********************************* Writing reference generated code to %s."), *ReferenceGeneratedCodePath);
 			UE_LOG(LogCompile, Log, TEXT("********************************* Deleting all files in ReferenceGeneratedCode."));
-			IFileManager::Get().DeleteDirectory(*(FString(FPaths::GameSavedDir()) / TEXT("ReferenceGeneratedCode/")), false, true);
-			IFileManager::Get().MakeDirectory(*(FString(FPaths::GameSavedDir()) / TEXT("ReferenceGeneratedCode/")));
+			IFileManager::Get().DeleteDirectory(*ReferenceGeneratedCodePath, false, true);
+			IFileManager::Get().MakeDirectory(*ReferenceGeneratedCodePath);
 		}
-		else if( FParse::Param( FCommandLine::Get(), TEXT("VERIFYREF") ) )
+		else if (FParse::Param( FCommandLine::Get(), TEXT("VERIFYREF")))
 		{
 			bVerifyContents = true;
-			UE_LOG(LogCompile, Log, TEXT("********************************* Writing generated code to VerifyGeneratedCode and comparing to ReferenceGeneratedCode"));
+			UE_LOG(LogCompile, Log, TEXT("********************************* Writing generated code to %s and comparing to %s"), *VerifyGeneratedCodePath, *ReferenceGeneratedCodePath);
 			UE_LOG(LogCompile, Log, TEXT("********************************* Deleting all files in VerifyGeneratedCode."));
-			IFileManager::Get().DeleteDirectory(*(FString(FPaths::GameSavedDir()) / TEXT("VerifyGeneratedCode/")), false, true);
-			IFileManager::Get().MakeDirectory(*(FString(FPaths::GameSavedDir()) / TEXT("VerifyGeneratedCode/")));
+			IFileManager::Get().DeleteDirectory(*VerifyGeneratedCodePath, false, true);
+			IFileManager::Get().MakeDirectory(*VerifyGeneratedCodePath);
 		}
 	}
 
@@ -4811,10 +4821,11 @@ void FNativeClassHeaderGenerator::ExportGeneratedCPP()
 		SaveHeaderIfChanged(*ConstructorPath,*ConstructorsWithIfdef);
 	}
 }
-
 /** Get all script plugins based on ini setting */
 void GetScriptPlugins(TArray<IScriptGeneratorPluginInterface*>& ScriptPlugins)
 {
+	FScopedDurationTimer PluginTimeTracker(GPluginOverheadTime);
+
 	const FString CodeGeneratorPluginCategory(TEXT("UnrealHeaderTool.Code Generator"));
 	const TArray<FPluginStatus> Plugins = IPluginManager::Get().QueryStatusForAllPlugins();
 	for (auto PluginIt(Plugins.CreateConstIterator()); PluginIt; ++PluginIt)
@@ -4909,6 +4920,9 @@ ECompilationResult::Type UnrealHeaderTool_Main(const FString& ModuleInfoFilename
 		FolderType_Count
 	};
 
+	double TotalModulePreparseTime = 0.0;
+	double TotalParseAndCodegenTime = 0.0;
+
 	for (const auto& Module : GManifest.Modules)
 	{
 		if (Result != ECompilationResult::Succeeded)
@@ -4916,6 +4930,11 @@ ECompilationResult::Type UnrealHeaderTool_Main(const FString& ModuleInfoFilename
 			break;
 		}
 
+		double ThisModulePreparseTime = 0.0;
+		int32 NumHeadersPreparsed = 0;
+		FDurationTimer ThisModuleTimer(ThisModulePreparseTime);
+		ThisModuleTimer.Start();
+
 		UPackage* Package = Cast<UPackage>(StaticFindObjectFast(UPackage::StaticClass(), NULL, FName(*Module.LongPackageName), false, false));
 		if (Package == NULL)
 		{
@@ -4930,6 +4949,7 @@ ECompilationResult::Type UnrealHeaderTool_Main(const FString& ModuleInfoFilename
 
 		GPackageToManifestModuleMap.Add(Package, &Module);
 
+		// Pre-parse the headers
 		for (int32 PassIndex = 0; PassIndex < FolderType_Count && Result == ECompilationResult::Succeeded; ++PassIndex)
 		{
 			EHeaderFolderTypes CurrentlyProcessing = (EHeaderFolderTypes)PassIndex;
@@ -4945,24 +4965,26 @@ ECompilationResult::Type UnrealHeaderTool_Main(const FString& ModuleInfoFilename
 				continue;
 			}
 
-			for (const FString& Filename : UObjectHeaders)
+			NumHeadersPreparsed += UObjectHeaders.Num();
+
+			for (const FString& RawFilename : UObjectHeaders)
 			{
 			#if !PLATFORM_EXCEPTIONS_DISABLED
 				try
 			#endif
 				{
 					// Import class.
-					const FString FullModulePath = FPaths::ConvertRelativePathToFull(ModuleInfoPath, Filename);
+					const FString FullFilename = FPaths::ConvertRelativePathToFull(ModuleInfoPath, RawFilename);
 
 					FString HeaderFile;
-					if (!FFileHelper::LoadFileToString(HeaderFile, *FullModulePath))
+					if (!FFileHelper::LoadFileToString(HeaderFile, *FullFilename))
 					{
-						FError::Throwf(TEXT( "UnrealHeaderTool was unable to load source file '%s'"), *FullModulePath);
+						FError::Throwf(TEXT("UnrealHeaderTool was unable to load source file '%s'"), *FullFilename);
 					}
 
-					TSharedRef<FUnrealSourceFile> UnrealSourceFile = GenerateCodeForHeader(Package, Filename, RF_Public | RF_Standalone, *HeaderFile);
+					TSharedRef<FUnrealSourceFile> UnrealSourceFile = PerformInitialParseOnHeader(Package, RawFilename, RF_Public | RF_Standalone, *HeaderFile);
 
-					GUnrealSourceFilesMap.Add(Filename, UnrealSourceFile);
+					GUnrealSourceFilesMap.Add(RawFilename, UnrealSourceFile);
 
 					if (CurrentlyProcessing == PublicClassesHeaders)
 					{
@@ -4975,18 +4997,18 @@ ECompilationResult::Type UnrealHeaderTool_Main(const FString& ModuleInfoFilename
 					}
 
 					// Save metadata for the class path, both for it's include path and relative to the module base directory
-					if(FullModulePath.StartsWith(Module.BaseDirectory))
+					if (FullFilename.StartsWith(Module.BaseDirectory))
 					{
 						// Get the path relative to the module directory
-						const TCHAR *ModuleRelativePath = *FullModulePath + Module.BaseDirectory.Len();
+						const TCHAR* ModuleRelativePath = *FullFilename + Module.BaseDirectory.Len();
 
 						UnrealSourceFile->SetModuleRelativePath(ModuleRelativePath);
 
 						// Calculate the include path
-						const TCHAR *IncludePath = ModuleRelativePath;
+						const TCHAR* IncludePath = ModuleRelativePath;
 
 						// Walk over the first potential slash
-						if(*IncludePath == '/')
+						if (*IncludePath == TEXT('/'))
 						{
 							IncludePath++;
 						}
@@ -4995,34 +5017,34 @@ ECompilationResult::Type UnrealHeaderTool_Main(const FString& ModuleInfoFilename
 						static const TCHAR PublicFolderName[]  = TEXT("Public/");
 						static const TCHAR PrivateFolderName[] = TEXT("Private/");
 						static const TCHAR ClassesFolderName[] = TEXT("Classes/");
-						if(FCString::Strnicmp(IncludePath, PublicFolderName, ARRAY_COUNT(PublicFolderName) - 1) == 0)
+						if (FCString::Strnicmp(IncludePath, PublicFolderName, ARRAY_COUNT(PublicFolderName) - 1) == 0)
 						{
 							IncludePath += (ARRAY_COUNT(PublicFolderName) - 1);
 						}
-						else if(FCString::Strnicmp(IncludePath, PrivateFolderName, ARRAY_COUNT(PrivateFolderName) - 1) == 0)
+						else if (FCString::Strnicmp(IncludePath, PrivateFolderName, ARRAY_COUNT(PrivateFolderName) - 1) == 0)
 						{
 							IncludePath += (ARRAY_COUNT(PrivateFolderName) - 1);
 						}
-						else if(FCString::Strnicmp(IncludePath, ClassesFolderName, ARRAY_COUNT(ClassesFolderName) - 1) == 0)
+						else if (FCString::Strnicmp(IncludePath, ClassesFolderName, ARRAY_COUNT(ClassesFolderName) - 1) == 0)
 						{
 							IncludePath += (ARRAY_COUNT(ClassesFolderName) - 1);
 						}
 
 						// Add the include path
-						if(*IncludePath != 0)
+						if (*IncludePath != 0)
 						{
 							UnrealSourceFile->SetIncludePath(MoveTemp(IncludePath));
 						}
 					}
 				}
 			#if !PLATFORM_EXCEPTIONS_DISABLED
-				catch( TCHAR* ErrorMsg )
+				catch (TCHAR* ErrorMsg)
 				{
 					TGuardValue<ELogTimes::Type> DisableLogTimes(GPrintLogTimes, ELogTimes::None);
 
 					FString Prefix;
 
-					const FString AbsFilename = IFileManager::Get().ConvertToAbsolutePathForExternalAppForRead(*Filename);
+					const FString AbsFilename = IFileManager::Get().ConvertToAbsolutePathForExternalAppForRead(*RawFilename);
 					Prefix = FString::Printf(TEXT("%s(1): "), *AbsFilename);
 
 					FString FormattedErrorMessage = FString::Printf(TEXT("%sError: %s\r\n"), *Prefix, ErrorMsg);
@@ -5040,12 +5062,17 @@ ECompilationResult::Type UnrealHeaderTool_Main(const FString& ModuleInfoFilename
 				Result = ECompilationResult::OtherCompilationError;
 			}
 		}
-	}
 
+		ThisModuleTimer.Stop();
+		TotalModulePreparseTime += ThisModulePreparseTime;
+		UE_LOG(LogCompile, Log, TEXT("Preparsed module %s containing %i files(s) in %.2f secs."), *Module.LongPackageName, NumHeadersPreparsed, ThisModulePreparseTime);
+	}
 
-	// Save generated headers
-	if ( Result == ECompilationResult::Succeeded )
+	// Do the actual parse of the headers and generate for them
+	if (Result == ECompilationResult::Succeeded)
 	{
+		FScopedDurationTimer ParseAndCodeGenTimer(TotalParseAndCodegenTime);
+
 		// Verify that all script declared superclasses exist.
 		for (const UClass* ScriptClass : TObjectRange<UClass>())
 		{
@@ -5125,9 +5152,12 @@ ECompilationResult::Type UnrealHeaderTool_Main(const FString& ModuleInfoFilename
 				}
 			}
 
-			for (auto ScriptGenerator : ScriptPlugins)
 			{
-				ScriptGenerator->FinishExport();
+				FScopedDurationTimer PluginTimeTracker(GPluginOverheadTime);
+				for (auto ScriptGenerator : ScriptPlugins)
+				{
+					ScriptGenerator->FinishExport();
+				}
 			}
 		}
 	}
@@ -5135,11 +5165,18 @@ ECompilationResult::Type UnrealHeaderTool_Main(const FString& ModuleInfoFilename
 	// Avoid TArray slack for meta data.
 	GScriptHelper.Shrink();
 
-	if( bWriteContents )
+	UE_LOG(LogCompile, Log, TEXT("Preparsing %i modules took %f seconds"), GManifest.Modules.Num(), TotalModulePreparseTime);
+	UE_LOG(LogCompile, Log, TEXT("Parsing took %f seconds"), TotalParseAndCodegenTime - GHeaderCodeGenTime);
+	UE_LOG(LogCompile, Log, TEXT("Code generation took %f seconds"), GHeaderCodeGenTime);
+	UE_LOG(LogCompile, Log, TEXT("ScriptPlugin overhead was %f seconds"), GPluginOverheadTime);
+	UE_LOG(LogCompile, Log, TEXT("Macroize time was %f seconds"), GMacroizeTime);
+	UE_LOG(LogCompile, Log, TEXT("Tabify time was was %f seconds"), GTabifyTime);
+
+	if (bWriteContents)
 	{
 		UE_LOG(LogCompile, Log, TEXT("********************************* Wrote reference generated code to ReferenceGeneratedCode."));
 	}
-	else if( bVerifyContents )
+	else if (bVerifyContents)
 	{
 		UE_LOG(LogCompile, Log, TEXT("********************************* Wrote generated code to VerifyGeneratedCode and compared to ReferenceGeneratedCode"));
 		for (FString& Msg : ChangeMessages)
@@ -5159,7 +5196,7 @@ ECompilationResult::Type UnrealHeaderTool_Main(const FString& ModuleInfoFilename
 
 	GIsRequestingExit = true;
 
-	if(Result == ECompilationResult::Succeeded && NumFailures > 0)
+	if ((Result == ECompilationResult::Succeeded) && (NumFailures > 0))
 	{
 		return ECompilationResult::OtherCompilationError;
 	}
@@ -5282,13 +5319,7 @@ UClass* ProcessParsedClass(bool bClassIsAnInterface, TArray<FHeaderProvider> &De
 }
 
 
-TSharedRef<FUnrealSourceFile> GenerateCodeForHeader
-(
-	UPackage*     InParent,
-	const FString& FileName,
-	EObjectFlags Flags,
-	const TCHAR* Buffer
-)
+TSharedRef<FUnrealSourceFile> PerformInitialParseOnHeader(UPackage* InParent, const FString& FileName, EObjectFlags Flags, const TCHAR* Buffer)
 {
 	const TCHAR* InBuffer = Buffer;
 
diff --git a/Engine/Source/Programs/UnrealHeaderTool/Private/HeaderParser.cpp b/Engine/Source/Programs/UnrealHeaderTool/Private/HeaderParser.cpp
--- a/Engine/Source/Programs/UnrealHeaderTool/Private/HeaderParser.cpp
+++ b/Engine/Source/Programs/UnrealHeaderTool/Private/HeaderParser.cpp
@@ -13,6 +13,9 @@
 #include "Manifest.h"
 #include "UnitConversion.h"
 
+double GPluginOverheadTime = 0.0;
+double GHeaderCodeGenTime = 0.0;
+
 /*-----------------------------------------------------------------------------
 	Constants & declarations.
 -----------------------------------------------------------------------------*/
@@ -6760,23 +6763,28 @@ ECompilationResult::Type FHeaderParser::ParseAllHeadersInside(
 			// from the feedback context.
 			Warn->SetContext(NULL);
 
-			ExportNativeHeaders(
-				CurrentPackage,
-				ModuleClasses,
-				Module.SaveExportedHeaders
+			double ExportTime = 0.0;
+			{
+				FScopedDurationTimer Timer(ExportTime);
+				ExportNativeHeaders(
+					CurrentPackage,
+					ModuleClasses,
+					Module.SaveExportedHeaders
 #if WITH_HOT_RELOAD_CTORS
-				, bExportVTableConstructors
+					, bExportVTableConstructors
 #endif // WITH_HOT_RELOAD_CTORS
-			);
+				);
+			}
+			GHeaderCodeGenTime += ExportTime;
 
 			// Done with header generation
 			if (HeaderParser.LinesParsed > 0)
 			{
-				UE_LOG(LogCompile, Log,  TEXT("Success: Parsed %i line(s), %i statement(s).\r\n"), HeaderParser.LinesParsed, HeaderParser.StatementsParsed );
+				UE_LOG(LogCompile, Log, TEXT("Success: Parsed %i line(s), %i statement(s) in %.2f secs.\r\n"), HeaderParser.LinesParsed, HeaderParser.StatementsParsed, ExportTime);
 			}
 			else
 			{
-				UE_LOG(LogCompile, Log,  TEXT("Success: Everything is up to date") );
+				UE_LOG(LogCompile, Log, TEXT("Success: Everything is up to date (in %.2f secs)"), ExportTime);
 			}
 		}
 	}
@@ -6792,6 +6800,8 @@ ECompilationResult::Type FHeaderParser::ParseAllHeadersInside(
 
 	if (Result == ECompilationResult::Succeeded && ScriptPlugins.Num())
 	{
+		FScopedDurationTimer PluginTimeTracker(GPluginOverheadTime);
+
 		auto RootNode = &ModuleClasses.GetClassTree();
 		for (auto Plugin : ScriptPlugins)
 		{
diff --git a/Engine/Source/Programs/UnrealHeaderTool/Private/HeaderParser.h b/Engine/Source/Programs/UnrealHeaderTool/Private/HeaderParser.h
--- a/Engine/Source/Programs/UnrealHeaderTool/Private/HeaderParser.h
+++ b/Engine/Source/Programs/UnrealHeaderTool/Private/HeaderParser.h
@@ -5,10 +5,13 @@
 #include "ParserHelper.h"
 #include "BaseParser.h"
 #include "CompilationResult.h"
+#include "ScopedTimers.h"
 
 class FClass;
 class FClasses;
 
+extern double GPluginOverheadTime;
+extern double GHeaderCodeGenTime;
 
 /*-----------------------------------------------------------------------------
 	Constants & types.
