Optimized UnrealHeaderTool
- Optimized Tabify and Macroize
- Replaced various %s Spc(constant) calls with direct string literals
- Replaced injected indentation spaces with tabs (avoiding some of the work in Tabify; they were ultimately being replaced anyways)
- Replaced various %s of LINE_TERMINATOR with compile-time string literal concatenation
- Avoided case-insensitive work in various places
- Replaced direct references to FStringOutputDevice with a typedef to facilitate future optimizations
- Used Reset() instead of Empty() on various string builders and strings to reduce reallocations
[INTEGRATE] Change 2494450 by [REDACTED_PERSON][REDACTED_HANDLE] on 2015/03/27 16:48:25

diff --git a/Engine/Source/Programs/UnrealHeaderTool/Private/CodeGenerator.cpp b/Engine/Source/Programs/UnrealHeaderTool/Private/CodeGenerator.cpp
--- a/Engine/Source/Programs/UnrealHeaderTool/Private/CodeGenerator.cpp
+++ b/Engine/Source/Programs/UnrealHeaderTool/Private/CodeGenerator.cpp
@@ -34,6 +34,17 @@ FCompilerMetadataManager GScriptHelper;
 /** C++ name lookup helper */
 FNameLookupCPP NameLookupCPP;
 
+static const int32 MaxNumTabs = 64;
+static const TCHAR TabString[MaxNumTabs + 1] = TEXT("\t\t\t\t\t\t\t\t\t\t\t\t\t\t\t\t\t\t\t\t\t\t\t\t\t\t\t\t\t\t\t\t\t\t\t\t\t\t\t\t\t\t\t\t\t\t\t\t\t\t\t\t\t\t\t\t\t\t\t\t\t\t\t\t");
+
+const TCHAR* GetTabs(int32 NumTabs)
+{
+	checkSlow((NumTabs >= 0) && (NumTabs < MaxNumTabs));
+
+	const uint32 TabPos = MaxNumTabs - NumTabs;
+	return TabString + TabPos;
+}
+
 /////////////////////////////////////////////////////
 // FFlagAudit
 
@@ -128,41 +139,65 @@ static FString Tabify(const TCHAR* Input)
 	FScopedDurationTimer Tracker(GTabifyTime);
 
 	FString Result(Input); 
-
 	const int32 SpacesPerTab = 4;
-	Result.ReplaceInline(TEXT("\r\n"), TEXT("\n"));
+	Result.ReplaceInline(TEXT("\r\n"), TEXT("\n"), ESearchCase::CaseSensitive);
 
 	TArray<FString> Lines;
-	Result.ParseIntoArray( Lines, TEXT( "\n" ), false );
+	Result.ParseIntoArray(Lines, TEXT("\n"), false);
+
+	Result.Reset();
 
-	Result = TEXT("");
 	for (int32 Index = 0; Index < Lines.Num(); Index++)
 	{
 		if (Index)
 		{
 			Result += TEXT("\r\n");
 		}
+
 		FString& Line = Lines[Index];
+
+		int32 NumSpaces = 0;
 		int32 FirstNonWS = 0;
+		int32 NumTrueTabs = 0;
 		while (1)
 		{
 			TCHAR c = *(*Line + FirstNonWS);
-			if (c != TEXT('\t') && c != TEXT(' '))
+
+			if (c == TEXT('\t'))
+			{
+				++NumTrueTabs;
+				NumSpaces += SpacesPerTab;
+			}
+			else if (c == TEXT(' '))
+			{
+				++NumSpaces;
+			}
+			else
 			{
 				break;
 			}
 			FirstNonWS++;
 		}
-		FString Start = Line.Left(FirstNonWS);
-		FString Remainder = *Line + FirstNonWS;
-		Start.ReplaceInline(TEXT("\t"), FCString::Spc(SpacesPerTab));
-		int32 NumTabs = Start.Len() / SpacesPerTab;
-		for (int32 Tab = 0; Tab < NumTabs; Tab++)
+
+		const int32 NumTabs = NumSpaces / SpacesPerTab;
+		const int32 LeftoverSpaces = NumSpaces % SpacesPerTab;
+		if ((NumTrueTabs == NumTabs) && (LeftoverSpaces == 0))
 		{
-			Result += TEXT("\t");
+			Result += Line;
+		}
+		else
+		{
+			FString Remainder = *Line + FirstNonWS;
+			for (int32 Tab = 0; Tab < NumTabs; Tab++)
+			{
+				Result += TEXT("\t");
+			}
+			if (LeftoverSpaces > 0)
+			{
+				Result += FCString::Spc(LeftoverSpaces);
+			}
+			Result += Remainder;
 		}
-		Result += FCString::Spc(Start.Len() - NumTabs * SpacesPerTab);
-		Result += Remainder;
 	}
 
 	return Result;
@@ -219,12 +254,22 @@ FString Macroize(const TCHAR* MacroName, const TCHAR* StringToMacroize)
 	FString Result = StringToMacroize;
 	if (Result.Len())
 	{
-		Result.ReplaceInline(TEXT("\r\n"), TEXT("\n"));
-		Result.ReplaceInline(TEXT("\n"), TEXT(" \\\n"));
-		check(Result.EndsWith(TEXT(" \\\n")));
-		Result = Result.LeftChop(3);
-		Result += TEXT("\n\n\n");
-		Result.ReplaceInline(TEXT("\n"), TEXT("\r\n"));		
+		Result.ReplaceInline(TEXT("\r\n"), TEXT("\n"), ESearchCase::CaseSensitive);
+		Result.ReplaceInline(TEXT("\n"), TEXT(" \\\n"), ESearchCase::CaseSensitive);
+		checkSlow(Result.EndsWith(TEXT(" \\\n"), ESearchCase::CaseSensitive));
+
+		if (Result.Len() >= 3)
+		{
+			for (int32 Index = Result.Len() - 3; Index < Result.Len(); ++Index)
+			{
+				Result[Index] = TEXT('\n');
+			}
+		}
+		else
+		{
+			Result = TEXT("\n\n\n");
+		}
+		Result.ReplaceInline(TEXT("\n"), TEXT("\r\n"), ESearchCase::CaseSensitive);
 	}
 	return FString::Printf(TEXT("#define %s%s\r\n"), MacroName, Result.Len() ? TEXT(" \\") : TEXT("")) + Result;
 }
@@ -403,7 +448,7 @@ static FString GetMetaDataCodeForObject(UObject* Object, const TCHAR* SymbolName
  * @param	TextIndent			Current text indentation
  * @param	ImportsDefaults		whether this struct will be serialized with a default value
  */
-void FNativeClassHeaderGenerator::ExportProperties( UStruct* Struct, int32 TextIndent, bool bAccessSpecifiers, class FStringOutputDevice* Output )
+void FNativeClassHeaderGenerator::ExportProperties(UStruct* Struct, int32 TextIndent, bool bAccessSpecifiers, FUHTStringBuilder* Output)
 {
 	UProperty*	Previous			= NULL;
 	UProperty*	PreviousNonEditorOnly = NULL;
@@ -412,8 +457,8 @@ void FNativeClassHeaderGenerator::ExportProperties( UStruct* Struct, int32 TextI
 	bool		bEmittedHasEditorOnlyMacro = false;
 	bool		bEmittedHasScriptAlign = false;
 
-	check(Output != NULL)
-	FStringOutputDevice& HeaderOutput = *Output;
+	check(Output != NULL);
+	FUHTStringBuilder& HeaderOutput = *Output;
 
 
 	// Find last property in the lowest base class that has any properties
@@ -470,7 +515,7 @@ void FNativeClassHeaderGenerator::ExportProperties( UStruct* Struct, int32 TextI
 	{
 		UProperty* Current = *It;
 
-		FStringOutputDevice PropertyText;
+		FUHTStringBuilder PropertyText;
 
 		// Disregard properties with 0 size like functions.
 		if (It.GetStruct() == Struct)
@@ -514,7 +559,7 @@ void FNativeClassHeaderGenerator::ExportProperties( UStruct* Struct, int32 TextI
 								bEmittedHasEditorOnlyMacro = false;
 							}
 
-							PropertyText.Logf(TEXT("%s:%s"), *AccessSpecifier, LINE_TERMINATOR);
+							PropertyText.Logf(TEXT("%s:") LINE_TERMINATOR, *AccessSpecifier);
 						}
 					}
 				}
@@ -537,22 +582,18 @@ void FNativeClassHeaderGenerator::ExportProperties( UStruct* Struct, int32 TextI
 			// Export property specifiers
 			// Indent code and export CPP text.
 			{
-				FStringOutputDevice JustPropertyDecl;
-				JustPropertyDecl.Log( FCString::Spc(TextIndent+4) );
+				FUHTStringBuilder JustPropertyDecl;
 
 				const FString* Dim = GArrayDimensions.Find(Current);
 				Current->ExportCppDeclaration( JustPropertyDecl, EExportedDeclaration::Member, Dim ? **Dim : NULL);
 				ApplyAlternatePropertyExportText(*It, JustPropertyDecl);
 
 				// Finish up line.
-				JustPropertyDecl.Log(TEXT(";"));
-
-				// Finish up line.
-				PropertyText.Logf(TEXT("%s\r\n"), *JustPropertyDecl);
+				PropertyText.Logf(TEXT("%s%s;\r\n"), GetTabs(TextIndent + 1), *JustPropertyDecl);
 			}
 
 			LastInSuper	= NULL;
-			Previous	= Current;
+			Previous = Current;
 			if (!Current->IsEditorOnlyProperty())
 			{
 				PreviousNonEditorOnly = Current;
@@ -563,20 +604,20 @@ void FNativeClassHeaderGenerator::ExportProperties( UStruct* Struct, int32 TextI
 	}
 	if (bCurrentlyInNotCPPBlock)
 	{
-		HeaderOutput.Log(TEXT("#endif\r\n"));
+		HeaderOutput += TEXT("#endif\r\n");
 		bCurrentlyInNotCPPBlock = false;
 	}
 
 	// End of property list.  If we haven't generated the WITH_EDITORONLY_DATA #endif, do so now.
 	if (bEmittedHasEditorOnlyMacro)
 	{
-		HeaderOutput.Logf( TEXT("#endif // WITH_EDITORONLY_DATA\r\n") );
+		HeaderOutput += TEXT("#endif // WITH_EDITORONLY_DATA\r\n");
 	}
 
 	// if the last property that was exported wasn't public, emit a line to reset the access to "public" so that we don't interfere with cpptext
-	if ( CurrentExportType != PROPEXPORT_Public )
+	if (CurrentExportType != PROPEXPORT_Public)
 	{
-		HeaderOutput.Logf(TEXT("public:%s"), LINE_TERMINATOR);
+		HeaderOutput += TEXT("public:") LINE_TERMINATOR;
 	}
 }
 
@@ -634,7 +675,7 @@ FString FNativeClassHeaderGenerator::GetSingletonName(UField* Item, bool bRequir
 
 	FString ClassString = NameLookupCPP.GetNameCPP(Item->GetClass());
 	// Can't use long package names in function names.
-	if (Result.StartsWith(TEXT("/Script/")))
+	if (Result.StartsWith(TEXT("/Script/"), ESearchCase::CaseSensitive))
 	{
 		Result = FPackageName::GetShortName(Result);
 	}
@@ -795,7 +836,8 @@ inline FString GetEventStructParamsName(UObject* Outer, const TCHAR* FunctionNam
 	}
 	else if (Outer->IsA<UPackage>())
 	{
-		OuterName = ((UPackage*)Outer)->GetName().Replace(TEXT("/"), TEXT("_"));
+		OuterName = ((UPackage*)Outer)->GetName();
+		OuterName.ReplaceInline(TEXT("/"), TEXT("_"), ESearchCase::CaseSensitive);
 	}
 	else
 	{
@@ -864,7 +906,7 @@ void FNativeClassHeaderGenerator::OutputProperty(FString& Meta, FOutputDevice& O
 static FString BackSlashAndIndent(const FString& Text)
 {
 	FString Trailer(TEXT("    \\\r\n    "));
-	FString Result = FString(TEXT("    ")) + Text.Replace(TEXT("\r\n"), *Trailer);
+	FString Result = FString(TEXT("    ")) + Text.Replace(TEXT("\r\n"), *Trailer, ESearchCase::CaseSensitive);
 	return Result.LeftChop(Trailer.Len()) + TEXT("\r\n");
 }
 
@@ -935,7 +977,7 @@ void FNativeClassHeaderGenerator::ExportGeneratedPackageInitCode(UPackage* Packa
 	FString ApiString = GetAPIString();
 	FString SingletonName(GetPackageSingletonName(Package));
 
-	FStringOutputDevice& GeneratedFunctionText = GetGeneratedFunctionTextDevice();
+	FUHTStringBuilder& GeneratedFunctionText = GetGeneratedFunctionTextDevice();
 
 	GeneratedFunctionDeclarations.Logf(TEXT("    %sclass UPackage* %s;\r\n"), *ApiString, *SingletonName);
 
@@ -1009,14 +1051,14 @@ void FNativeClassHeaderGenerator::ExportGeneratedPackageInitCode(UPackage* Packa
 
 }
 
-void FNativeClassHeaderGenerator::ExportNativeGeneratedInitCode(FClass* Class, FStringOutputDevice& OutFriendText)
+void FNativeClassHeaderGenerator::ExportNativeGeneratedInitCode(FClass* Class, FUHTStringBuilder& OutFriendText)
 {
-	FStringOutputDevice& GeneratedFunctionText = GetGeneratedFunctionTextDevice();
+	FUHTStringBuilder& GeneratedFunctionText = GetGeneratedFunctionTextDevice();
 	check(!OutFriendText.Len());
 	// Emit code to build the UObjects that used to be in .u files
 	bool bIsNoExport = Class->HasAnyClassFlags(CLASS_NoExport);
-	FStringOutputDevice BodyText;
-	FStringOutputDevice CallSingletons;
+	FUHTStringBuilder BodyText;
+	FUHTStringBuilder CallSingletons;
 	FString ApiString = GetAPIString();
 
 	TArray<UFunction*> FunctionsToExport;
@@ -1044,7 +1086,7 @@ void FNativeClassHeaderGenerator::ExportNativeGeneratedInitCode(FClass* Class, F
 		CallSingletons.Logf(TEXT("                OuterClass->LinkChild(%s);\r\n"), *GetSingletonName(Function));
 	}
 
-	FStringOutputDevice GeneratedClassRegisterFunctionText;
+	FUHTStringBuilder GeneratedClassRegisterFunctionText;
 
 	// The class itself.
 	{
@@ -1175,10 +1217,10 @@ void FNativeClassHeaderGenerator::ExportNativeGeneratedInitCode(FClass* Class, F
 		GeneratedFunctionText.Logf(TEXT("    /* friend declarations for pasting into noexport class %s\r\n"), NameLookupCPP.GetNameCPP(Class));
 		GeneratedFunctionText.Log(OutFriendText);
 		GeneratedFunctionText.Logf(TEXT("    */\r\n"));
-		OutFriendText.Empty();
+		OutFriendText.Reset();
 	}
 	FString SingletonName(GetSingletonName(Class));
-	SingletonName.ReplaceInline(TEXT("()"), TEXT("")); // function address
+	SingletonName.ReplaceInline(TEXT("()"), TEXT(""), ESearchCase::CaseSensitive); // function address
 
 	{	
 		auto ClassNameCPP = NameLookupCPP.GetNameCPP(Class);
@@ -1208,7 +1250,7 @@ void FNativeClassHeaderGenerator::ExportFunction(UFunction* Function, FScope* Sc
 	SingletonNameToExternDecl.Add(SingletonName, Extern);
 	GeneratedFunctionDeclarations.Log(*Extern);
 
-	FStringOutputDevice CurrentFunctionText;
+	FUHTStringBuilder CurrentFunctionText;
 
 	CurrentFunctionText.Logf(TEXT("    UFunction* %s\r\n"), *SingletonName);
 	CurrentFunctionText.Logf(TEXT("    {\r\n"));
@@ -1220,11 +1262,11 @@ void FNativeClassHeaderGenerator::ExportFunction(UFunction* Function, FScope* Sc
 		FindNoExportStructs(Structs, Function);
 		if (Structs.Num())
 		{
-			ExportMirrorsForNoexportStructs(Structs, 8, CurrentFunctionText);
+			ExportMirrorsForNoexportStructs(Structs, /*Indent=*/ 2, CurrentFunctionText);
 		}
 		TArray<UFunction*> CallbackFunctions;
 		CallbackFunctions.Add(Function);
-		ExportEventParms(*Scope, CallbackFunctions, CurrentFunctionText, 8, false);
+		ExportEventParms(*Scope, CallbackFunctions, CurrentFunctionText, /*Indent=*/ 2, /*bOutputConstructor=*/ false);
 	}
 
 	if (UObject* Outer = Function->GetOuter())
@@ -1359,7 +1401,7 @@ void FNativeClassHeaderGenerator::ExportNatives(FClass* Class)
 	GeneratedPackageCPP.Logf(TEXT("    }\r\n"));
 }
 
-void FNativeClassHeaderGenerator::ExportInterfaceCallFunctions(const TArray<UFunction*>& InCallbackFunctions, UClass* Class, FClassMetaData* ClassData, FStringOutputDevice& HeaderOutput)
+void FNativeClassHeaderGenerator::ExportInterfaceCallFunctions(const TArray<UFunction*>& InCallbackFunctions, UClass* Class, FClassMetaData* ClassData, FUHTStringBuilder& HeaderOutput)
 {
 	TArray<UFunction*> CallbackFunctions = InCallbackFunctions;
 	CallbackFunctions.Sort();
@@ -1377,39 +1419,39 @@ void FNativeClassHeaderGenerator::ExportInterfaceCallFunctions(const TArray<UFun
 		FString ExtraParam = FString::Printf(TEXT("%sUObject* O"), ConstQualifier);
 
 		ExportNativeFunctionHeader(FunctionData, HeaderOutput, EExportFunctionType::Interface, EExportFunctionHeaderStyle::Declaration, *ExtraParam);
-		HeaderOutput.Logf( TEXT(";%s"), LINE_TERMINATOR );
+		HeaderOutput.Logf( TEXT(";") LINE_TERMINATOR );
 
 
 		ExportNativeFunctionHeader(FunctionData, GeneratedPackageCPP, EExportFunctionType::Interface, EExportFunctionHeaderStyle::Definition, *ExtraParam);
-		GeneratedPackageCPP.Logf( TEXT("%s%s{%s"), LINE_TERMINATOR, FCString::Spc(4), LINE_TERMINATOR );
+		GeneratedPackageCPP.Logf( LINE_TERMINATOR TEXT("\t{") LINE_TERMINATOR );
 
-		GeneratedPackageCPP.Logf(TEXT("%scheck(O != NULL);%s"), FCString::Spc(8), LINE_TERMINATOR);
-		GeneratedPackageCPP.Logf(TEXT("%scheck(O->GetClass()->ImplementsInterface(U%s::StaticClass()));%s"), FCString::Spc(8), *ClassName, LINE_TERMINATOR);
+		GeneratedPackageCPP.Logf(TEXT("\t\tcheck(O != NULL);") LINE_TERMINATOR);
+		GeneratedPackageCPP.Logf(TEXT("\t\tcheck(O->GetClass()->ImplementsInterface(U%s::StaticClass()));") LINE_TERMINATOR, *ClassName);
 
 		auto Parameters = GetFunctionParmsAndReturn(FunctionData.FunctionReference);
 
 		// See if we need to create Parms struct
-		bool bHasParms = Parameters.HasParms();
+		const bool bHasParms = Parameters.HasParms();
 		if (bHasParms)
 		{
 			FString EventParmStructName = GetEventStructParamsName(Function->GetOuter(), *FunctionName);
-			GeneratedPackageCPP.Logf(TEXT("%s%s Parms;%s"), FCString::Spc(8), *EventParmStructName, LINE_TERMINATOR);
+			GeneratedPackageCPP.Logf(TEXT("\t\t%s Parms;") LINE_TERMINATOR, *EventParmStructName);
 		}
 
-		GeneratedPackageCPP.Logf(TEXT("%sUFunction* const Func = O->FindFunction(%s_%s);%s"), FCString::Spc(8), *API, *FunctionName, LINE_TERMINATOR);
-		GeneratedPackageCPP.Logf(TEXT("%sif (Func)%s"), FCString::Spc(8), LINE_TERMINATOR);
-		GeneratedPackageCPP.Logf(TEXT("%s{%s"), FCString::Spc(8), LINE_TERMINATOR);
+		GeneratedPackageCPP.Logf(TEXT("\t\tUFunction* const Func = O->FindFunction(%s_%s);") LINE_TERMINATOR, *API, *FunctionName);
+		GeneratedPackageCPP.Logf(TEXT("\t\tif (Func)") LINE_TERMINATOR);
+		GeneratedPackageCPP.Logf(TEXT("\t\t{") LINE_TERMINATOR);
 
 		// code to populate Parms struct
 		for (auto It = Parameters.Parms.CreateConstIterator(); It; ++It)
 		{
 			UProperty* Param = *It;
 
-			GeneratedPackageCPP.Logf(TEXT("%sParms.%s=%s;%s"), FCString::Spc(12), *Param->GetName(), *Param->GetName(), LINE_TERMINATOR);
+			GeneratedPackageCPP.Logf(TEXT("\t\t\tParms.%s=%s;") LINE_TERMINATOR, *Param->GetName(), *Param->GetName());
 		}
 
 		const FString ObjectRef = FunctionData.FunctionReference->HasAllFunctionFlags(FUNC_Const) ? FString::Printf(TEXT("const_cast<UObject*>(O)"), *ClassName) : TEXT("O");
-		GeneratedPackageCPP.Logf(TEXT("%s%s->ProcessEvent(Func, %s);%s"), FCString::Spc(12), *ObjectRef, bHasParms ? TEXT("&Parms") : TEXT("NULL"), LINE_TERMINATOR);
+		GeneratedPackageCPP.Logf(TEXT("\t\t\t%s->ProcessEvent(Func, %s);") LINE_TERMINATOR, *ObjectRef, bHasParms ? TEXT("&Parms") : TEXT("NULL"));
 
 		for (auto It = Parameters.Parms.CreateConstIterator(); It; ++It)
 		{
@@ -1417,23 +1459,23 @@ void FNativeClassHeaderGenerator::ExportInterfaceCallFunctions(const TArray<UFun
 
 			if( Param->HasAllPropertyFlags(CPF_OutParm) && !Param->HasAnyPropertyFlags(CPF_ConstParm|CPF_ReturnParm))
 			{
-				GeneratedPackageCPP.Logf(TEXT("%s%s=Parms.%s;%s"), FCString::Spc(12), *Param->GetName(), *Param->GetName(), LINE_TERMINATOR);
+				GeneratedPackageCPP.Logf(TEXT("\t\t\t%s=Parms.%s;") LINE_TERMINATOR, *Param->GetName(), *Param->GetName());
 			}
 		}
 
-		GeneratedPackageCPP.Logf(TEXT("%s}%s"), FCString::Spc(8), LINE_TERMINATOR);
+		GeneratedPackageCPP += TEXT("\t\t}") LINE_TERMINATOR;
 		
 
 		// else clause to call back into native if it's a BlueprintNativeEvent
 		if (Function->FunctionFlags & FUNC_Native)
 		{
-			GeneratedPackageCPP.Logf(TEXT("%selse if (auto I = (%sI%s*)(O->GetNativeInterfaceAddress(U%s::StaticClass())))%s"), FCString::Spc(8), ConstQualifier, *ClassName, *ClassName, LINE_TERMINATOR);
-			GeneratedPackageCPP.Logf(TEXT("%s{%s"), FCString::Spc(8), LINE_TERMINATOR);
+			GeneratedPackageCPP.Logf(TEXT("\t\telse if (auto I = (%sI%s*)(O->GetNativeInterfaceAddress(U%s::StaticClass())))") LINE_TERMINATOR, ConstQualifier, *ClassName, *ClassName);
+			GeneratedPackageCPP += TEXT("\t\t{") LINE_TERMINATOR;
 
-			GeneratedPackageCPP.Logf(FCString::Spc(12));
+			GeneratedPackageCPP += TEXT("\t\t\t");
 			if (Parameters.Return)
 			{
-				GeneratedPackageCPP.Logf(TEXT("Parms.ReturnValue = "));
+				GeneratedPackageCPP += TEXT("Parms.ReturnValue = ");
 			}
 
 			GeneratedPackageCPP.Logf(TEXT("I->%s_Implementation("), *FunctionName);
@@ -1452,17 +1494,17 @@ void FNativeClassHeaderGenerator::ExportInterfaceCallFunctions(const TArray<UFun
 				GeneratedPackageCPP.Logf(TEXT("%s"), *Param->GetName());
 			}
 
-			GeneratedPackageCPP.Logf(TEXT(");%s"), LINE_TERMINATOR);
+			GeneratedPackageCPP.Logf(TEXT(");") LINE_TERMINATOR);
 
-			GeneratedPackageCPP.Logf(TEXT("%s}%s"), FCString::Spc(8), LINE_TERMINATOR);
+			GeneratedPackageCPP.Logf(TEXT("\t\t}") LINE_TERMINATOR);
 		}
 
 		if (Parameters.Return)
 		{
-			GeneratedPackageCPP.Logf(TEXT("%sreturn Parms.ReturnValue;%s"), FCString::Spc(8), LINE_TERMINATOR);
+			GeneratedPackageCPP.Logf(TEXT("\t\treturn Parms.ReturnValue;") LINE_TERMINATOR);
 		}
 
-		GeneratedPackageCPP.Logf(TEXT("%s}%s"), FCString::Spc(4), LINE_TERMINATOR);
+		GeneratedPackageCPP.Logf(TEXT("\t}") LINE_TERMINATOR);
 	}
 }
 
@@ -1513,7 +1555,7 @@ FString GetPreservedAccessSpecifierString(FClass* Class)
 /**
  * Exports interface body macros.
  */
-void ExportInterfaceBodyMacros(FStringOutputDevice& Out, FUnrealSourceFile& SourceFile, int32 ClassGeneratedBodyLine, FClass* Class, const FString& StdCtorCall, const FString& EnhCtorCall, const FString& UInterfaceBoilerplate)
+void ExportInterfaceBodyMacros(FUHTStringBuilder& Out, FUnrealSourceFile& SourceFile, int32 ClassGeneratedBodyLine, FClass* Class, const FString& StdCtorCall, const FString& EnhCtorCall, const FString& UInterfaceBoilerplate)
 {
 	Out.Log(TEXT("#undef GENERATED_UINTERFACE_BODY_COMMON\r\n"));
 	Out.Log(Macroize(TEXT("GENERATED_UINTERFACE_BODY_COMMON()"), *UInterfaceBoilerplate));
@@ -1522,7 +1564,7 @@ void ExportInterfaceBodyMacros(FStringOutputDevice& Out, FUnrealSourceFile& Sour
 
 	auto DeprecationPushString = TEXT("PRAGMA_DISABLE_DEPRECATION_WARNINGS") LINE_TERMINATOR;
 	auto DeprecationPopString = TEXT("PRAGMA_POP") LINE_TERMINATOR;
-	auto Offset = FCString::Spc(4);
+	auto Offset = TEXT("\t");
 
 	Out.Log(Macroize(*SourceFile.GetGeneratedBodyMacroName(ClassGeneratedBodyLine, true), *(FString() +
 		Offset + DeprecationWarning +
@@ -1545,7 +1587,7 @@ void FNativeClassHeaderGenerator::ExportSourceFileHeader(FClasses& AllClasses, F
 	ExportSourceFileHeaderRecursive(AllClasses, SourceFile, VisitedSet, false);
 }
 
-void WriteMacro(FStringOutputDevice &Output, const FString& MacroName, const FString& MacroContent)
+void WriteMacro(FUHTStringBuilder& Output, const FString& MacroName, const FString& MacroContent)
 {
 	Output.Log(*Macroize(*MacroName, *MacroContent));
 }
@@ -1617,11 +1659,11 @@ void FNativeClassHeaderGenerator::ExportClassesFromSourceFileInner(FUnrealSource
 
 	for (auto* RawUClass : SourceFile.GetDefinedClasses())
 	{
-		PrologMacroCalls.Empty();
-		InClassMacroCalls.Empty();
-		InClassNoPureDeclsMacroCalls.Empty();
-		StandardUObjectConstructorsMacroCall.Empty();
-		EnhancedUObjectConstructorsMacroCall.Empty();
+		PrologMacroCalls.Reset();
+		InClassMacroCalls.Reset();
+		InClassNoPureDeclsMacroCalls.Reset();
+		StandardUObjectConstructorsMacroCall.Reset();
+		EnhancedUObjectConstructorsMacroCall.Reset();
 
 		auto* Class = (FClass*)RawUClass;
 		auto* ClassData = GScriptHelper.FindClassData(Class);
@@ -1640,7 +1682,7 @@ void FNativeClassHeaderGenerator::ExportClassesFromSourceFileInner(FUnrealSource
 		// Class definition.
 		ExportNatives(Class);
 
-		FStringOutputDevice FriendText;
+		FUHTStringBuilder FriendText;
 		ExportNativeGeneratedInitCode(Class, FriendText);
 
 		FClass* SuperClass = Class->GetSuperClass();
@@ -1652,13 +1694,13 @@ void FNativeClassHeaderGenerator::ExportClassesFromSourceFileInner(FUnrealSource
 		if (Class->HasAnyClassFlags(CLASS_Interface))
 		{
 			{
-				FStringOutputDevice UInterfaceBoilerplate;
+				FUHTStringBuilder UInterfaceBoilerplate;
 
 				// Export the class's native function registration.
 				UInterfaceBoilerplate.Logf(TEXT("    private:\r\n"));
 				UInterfaceBoilerplate.Logf(TEXT("    static void StaticRegisterNatives%s();\r\n"), NameLookupCPP.GetNameCPP(Class));
 				UInterfaceBoilerplate.Log(*FriendText);
-				FriendText.Empty();
+				FriendText.Reset();
 				UInterfaceBoilerplate.Logf(TEXT("public:\r\n"));
 
 				// Build the DECLARE_CLASS line
@@ -1703,7 +1745,7 @@ void FNativeClassHeaderGenerator::ExportClassesFromSourceFileInner(FUnrealSource
 
 
 			// Thunk functions
-			FStringOutputDevice InterfaceBoilerplate;
+			FUHTStringBuilder InterfaceBoilerplate;
 
 			InterfaceBoilerplate.Logf(TEXT("protected:\r\n\tvirtual ~%s() {}\r\npublic:\r\n"), *InterfaceCPPName);
 			InterfaceBoilerplate.Logf(TEXT("\ttypedef %s UClassType;\r\n"), ClassCPPName);
@@ -1719,7 +1761,7 @@ void FNativeClassHeaderGenerator::ExportClassesFromSourceFileInner(FUnrealSource
 
 			// Replication, add in the declaration for GetLifetimeReplicatedProps() automatically if there are any net flagged properties
 			bool bNeedsRep = false;
-			FStringOutputDevice StaticChecks;
+			FUHTStringBuilder StaticChecks;
 			auto ClassName = FString(NameLookupCPP.GetNameCPP(Class));
 			for (TFieldIterator<UProperty> It(Class, EFieldIteratorFlags::ExcludeSuper); It; ++It)
 			{
@@ -1753,11 +1795,11 @@ void FNativeClassHeaderGenerator::ExportClassesFromSourceFileInner(FUnrealSource
 
 			auto NoPureDeclsMacroName = SourceFile.GetGeneratedMacroName(ClassData, TEXT("_INCLASS_IINTERFACE_NO_PURE_DECLS"));
 			WriteMacro(GeneratedHeaderText, NoPureDeclsMacroName, InterfaceBoilerplate);
-			InClassNoPureDeclsMacroCalls.Logf(TEXT("%s%s\r\n"), FCString::Spc(4), *NoPureDeclsMacroName);
+			InClassNoPureDeclsMacroCalls.Logf(TEXT("\t%s\r\n"), *NoPureDeclsMacroName);
 
 			auto MacroName = SourceFile.GetGeneratedMacroName(ClassData, TEXT("_INCLASS_IINTERFACE"));
 			WriteMacro(GeneratedHeaderText, MacroName, InterfaceBoilerplate + StaticChecks);
-			InClassMacroCalls.Logf(TEXT("%s%s\r\n"), FCString::Spc(4), *MacroName);
+			InClassMacroCalls.Logf(TEXT("\t%s\r\n"), *MacroName);
 		}
 		else
 		{
@@ -1767,7 +1809,7 @@ void FNativeClassHeaderGenerator::ExportClassesFromSourceFileInner(FUnrealSource
 			const TCHAR* SuperClassCPPName = (SuperClass != NULL) ? NameLookupCPP.GetNameCPP(SuperClass) : NULL;
 
 			{
-				FStringOutputDevice ClassBoilerplate;
+				FUHTStringBuilder ClassBoilerplate;
 
 				// Export the class's native function registration.
 				ClassBoilerplate.Logf(TEXT("    private:\r\n"));
@@ -1809,16 +1851,16 @@ void FNativeClassHeaderGenerator::ExportClassesFromSourceFileInner(FUnrealSource
 				// @todo srobb: clean up _getUObject() code generation
 				if (((SuperClass == UClass::StaticClass()) || (SuperClass == UObject::StaticClass()) || (SuperClass && (SuperClass->ClassFlags & CLASS_Intrinsic))) && !Class->IsChildOf(UInterface::StaticClass()))
 				{
-					ClassBoilerplate.Logf(TEXT("%svirtual UObject* _getUObject() const { return const_cast<%s*>(this); }\r\n"), FCString::Spc(4), ClassCPPName);
+					ClassBoilerplate.Logf(TEXT("\tvirtual UObject* _getUObject() const { return const_cast<%s*>(this); }\r\n"), ClassCPPName);
 				}
 				else
 				{
-					ClassBoilerplate.Logf(TEXT("%svirtual UObject* _getUObject() const override { return const_cast<%s*>(this); }\r\n"), FCString::Spc(4), ClassCPPName);
+					ClassBoilerplate.Logf(TEXT("\tvirtual UObject* _getUObject() const override { return const_cast<%s*>(this); }\r\n"), ClassCPPName);
 				}
 
 				// Replication, add in the declaration for GetLifetimeReplicatedProps() automatically if there are any net flagged properties
 				bool bHasNetFlaggedProperties = false;
-				FStringOutputDevice StaticChecks;
+				FUHTStringBuilder StaticChecks;
 				for (TFieldIterator<UProperty> It(Class, EFieldIteratorFlags::ExcludeSuper); It; ++It)
 				{
 					if ((It->PropertyFlags & CPF_Net) != 0)
@@ -1851,11 +1893,11 @@ void FNativeClassHeaderGenerator::ExportClassesFromSourceFileInner(FUnrealSource
 
 				auto NoPureDeclsMacroName = SourceFile.GetGeneratedMacroName(ClassData, TEXT("_INCLASS_NO_PURE_DECLS"));
 				WriteMacro(GeneratedHeaderText, NoPureDeclsMacroName, ClassBoilerplate + StaticChecks);
-				InClassNoPureDeclsMacroCalls.Logf(TEXT("%s%s\r\n"), FCString::Spc(4), *NoPureDeclsMacroName);
+				InClassNoPureDeclsMacroCalls.Logf(TEXT("\t%s\r\n"), *NoPureDeclsMacroName);
 
 				FString MacroName = SourceFile.GetGeneratedMacroName(ClassData, TEXT("_INCLASS"));
 				WriteMacro(GeneratedHeaderText, MacroName, ClassBoilerplate);
-				InClassMacroCalls.Logf(TEXT("%s%s\r\n"), FCString::Spc(4), *MacroName);
+				InClassMacroCalls.Logf(TEXT("\t%s\r\n"), *MacroName);
 
 				ExportConstructorsMacros(SourceFile.GetGeneratedMacroName(ClassData), Class);
 			}
@@ -1905,12 +1947,12 @@ void FNativeClassHeaderGenerator::ExportClassesFromSourceFileInner(FUnrealSource
 * @param Class Class to generate constructor for.
 * @param API API string for this constructor.
 */
-void ExportCopyConstructorDefinition(FStringOutputDevice& Out, FClass* Class, const FString& API)
+void ExportCopyConstructorDefinition(FUHTStringBuilder& Out, FClass* Class, const FString& API)
 {
 	auto ClassNameCPP = NameLookupCPP.GetNameCPP(Class);
 	Out.Logf(TEXT("private:\r\n"));
-	Out.Logf(TEXT("%s/** Private copy-constructor, should never be used */\r\n"), FCString::Spc(4));
-	Out.Logf(TEXT("%s%s_API %s(const %s& InCopy);\r\n"), FCString::Spc(4), *API, ClassNameCPP, ClassNameCPP);
+	Out.Logf(TEXT("\t/** Private copy-constructor, should never be used */\r\n"));
+	Out.Logf(TEXT("\t%s_API %s(const %s& InCopy);\r\n"), *API, ClassNameCPP, ClassNameCPP);
 	Out.Logf(TEXT("public:\r\n"));
 }
 
@@ -1998,7 +2040,7 @@ void ExportVTableHelperCtorAndCaller(FStringOutputDevice& Out, FClass* Class, co
  * @param API API string for this constructor.
  * @param bExportVTableConstructors Export VTable constructors.
  */
-void ExportStandardConstructorsMacro(FStringOutputDevice& Out, FClass* Class, const FString& API
+void ExportStandardConstructorsMacro(FUHTStringBuilder& Out, FClass* Class, const FString& API
 #if WITH_HOT_RELOAD_CTORS
 		, bool bExportVTableConstructors
 #endif // WITH_HOT_RELOAD_CTORS
@@ -2006,10 +2048,10 @@ void ExportStandardConstructorsMacro(FStringOutputDevice& Out, FClass* Class, co
 {
 	if (!Class->HasAnyClassFlags(CLASS_CustomConstructor))
 	{
-		Out.Logf(TEXT("%s/** Standard constructor, called after all reflected properties have been initialized */\r\n"), FCString::Spc(4));
-		Out.Logf(TEXT("%s%s_API %s(const FObjectInitializer& ObjectInitializer);\r\n"), FCString::Spc(4), *API, NameLookupCPP.GetNameCPP(Class));
+		Out.Logf(TEXT("\t/** Standard constructor, called after all reflected properties have been initialized */\r\n"));
+		Out.Logf(TEXT("\t%s_API %s(const FObjectInitializer& ObjectInitializer);\r\n"), *API, NameLookupCPP.GetNameCPP(Class));
 	}
-	Out.Logf(TEXT("%sDEFINE_DEFAULT_OBJECT_INITIALIZER_CONSTRUCTOR_CALL(%s)\r\n"), FCString::Spc(4), NameLookupCPP.GetNameCPP(Class));
+	Out.Logf(TEXT("\tDEFINE_DEFAULT_OBJECT_INITIALIZER_CONSTRUCTOR_CALL(%s)\r\n"), NameLookupCPP.GetNameCPP(Class));
 
 #if WITH_HOT_RELOAD_CTORS
 	ExportVTableHelperCtorAndCaller(Out, Class, API, bExportVTableConstructors);
@@ -2025,12 +2067,12 @@ void ExportStandardConstructorsMacro(FStringOutputDevice& Out, FClass* Class, co
  * @param Class Class to generate constructor for.
  * @param API API string for this constructor.
  */
-void ExportConstructorDefinition(FStringOutputDevice& Out, FClass* Class, const FString& API)
+void ExportConstructorDefinition(FUHTStringBuilder& Out, FClass* Class, const FString& API)
 {
 	auto* ClassData = GScriptHelper.FindClassData(Class);
 	if (!ClassData->bConstructorDeclared)
 	{
-		Out.Logf(TEXT("%s/** Standard constructor, called after all reflected properties have been initialized */\r\n"), FCString::Spc(4));
+		Out.Logf(TEXT("\t/** Standard constructor, called after all reflected properties have been initialized */\r\n"));
 		
 		// Assume super class has OI constructor, this may not always be true but we should always be able to check this.
 		// In any case, it will default to old behaviour before we even checked this.
@@ -2046,12 +2088,12 @@ void ExportConstructorDefinition(FStringOutputDevice& Out, FClass* Class, const
 		}
 		if (bSuperClassObjectInitializerConstructorDeclared)
 		{
-			Out.Logf(TEXT("%s%s_API %s(const class FObjectInitializer& ObjectInitializer = FObjectInitializer::Get()) : Super(ObjectInitializer) { };\r\n"), FCString::Spc(4), *API, NameLookupCPP.GetNameCPP(Class));
+			Out.Logf(TEXT("\t%s_API %s(const class FObjectInitializer& ObjectInitializer = FObjectInitializer::Get()) : Super(ObjectInitializer) { };\r\n"), *API, NameLookupCPP.GetNameCPP(Class));
 			ClassData->bObjectInitializerConstructorDeclared = true;
 		}
 		else
 		{
-			Out.Logf(TEXT("%s%s_API %s() { };\r\n"), FCString::Spc(4), *API, NameLookupCPP.GetNameCPP(Class));
+			Out.Logf(TEXT("\t%s_API %s() { };\r\n"), *API, NameLookupCPP.GetNameCPP(Class));
 			ClassData->bDefaultConstructorDeclared = true;
 		}
 
@@ -2066,21 +2108,21 @@ void ExportConstructorDefinition(FStringOutputDevice& Out, FClass* Class, const
  * @param Out Output device to generate to.
  * @param Class Class to generate constructor call definition for.
  */
-void ExportDefaultConstructorCallDefinition(FStringOutputDevice& Out, FClass* Class)
+void ExportDefaultConstructorCallDefinition(FUHTStringBuilder& Out, FClass* Class)
 {
 	auto* ClassData = GScriptHelper.FindClassData(Class);
 
 	if (ClassData->bObjectInitializerConstructorDeclared)
 	{
-		Out.Logf(TEXT("%sDEFINE_DEFAULT_OBJECT_INITIALIZER_CONSTRUCTOR_CALL(%s)\r\n"), FCString::Spc(4), NameLookupCPP.GetNameCPP(Class));
+		Out.Logf(TEXT("\tDEFINE_DEFAULT_OBJECT_INITIALIZER_CONSTRUCTOR_CALL(%s)\r\n"), NameLookupCPP.GetNameCPP(Class));
 	}
 	else if (ClassData->bDefaultConstructorDeclared)
 	{
-		Out.Logf(TEXT("%sDEFINE_DEFAULT_CONSTRUCTOR_CALL(%s)\r\n"), FCString::Spc(4), NameLookupCPP.GetNameCPP(Class));
+		Out.Logf(TEXT("\tDEFINE_DEFAULT_CONSTRUCTOR_CALL(%s)\r\n"), NameLookupCPP.GetNameCPP(Class));
 	}
 	else
 	{
-		Out.Logf(TEXT("%sDEFINE_FORBIDDEN_DEFAULT_CONSTRUCTOR_CALL(%s)\r\n"), FCString::Spc(4), NameLookupCPP.GetNameCPP(Class));
+		Out.Logf(TEXT("\tDEFINE_FORBIDDEN_DEFAULT_CONSTRUCTOR_CALL(%s)\r\n"), NameLookupCPP.GetNameCPP(Class));
 	}
 }
 
@@ -2092,7 +2134,7 @@ void ExportDefaultConstructorCallDefinition(FStringOutputDevice& Out, FClass* Cl
  * @param API API string for this constructor.
  * @param bExportVTableConstructors Export VTable constructors.
  */
-void ExportEnhancedConstructorsMacro(FStringOutputDevice& Out, FClass* Class, const FString& API
+void ExportEnhancedConstructorsMacro(FUHTStringBuilder& Out, FClass* Class, const FString& API
 #if WITH_HOT_RELOAD_CTORS
 		, bool bExportVTableConstructors
 #endif // WITH_HOT_RELOAD_CTORS
@@ -2125,7 +2167,7 @@ FString GetBuildPath(FUnrealSourceFile& SourceFile)
 
 FString FNativeClassHeaderGenerator::GetListOfPublicHeaderGroupIncludesString(UPackage* Package)
 {
-	FStringOutputDevice Out;
+	FUHTStringBuilder Out;
 
 	// Fill with the rest source files from this package.
 	for (auto* SourceFile : GPublicSourceFileSet)
@@ -2154,7 +2196,8 @@ void FNativeClassHeaderGenerator::ExportConstructorsMacros(const FString& Constr
 		APIArg = TEXT("NO");
 	}
 
-	FStringOutputDevice StdMacro, EnhMacro;
+	FUHTStringBuilder StdMacro;
+	FUHTStringBuilder EnhMacro;
 	FString StdMacroName = ConstructorsMacroPrefix + TEXT("_STANDARD_CONSTRUCTORS");
 	FString EnhMacroName = ConstructorsMacroPrefix + TEXT("_ENHANCED_CONSTRUCTORS");
 
@@ -2174,8 +2217,8 @@ void FNativeClassHeaderGenerator::ExportConstructorsMacros(const FString& Constr
 	GeneratedHeaderText.Log(*Macroize(*StdMacroName, *StdMacro));
 	GeneratedHeaderText.Log(*Macroize(*EnhMacroName, *EnhMacro));
 
-	StandardUObjectConstructorsMacroCall.Logf(TEXT("%s%s\r\n"), FCString::Spc(4), *StdMacroName);
-	EnhancedUObjectConstructorsMacroCall.Logf(TEXT("%s%s\r\n"), FCString::Spc(4), *EnhMacroName);
+	StandardUObjectConstructorsMacroCall.Logf(TEXT("\t%s\r\n"), *StdMacroName);
+	EnhancedUObjectConstructorsMacroCall.Logf(TEXT("\t%s\r\n"), *EnhMacroName);
 }
 
 void FNativeClassHeaderGenerator::ExportClassesFromSourceFileWrapper(FUnrealSourceFile& SourceFile)
@@ -2208,7 +2251,7 @@ void FNativeClassHeaderGenerator::ExportClassesFromSourceFileWrapper(FUnrealSour
 
 	GeneratedHeaderText += EnumForeachText;
 
-	FStringOutputDevice GeneratedHeaderTextWithCopyright;
+	FUHTStringBuilder GeneratedHeaderTextWithCopyright;
 
 	GeneratedHeaderTextWithCopyright.Log(
 		TEXT("// Copyright 1998-2015 Epic Games, Inc. All Rights Reserved.\r\n")
@@ -2262,11 +2305,11 @@ void FNativeClassHeaderGenerator::ExportClassesFromSourceFileWrapper(FUnrealSour
 	ListOfAllUObjectHeaderIncludes                .Logf(TEXT("%s"), *IncludeStr);
 	ListOfPublicClassesUObjectHeaderModuleIncludes.Logf(TEXT("%s"), *IncludeStr);
 
-	ForwardDeclarations.Empty();
-	GeneratedHeaderTextBeforeForwardDeclarations.Empty();
-	GeneratedForwardDeclarations.Empty();
-	GeneratedHeaderText.Empty();
-	EnumForeachText.Empty();
+	ForwardDeclarations.Reset();
+	GeneratedHeaderTextBeforeForwardDeclarations.Reset();
+	GeneratedForwardDeclarations.Reset();
+	GeneratedHeaderText.Reset();
+	EnumForeachText.Reset();
 }
 
 void FNativeClassHeaderGenerator::ExportSourceFileHeaderRecursive(FClasses& AllClasses, FUnrealSourceFile* SourceFile, TSet<const FUnrealSourceFile*>& VisitedSet, bool bCheckDependenciesOnly)
@@ -2422,7 +2465,7 @@ void FNativeClassHeaderGenerator::ExportGeneratedStructBodyMacros(FUnrealSourceF
 			GeneratedHeaderText.Log(*Macroized);
 
 			const TCHAR* StructNameCPP = NameLookupCPP.GetNameCPP(Struct);
-			FString SingletonName = StaticConstructionString.Replace(TEXT("()"), TEXT("")); // function address
+			FString SingletonName = StaticConstructionString.Replace(TEXT("()"), TEXT(""), ESearchCase::CaseSensitive); // function address
 			FString GetCRCName = FString::Printf(TEXT("Get_%s_CRC"), *SingletonName);
 
 			GeneratedPackageCPP.Logf(TEXT("class UScriptStruct* %s::StaticStruct()\r\n"), StructNameCPP);
@@ -2480,7 +2523,7 @@ void FNativeClassHeaderGenerator::ExportGeneratedStructBodyMacros(FUnrealSourceF
 			SingletonNameToExternDecl.Add(SingletonName, Extern);
 			GeneratedFunctionDeclarations.Log(*Extern);
 
-			FStringOutputDevice GeneratedStructRegisterFunctionText;
+			FUHTStringBuilder GeneratedStructRegisterFunctionText;
 
 			GeneratedStructRegisterFunctionText.Logf(TEXT("    UScriptStruct* %s\r\n"), *SingletonName);
 			GeneratedStructRegisterFunctionText.Logf(TEXT("    {\r\n"));
@@ -2491,7 +2534,7 @@ void FNativeClassHeaderGenerator::ExportGeneratedStructBodyMacros(FUnrealSourceF
 			if (Structs.Num())
 			{
 				TGuardValue<bool> Guard(bIsExportingForOffsetDeterminationOnly, true);
-				ExportMirrorsForNoexportStructs(Structs, 8, GeneratedStructRegisterFunctionText);
+				ExportMirrorsForNoexportStructs(Structs, /*Indent=*/ 2, GeneratedStructRegisterFunctionText);
 			}
 
 			// Structs can either have a UClass or UPackage as outer (if delcared in non-UClass header).
@@ -2561,7 +2604,7 @@ void FNativeClassHeaderGenerator::ExportGeneratedStructBodyMacros(FUnrealSourceF
 
 			auto& GeneratedFunctionText = GetGeneratedFunctionTextDevice();
 			GeneratedFunctionText += GeneratedStructRegisterFunctionText;
-			GeneratedFunctionText.Logf(TEXT("    uint32 Get_%s_CRC() { return %uU; }\r\n"), *SingletonName.Replace(TEXT("()"), TEXT("")), StructCrc);
+			GeneratedFunctionText.Logf(TEXT("    uint32 Get_%s_CRC() { return %uU; }\r\n"), *SingletonName.Replace(TEXT("()"), TEXT(""), ESearchCase::CaseSensitive), StructCrc);
 
 			//CallSingletons.Logf(TEXT("                OuterClass->LinkChild(%s); // %u\r\n"), *SingletonName, StructCrc);
 		}
@@ -2584,7 +2627,7 @@ void FNativeClassHeaderGenerator::ExportGeneratedEnumsInitCode(const TArray<UEnu
 		const FString FriendApiString = GetAPIString();
 		const FString StaticConstructionString = GetSingletonName(Enum);
 
-		FString SingletonName = StaticConstructionString.Replace(TEXT("()"), TEXT("")); // function address
+		FString SingletonName = StaticConstructionString.Replace(TEXT("()"), TEXT(""), ESearchCase::CaseSensitive); // function address
 
 		GeneratedPackageCPP.Logf(TEXT("static class UEnum* %s_StaticEnum()\r\n"), *Enum->GetName());
 		GeneratedPackageCPP.Logf(TEXT("{\r\n"));
@@ -2609,7 +2652,7 @@ void FNativeClassHeaderGenerator::ExportGeneratedEnumsInitCode(const TArray<UEnu
 			SingletonNameToExternDecl.Add(SingletonName, Extern);
 			GeneratedFunctionDeclarations.Log(*Extern);
 
-			FStringOutputDevice GeneratedEnumRegisterFunctionText;
+			FUHTStringBuilder GeneratedEnumRegisterFunctionText;
 
 			GeneratedEnumRegisterFunctionText.Logf(TEXT("    UEnum* %s\r\n"), *SingletonName);
 			GeneratedEnumRegisterFunctionText.Logf(TEXT("    {\r\n"));
@@ -2665,7 +2708,7 @@ void FNativeClassHeaderGenerator::ExportGeneratedEnumsInitCode(const TArray<UEnu
 	}
 }
 
-void FNativeClassHeaderGenerator::ExportMirrorsForNoexportStructs(const TArray<UScriptStruct*>& NativeStructs, int32 TextIndent, FStringOutputDevice& HeaderOutput)
+void FNativeClassHeaderGenerator::ExportMirrorsForNoexportStructs(const TArray<UScriptStruct*>& NativeStructs, int32 TextIndent, FUHTStringBuilder& HeaderOutput)
 {
 	// reverse the order.
 	for (int32 i = NativeStructs.Num() - 1; i >= 0; --i)
@@ -2673,19 +2716,17 @@ void FNativeClassHeaderGenerator::ExportMirrorsForNoexportStructs(const TArray<U
 		UScriptStruct* Struct = NativeStructs[i];
 
 		// Export struct.
-		HeaderOutput.Logf( TEXT("%sstruct %s"), FCString::Spc(TextIndent), NameLookupCPP.GetNameCPP( Struct ) );
+		HeaderOutput.Logf(TEXT("%sstruct %s"), GetTabs(TextIndent), NameLookupCPP.GetNameCPP(Struct));
 		if (Struct->GetSuperStruct() != NULL)
 		{
 			HeaderOutput.Logf(TEXT(" : public %s"), NameLookupCPP.GetNameCPP(Struct->GetSuperStruct()));
 		}
-		HeaderOutput.Logf( TEXT("\r\n%s{\r\n"), FCString::Spc(TextIndent) );
+		HeaderOutput.Logf(TEXT("\r\n%s{\r\n"), GetTabs(TextIndent));
 
 		// Export the struct's CPP properties.
 		ExportProperties(Struct, TextIndent, /*bAccessSpecifiers=*/ false, &HeaderOutput);
 
-		HeaderOutput.Logf( TEXT("%s};\r\n"), FCString::Spc(TextIndent) );
-
-		HeaderOutput.Log( TEXT("\r\n") );
+		HeaderOutput.Logf(TEXT("%s};\r\n\r\n"), GetTabs(TextIndent));
 	}
 }
 
@@ -2698,10 +2739,10 @@ bool FNativeClassHeaderGenerator::WillExportEventParms( UFunction* Function )
 	return false;
 }
 
-void WriteEventFunctionPrologue(FStringOutputDevice& Output, int32 Indent, const FParmsAndReturnProperties& Parameters, UObject* FunctionOuter, const TCHAR* FunctionName)
+void WriteEventFunctionPrologue(FUHTStringBuilder& Output, int32 Indent, const FParmsAndReturnProperties& Parameters, UObject* FunctionOuter, const TCHAR* FunctionName)
 {
 	// now the body - first we need to declare a struct which will hold the parameters for the event/delegate call
-	Output.Logf( TEXT("\r\n%s{\r\n"), FCString::Spc(Indent) );
+	Output.Logf(TEXT("\r\n%s{\r\n"), GetTabs(Indent));
 
 	// declare and zero-initialize the parameters and return value, if applicable
 	if (!Parameters.HasParms())
@@ -2709,7 +2750,7 @@ void WriteEventFunctionPrologue(FStringOutputDevice& Output, int32 Indent, const
 
 	FString EventStructName = GetEventStructParamsName(FunctionOuter, FunctionName);
 
-	Output.Logf( TEXT("%s%s Parms;\r\n"), FCString::Spc(Indent + 4), *EventStructName );
+	Output.Logf(TEXT("%s%s Parms;\r\n"), GetTabs(Indent + 1), *EventStructName );
 
 	// Declare a parameter struct for this event/delegate and assign the struct members using the values passed into the event/delegate call.
 	for (auto It = Parameters.Parms.CreateConstIterator(); It; ++It)
@@ -2717,9 +2758,9 @@ void WriteEventFunctionPrologue(FStringOutputDevice& Output, int32 Indent, const
 		UProperty* Prop = *It;
 
 		const FString PropertyName = Prop->GetName();
-		if ( Prop->ArrayDim > 1 )
+		if (Prop->ArrayDim > 1)
 		{
-			Output.Logf( TEXT("%sFMemory::Memcpy(Parms.%s,%s,sizeof(Parms.%s));\r\n"), FCString::Spc(Indent + 4), *PropertyName, *PropertyName, *PropertyName );
+			Output.Logf(TEXT("%sFMemory::Memcpy(Parms.%s,%s,sizeof(Parms.%s));\r\n"), GetTabs(Indent + 1), *PropertyName, *PropertyName, *PropertyName);
 		}
 		else
 		{
@@ -2729,12 +2770,12 @@ void WriteEventFunctionPrologue(FStringOutputDevice& Output, int32 Indent, const
 				ValueAssignmentText += TEXT(" ? true : false");
 			}
 
-			Output.Logf( TEXT("%sParms.%s=%s;\r\n"), FCString::Spc(Indent + 4), *PropertyName, *ValueAssignmentText );
+			Output.Logf(TEXT("%sParms.%s=%s;\r\n"), GetTabs(Indent + 1), *PropertyName, *ValueAssignmentText);
 		}
 	}
 }
 
-void WriteEventFunctionEpilogue(FStringOutputDevice& Output, int32 Indent, const FParmsAndReturnProperties& Parameters, const TCHAR* FunctionName)
+void WriteEventFunctionEpilogue(FUHTStringBuilder& Output, int32 Indent, const FParmsAndReturnProperties& Parameters, const TCHAR* FunctionName)
 {
 	// Out parm copying.
 	for (auto It = Parameters.Parms.CreateConstIterator(); It; ++It)
@@ -2746,11 +2787,11 @@ void WriteEventFunctionEpilogue(FStringOutputDevice& Output, int32 Indent, const
 			const FString PropertyName = Prop->GetName();
 			if ( Prop->ArrayDim > 1 )
 			{
-				Output.Logf( TEXT("%sFMemory::Memcpy(&%s,&Parms.%s,sizeof(%s));\r\n"), FCString::Spc(Indent + 4), *PropertyName, *PropertyName, *PropertyName );
+				Output.Logf(TEXT("%sFMemory::Memcpy(&%s,&Parms.%s,sizeof(%s));\r\n"), GetTabs(Indent + 1), *PropertyName, *PropertyName, *PropertyName);
 			}
 			else
 			{
-				Output.Logf(TEXT("%s%s=Parms.%s;\r\n"), FCString::Spc(Indent + 4), *PropertyName, *PropertyName);
+				Output.Logf(TEXT("%s%s=Parms.%s;\r\n"), GetTabs(Indent + 1), *PropertyName, *PropertyName);
 			}
 		}
 	}
@@ -2760,19 +2801,19 @@ void WriteEventFunctionEpilogue(FStringOutputDevice& Output, int32 Indent, const
 	{
 		// Make sure uint32 -> bool is supported
 		bool bBoolProperty = Parameters.Return->IsA(UBoolProperty::StaticClass());
-		Output.Logf( TEXT("%sreturn %sParms.%s;\r\n"), FCString::Spc(Indent + 4), bBoolProperty ? TEXT("!!") : TEXT(""), *Parameters.Return->GetName() );
+		Output.Logf(TEXT("%sreturn %sParms.%s;\r\n"), GetTabs(Indent + 1), bBoolProperty ? TEXT("!!") : TEXT(""), *Parameters.Return->GetName());
 	}
-	Output.Logf( TEXT("%s}\r\n"), FCString::Spc(Indent) );
+	Output.Logf(TEXT("%s}\r\n"), GetTabs(Indent));
 }
 
 void FNativeClassHeaderGenerator::ExportDelegateDefinitions(FUnrealSourceFile& SourceFile, const TArray<UDelegateFunction*>& DelegateFunctions, const bool bWrapperImplementationsOnly)
 {
-	FStringOutputDevice HeaderOutput;
+	FUHTStringBuilder HeaderOutput;
 
 	for ( int32 i = DelegateFunctions.Num() - 1; i >= 0 ; i-- )
 	{
 		UFunction* Function = DelegateFunctions[i];
-		FStringOutputDevice DelegateOutput;
+		FUHTStringBuilder DelegateOutput;
 		check( Function->HasAnyFunctionFlags( FUNC_Delegate ) );
 
 		if (bWrapperImplementationsOnly)
@@ -2827,7 +2868,7 @@ void FNativeClassHeaderGenerator::ExportDelegateDefinitions(FUnrealSourceFile& S
 			{
 				const TCHAR* DelegateType = bIsMulticastDelegate ? TEXT( "ProcessMulticastDelegate" ) : TEXT( "ProcessDelegate" );
 				const TCHAR* DelegateArg  = Parameters.HasParms() ? TEXT("&Parms") : TEXT("NULL");
-				DelegateOutput.Logf(TEXT("%s%s.%s<UObject>(%s);\r\n"), FCString::Spc(4), *DelegateName, DelegateType, DelegateArg);
+				DelegateOutput.Logf(TEXT("\t%s.%s<UObject>(%s);\r\n"), *DelegateName, DelegateType, DelegateArg);
 			}
 			WriteEventFunctionEpilogue(DelegateOutput, 0, Parameters, *DelegateName);
 		}
@@ -2861,8 +2902,8 @@ void FNativeClassHeaderGenerator::ExportDelegateDefinitions(FUnrealSourceFile& S
  */
 void ExportProtoDeclaration(FOutputDevice& Out, const FString& MessageName, TFieldIterator<UProperty>& Properties, uint64 PropertyFlags, int32 Indent)
 {
-	Out.Logf(TEXT("%smessage CMsg%sMessage\r\n"), FCString::Spc(Indent), *MessageName);
-	Out.Logf(TEXT("%s{\r\n"), FCString::Spc(Indent));
+	Out.Logf(TEXT("%smessage CMsg%sMessage\r\n"), GetTabs(Indent), *MessageName);
+	Out.Logf(TEXT("%s{\r\n"), GetTabs(Indent));
 
 	static TMap<FString, FString> ToProtoTypeMappings;
 	static bool bInitMapping = false;
@@ -2927,7 +2968,7 @@ void ExportProtoDeclaration(FOutputDevice& Out, const FString& MessageName, TFie
 			if (StructProperty != NULL)
 			{
 				TFieldIterator<UProperty> StructIt(StructProperty->Struct);
-				ExportProtoDeclaration(Out, VariableTypeName, StructIt, CPF_AllFlags, Indent + 4);
+				ExportProtoDeclaration(Out, VariableTypeName, StructIt, CPF_AllFlags, Indent + 1);
 				VariableTypeName = FString::Printf(TEXT("CMsg%sMessage"), *VariableTypeName);
 				ProtoTypeName = &VariableTypeName;
 			}
@@ -2936,15 +2977,8 @@ void ExportProtoDeclaration(FOutputDevice& Out, const FString& MessageName, TFie
 				ProtoTypeName = ToProtoTypeMappings.Find(VariableTypeName);
 			}
 
-			Out.Log(FCString::Spc(Indent + 4));
-			if (bIsRepeated)
-			{
-				Out.Log( TEXT("repeated "));
-			}
-			else
-			{
-				Out.Log( TEXT("optional "));
-			}
+			Out.Log(GetTabs(Indent + 1));
+			Out.Log(bIsRepeated ? TEXT("repeated ") : TEXT("optional "));
 
 			if (ProtoTypeName != NULL)
 			{
@@ -2962,7 +2996,7 @@ void ExportProtoDeclaration(FOutputDevice& Out, const FString& MessageName, TFie
 		}
 	}
 
-	Out.Logf( TEXT("%s}\r\n"), FCString::Spc(Indent) );
+	Out.Logf(TEXT("%s}\r\n"), GetTabs(Indent));
 }
 
 /**
@@ -2972,10 +3006,10 @@ void ExportProtoDeclaration(FOutputDevice& Out, const FString& MessageName, TFie
  * @param Indent starting indentation level
  * @param Output optional output redirect
  */
-void FNativeClassHeaderGenerator::ExportProtoMessage(const TArray<UFunction*>& InCallbackFunctions, FClassMetaData* ClassData, int32 Indent, class FStringOutputDevice* Output)
+void FNativeClassHeaderGenerator::ExportProtoMessage(const TArray<UFunction*>& InCallbackFunctions, FClassMetaData* ClassData, int32 Indent, FUHTStringBuilder* Output)
 {
 	// Parms struct definitions.
-	FStringOutputDevice HeaderOutput;
+	FUHTStringBuilder HeaderOutput;
 
 	TArray<UFunction*> CallbackFunctions = InCallbackFunctions;
 	CallbackFunctions.Sort();
@@ -3010,7 +3044,7 @@ void FNativeClassHeaderGenerator::ExportProtoMessage(const TArray<UFunction*>& I
 					ParameterList += ParamName;
 				}
 
-				HeaderOutput.Logf(TEXT("// %s%s(%s)\r\n"), FCString::Spc(Indent), *FunctionName, *ParameterList);
+				HeaderOutput.Logf(TEXT("// %s%s(%s)\r\n"), GetTabs(Indent), *FunctionName, *ParameterList);
 
 				TFieldIterator<UProperty> ParamIt(Function);
 				ExportProtoDeclaration(HeaderOutput, FunctionName, ParamIt, CPF_Parm, Indent);
@@ -3050,8 +3084,8 @@ static FString FixJavaName(const FString &StringIn)
  */
 void ExportMCPDeclaration(FOutputDevice& Out, const FString& MessageName, TFieldIterator<UProperty>& Properties, uint64 PropertyFlags, int32 Indent)
 {
-	Out.Logf(TEXT("%spublic class %sCommand extends ProfileCommand\r\n"), FCString::Spc(Indent), *MessageName);
-	Out.Logf(TEXT("%s{\r\n"), FCString::Spc(Indent));
+	Out.Logf(TEXT("%spublic class %sCommand extends ProfileCommand\r\n"), GetTabs(Indent), *MessageName);
+	Out.Logf(TEXT("%s{\r\n"), GetTabs(Indent));
 
 	static TMap<FString, FString> ToMCPTypeMappings;
 	static TMap<FString, FString> ToAnnotationMappings;
@@ -3134,7 +3168,7 @@ void ExportMCPDeclaration(FOutputDevice& Out, const FString& MessageName, TField
 			{
 				// TODO Implement structs
 /*				TFieldIterator<UProperty> StructIt(StructProperty->Struct);
-				ExportMCPDeclaration(Out, VariableTypeName, StructIt, CPF_AllFlags, Indent + 4);
+				ExportMCPDeclaration(Out, VariableTypeName, StructIt, CPF_AllFlags, Indent + 1);
 				VariableTypeName = FString::Printf(TEXT("CMsg%sMessage"), *VariableTypeName);
 				MCPTypeName = &VariableTypeName;*/
 				continue;
@@ -3157,17 +3191,17 @@ void ExportMCPDeclaration(FOutputDevice& Out, const FString& MessageName, TField
 
 			if (AnnotationName != NULL && !AnnotationName->IsEmpty())
 			{
-				Out.Log(FCString::Spc(Indent + 4));
+				Out.Log(GetTabs(Indent + 1));
 				Out.Logf(TEXT("%s\r\n"), **AnnotationName);
 			}
 
 			if (MCPTypeName != NULL)
 			{
-				Out.Log(FCString::Spc(Indent + 4));
+				Out.Log(GetTabs(Indent + 1));
 				Out.Logf(TEXT("private %s %s;\r\n"), **MCPTypeName, *PropertyName);
 
 				ConstructorParams += FString::Printf(TEXT(", %s %s"), **MCPTypeName, *PropertyName);
-				ConstructorText += FString::Printf(TEXT("%sthis.%s = %s;\r\n"), FCString::Spc(Indent + 8), *PropertyName, *PropertyName);
+				ConstructorText += FString::Printf(TEXT("%sthis.%s = %s;\r\n"), GetTabs(Indent + 2), *PropertyName, *PropertyName);
 			}
 			else
 			{
@@ -3180,18 +3214,18 @@ void ExportMCPDeclaration(FOutputDevice& Out, const FString& MessageName, TField
 		}
 	}
 
-	Out.Logf(TEXT("\r\n%spublic %sCommand(String epicId, String profileId%s)\r\n"), FCString::Spc(Indent + 4), *MessageName, *ConstructorParams);
-	Out.Logf(TEXT("%s{\r\n"), FCString::Spc(Indent + 4));
-	Out.Logf(TEXT("%ssuper(epicId, profileId);\r\n"), FCString::Spc(Indent + 8));
+	Out.Logf(TEXT("\r\n%spublic %sCommand(String epicId, String profileId%s)\r\n"), GetTabs(Indent + 1), *MessageName, *ConstructorParams);
+	Out.Logf(TEXT("%s{\r\n"), GetTabs(Indent + 1));
+	Out.Logf(TEXT("%ssuper(epicId, profileId);\r\n"), GetTabs(Indent + 2));
 	Out.Logf(TEXT("%s"), *ConstructorText);
-	Out.Logf(TEXT("%s}\r\n"), FCString::Spc(Indent + 4));
+	Out.Logf(TEXT("%s}\r\n"), GetTabs(Indent + 1));
 
-	Out.Logf(TEXT("\r\n%s@Override\r\n"), FCString::Spc(Indent + 4));
-	Out.Logf(TEXT("%sprotected void execute(@Name(\"profile\") @NotNull ProfileEx profile)\r\n"), FCString::Spc(Indent + 4));
-	Out.Logf(TEXT("%s{\r\n"), FCString::Spc(Indent + 4));
-	Out.Logf(TEXT("%s}\r\n"), FCString::Spc(Indent + 4));
+	Out.Logf(TEXT("\r\n%s@Override\r\n"), GetTabs(Indent + 1));
+	Out.Logf(TEXT("%sprotected void execute(@Name(\"profile\") @NotNull ProfileEx profile)\r\n"), GetTabs(Indent + 1));
+	Out.Logf(TEXT("%s{\r\n"), GetTabs(Indent + 1));
+	Out.Logf(TEXT("%s}\r\n"), GetTabs(Indent + 1));
 
-	Out.Logf( TEXT("%s}\r\n"), FCString::Spc(Indent) );
+	Out.Logf(TEXT("%s}\r\n"), GetTabs(Indent));
 }
 
 /**
@@ -3201,10 +3235,10 @@ void ExportMCPDeclaration(FOutputDevice& Out, const FString& MessageName, TField
  * @param Indent starting indentation level
  * @param Output optional output redirect
  */
-void FNativeClassHeaderGenerator::ExportMCPMessage(const TArray<UFunction*>& InCallbackFunctions, FClassMetaData* ClassData, int32 Indent, class FStringOutputDevice* Output)
+void FNativeClassHeaderGenerator::ExportMCPMessage(const TArray<UFunction*>& InCallbackFunctions, FClassMetaData* ClassData, int32 Indent, FUHTStringBuilder* Output)
 {
 	// Parms struct definitions.
-	FStringOutputDevice HeaderOutput;
+	FUHTStringBuilder HeaderOutput;
 
 	TArray<UFunction*> CallbackFunctions = InCallbackFunctions;
 	CallbackFunctions.Sort();
@@ -3239,7 +3273,7 @@ void FNativeClassHeaderGenerator::ExportMCPMessage(const TArray<UFunction*>& InC
 					ParameterList += ParamName;
 				}
 
-				HeaderOutput.Logf(TEXT("// %s%s(%s)\r\n"), FCString::Spc(Indent), *FunctionName, *ParameterList);
+				HeaderOutput.Logf(TEXT("// %s%s(%s)\r\n"), GetTabs(Indent), *FunctionName, *ParameterList);
 
 				TFieldIterator<UProperty> ParamIt(Function);
 				ExportMCPDeclaration(HeaderOutput, FunctionName, ParamIt, CPF_Parm, Indent);
@@ -3257,7 +3291,7 @@ void FNativeClassHeaderGenerator::ExportMCPMessage(const TArray<UFunction*>& InC
 	}
 }
 
-void FNativeClassHeaderGenerator::ExportEventParm(UFunction* Function, FStringOutputDevice &HeaderOutput, int32 Indent, bool bOutputConstructor)
+void FNativeClassHeaderGenerator::ExportEventParm(UFunction* Function, FUHTStringBuilder& HeaderOutput, int32 Indent, bool bOutputConstructor)
 {
 	if (!WillExportEventParms(Function))
 	{
@@ -3271,16 +3305,16 @@ void FNativeClassHeaderGenerator::ExportEventParm(UFunction* Function, FStringOu
 	}
 
 	FString EventParmStructName = GetEventStructParamsName(Function->GetOuter(), *FunctionName);
-	HeaderOutput.Logf(TEXT("%sstruct %s\r\n"), FCString::Spc(Indent), *EventParmStructName);
-	HeaderOutput.Logf(TEXT("%s{\r\n"), FCString::Spc(Indent));
+	HeaderOutput.Logf(TEXT("%sstruct %s\r\n"), GetTabs(Indent), *EventParmStructName);
+	HeaderOutput.Logf(TEXT("%s{\r\n"), GetTabs(Indent));
 
 	for (TFieldIterator<UProperty> It(Function); It && (It->PropertyFlags&CPF_Parm); ++It)
 	{
 		UProperty* Prop = *It;
 		ForwardDeclarations.Add(Prop);
 
-		FStringOutputDevice PropertyText;
-		PropertyText.Log(FCString::Spc(4 + Indent));
+		FUHTStringBuilder PropertyText;
+		PropertyText.Log(GetTabs(Indent + 1));
 
 		bool bEmitConst = Prop->HasAnyPropertyFlags(CPF_ConstParm) && Prop->IsA<UObjectProperty>();
 
@@ -3313,7 +3347,7 @@ void FNativeClassHeaderGenerator::ExportEventParm(UFunction* Function, FStringOu
 	UProperty* Prop = Function->GetReturnProperty();
 	if (Prop && bOutputConstructor)
 	{
-		FStringOutputDevice InitializationAr;
+		FUHTStringBuilder InitializationAr;
 
 		UStructProperty* InnerStruct = Cast<UStructProperty>(Prop);
 		bool bNeedsOutput = true;
@@ -3336,20 +3370,20 @@ void FNativeClassHeaderGenerator::ExportEventParm(UFunction* Function, FStringOu
 		if (bNeedsOutput)
 		{
 			check(Prop->ArrayDim == 1); // can't return arrays
-			HeaderOutput.Logf(TEXT("\r\n%s/** Constructor, intializes return property only **/\r\n"), FCString::Spc(4 + Indent));
-			HeaderOutput.Logf(TEXT("%s%s()\r\n"), FCString::Spc(4 + Indent), *EventParmStructName);
-			HeaderOutput.Logf(TEXT("%s%s %s(%s)\r\n"), FCString::Spc(8 + Indent), TEXT(":"), *Prop->GetName(), *GetNullParameterValue(Prop, false, true));
-			HeaderOutput.Logf(TEXT("%s{\r\n"), FCString::Spc(4 + Indent));
-			HeaderOutput.Logf(TEXT("%s}\r\n"), FCString::Spc(4 + Indent));
+			HeaderOutput.Logf(TEXT("\r\n%s/** Constructor, intializes return property only **/\r\n"), GetTabs(Indent + 1));
+			HeaderOutput.Logf(TEXT("%s%s()\r\n"), GetTabs(Indent + 1), *EventParmStructName);
+			HeaderOutput.Logf(TEXT("%s%s %s(%s)\r\n"), GetTabs(Indent + 2), TEXT(":"), *Prop->GetName(), *GetNullParameterValue(Prop, false, true));
+			HeaderOutput.Logf(TEXT("%s{\r\n"), GetTabs(Indent + 1));
+			HeaderOutput.Logf(TEXT("%s}\r\n"), GetTabs(Indent + 1));
 		}
 	}
-	HeaderOutput.Logf(TEXT("%s};\r\n"), FCString::Spc(Indent));
+	HeaderOutput.Logf(TEXT("%s};\r\n"), GetTabs(Indent));
 }
 
-void FNativeClassHeaderGenerator::ExportEventParms(FScope& Scope, const TArray<UFunction*>& InCallbackFunctions, FStringOutputDevice& Output, int32 Indent, bool bOutputConstructor)
+void FNativeClassHeaderGenerator::ExportEventParms(FScope& Scope, const TArray<UFunction*>& InCallbackFunctions, FUHTStringBuilder& Output, int32 Indent, bool bOutputConstructor)
 {
 	// Parms struct definitions.
-	FStringOutputDevice HeaderOutput;
+	FUHTStringBuilder HeaderOutput;
 
 	TArray<UFunction*> CallbackFunctions = InCallbackFunctions;
 	CallbackFunctions.Sort();
@@ -3442,7 +3476,8 @@ FString FNativeClassHeaderGenerator::GetFunctionReturnString(UFunction* Function
 		FString ExtendedReturnType;
 		ForwardDeclarations.Add(Return);
 		FString ReturnType = Return->GetCPPType(&ExtendedReturnType, CPPF_ArgumentOrReturnValue);
-		FStringOutputDevice ReplacementText(*ReturnType);
+		FUHTStringBuilder ReplacementText;
+		ReplacementText += ReturnType;
 		ApplyAlternatePropertyExportText(Return, ReplacementText);
 		return ReplacementText + ExtendedReturnType;
 	}
@@ -3466,7 +3501,7 @@ FString GetFunctionConstModifierString(UFunction* Function)
 	return FString();
 }
 
-void FNativeClassHeaderGenerator::ExportFunctionChecks(const FFuncInfo& FunctionData, FStringOutputDevice& CheckClasses, FStringOutputDevice& StaticChecks, const FString& ClassName)
+void FNativeClassHeaderGenerator::ExportFunctionChecks(const FFuncInfo& FunctionData, FUHTStringBuilder& CheckClasses, FUHTStringBuilder& StaticChecks, const FString& ClassName)
 {
 	auto Function = FunctionData.FunctionReference;
 	auto FunctionReturnType = GetFunctionReturnString(Function);
@@ -3500,7 +3535,7 @@ void FNativeClassHeaderGenerator::ExportFunctionChecks(const FFuncInfo& Function
 	//
 	// Coin static_assert message
 	//
-	FStringOutputDevice AssertMessage;
+	FUHTStringBuilder AssertMessage;
 	AssertMessage.Logf(TEXT("Function %s was marked as %s"), *(Function->GetName()), *FunctionSpecifiers[0]);
 	for (int32 i = 1; i < FunctionSpecifiers.Num(); ++i)
 	{
@@ -3528,7 +3563,7 @@ void FNativeClassHeaderGenerator::ExportFunctionChecks(const FFuncInfo& Function
 	}
 }
 
-void FNativeClassHeaderGenerator::ExportNativeFunctionHeader( const FFuncInfo& FunctionData, FStringOutputDevice& HeaderOutput, EExportFunctionType::Type FunctionType, EExportFunctionHeaderStyle::Type FunctionHeaderStyle, const TCHAR* ExtraParam )
+void FNativeClassHeaderGenerator::ExportNativeFunctionHeader(const FFuncInfo& FunctionData, FUHTStringBuilder& HeaderOutput, EExportFunctionType::Type FunctionType, EExportFunctionHeaderStyle::Type FunctionHeaderStyle, const TCHAR* ExtraParam)
 {
 	UFunction* Function = FunctionData.FunctionReference;
 
@@ -3536,7 +3571,10 @@ void FNativeClassHeaderGenerator::ExportNativeFunctionHeader( const FFuncInfo& F
 	const bool bIsInterface  = !bIsDelegate && Function->GetOwnerClass()->HasAnyClassFlags(CLASS_Interface);
 	const bool bIsK2Override = Function->HasAnyFunctionFlags( FUNC_BlueprintEvent );
 	
-	HeaderOutput.Log( FCString::Spc( bIsDelegate ? 0 : 4) );
+	if (!bIsDelegate)
+	{
+		HeaderOutput += TEXT("\t");
+	}
 
 	if (FunctionHeaderStyle == EExportFunctionHeaderStyle::Declaration)
 	{
@@ -3580,7 +3618,8 @@ void FNativeClassHeaderGenerator::ExportNativeFunctionHeader( const FFuncInfo& F
 		FString ExtendedReturnType;
 		FString ReturnType = Return->GetCPPType(&ExtendedReturnType, (FunctionHeaderStyle == EExportFunctionHeaderStyle::Definition && (FunctionType != EExportFunctionType::Interface) ? CPPF_Implementation : 0) | CPPF_ArgumentOrReturnValue);
 		ForwardDeclarations.Add(Return);
-		FStringOutputDevice ReplacementText(*ReturnType);
+		FUHTStringBuilder ReplacementText;
+		ReplacementText += ReturnType;
 		ApplyAlternatePropertyExportText(Return, ReplacementText);
 		HeaderOutput.Logf(TEXT("%s%s"), *ReplacementText, *ExtendedReturnType);
 	}
@@ -3629,7 +3668,7 @@ void FNativeClassHeaderGenerator::ExportNativeFunctionHeader( const FFuncInfo& F
 			HeaderOutput.Log(TEXT(", "));
 		}
 
-		FStringOutputDevice PropertyText;
+		FUHTStringBuilder PropertyText;
 
 		const FString* Dim = GArrayDimensions.Find(Property);
 		Property->ExportCppDeclaration( PropertyText, EExportedDeclaration::Parameter, Dim ? **Dim : NULL );
@@ -3663,7 +3702,7 @@ void FNativeClassHeaderGenerator::ExportNativeFunctionHeader( const FFuncInfo& F
  * @param Return return parameter for the function
  * @param DeprecationWarningOutputDevice Device to output deprecation warnings for _Validate and _Implementation functions.
  */
-void FNativeClassHeaderGenerator::ExportFunctionThunk(FStringOutputDevice& RPCWrappers, UFunction* Function, const FFuncInfo& FunctionData, const TArray<UProperty*>& Parameters, UProperty* Return, FStringOutputDevice& DeprecationWarningOutputDevice)
+void FNativeClassHeaderGenerator::ExportFunctionThunk(FUHTStringBuilder& RPCWrappers, UFunction* Function, const FFuncInfo& FunctionData, const TArray<UProperty*>& Parameters, UProperty* Return, FUHTStringBuilder& DeprecationWarningOutputDevice)
 {
 	// export the GET macro for this parameter
 	FString ParameterList;
@@ -3686,7 +3725,8 @@ void FNativeClassHeaderGenerator::ExportFunctionThunk(FStringOutputDevice& RPCWr
 		{
 			EvalBaseText += Param->GetCPPMacroType(TypeText);
 		}
-		FStringOutputDevice ReplacementText(*TypeText);
+		FUHTStringBuilder ReplacementText;
+		ReplacementText += TypeText;
 		ApplyAlternatePropertyExportText(Param, ReplacementText);
 		TypeText = ReplacementText;
 
@@ -3710,7 +3750,7 @@ void FNativeClassHeaderGenerator::ExportFunctionThunk(FStringOutputDevice& RPCWr
 
 		EvalParameterText = FString::Printf(TEXT("(%s%s%s)"), *TypeText, *ParamName, *DefaultValueText);
 
-		RPCWrappers.Logf(TEXT("%s%s%s%s;%s"), FCString::Spc(8), *EvalBaseText, *EvalModifierText, *EvalParameterText, LINE_TERMINATOR);
+		RPCWrappers.Logf(TEXT("\t\t%s%s%s;") LINE_TERMINATOR, *EvalBaseText, *EvalModifierText, *EvalParameterText);
 
 		// add this property to the parameter list string
 		if (ParameterList.Len())
@@ -3765,7 +3805,7 @@ void FNativeClassHeaderGenerator::ExportFunctionThunk(FStringOutputDevice& RPCWr
 		ParameterList += ParamName;
 	}
 
-	RPCWrappers.Logf(TEXT("%sP_FINISH;%s"), FCString::Spc(8), LINE_TERMINATOR);
+	RPCWrappers += TEXT("\t\tP_FINISH;") LINE_TERMINATOR;
 
 	const auto DeprecationPushString = TEXT("PRAGMA_ENABLE_DEPRECATION_WARNINGS") LINE_TERMINATOR;
 	const auto DeprecationPopString = TEXT("PRAGMA_POP") LINE_TERMINATOR;
@@ -3799,7 +3839,7 @@ void FNativeClassHeaderGenerator::ExportFunctionThunk(FStringOutputDevice& RPCWr
 		&& (FunctionData.FunctionFlags & FUNC_NetValidate) && !bHasValidate;
 
 	//Emit warning here if necessary
-	FStringOutputDevice FunctionDeclaration;
+	FUHTStringBuilder FunctionDeclaration;
 	ExportNativeFunctionHeader(FunctionData, FunctionDeclaration, EExportFunctionType::Function, EExportFunctionHeaderStyle::Declaration);
 	FunctionDeclaration.Trim();
 
@@ -3812,11 +3852,11 @@ void FNativeClassHeaderGenerator::ExportFunctionThunk(FStringOutputDevice& RPCWr
 		{
 			DeprecationWarningOutputDevice.Logf(TEXT("EMIT_CUSTOM_WARNING_AT_LINE(%d, \"Function %s::%s needs additional native validation function declaration %s due to its properties. Currently %s declaration is autogenerated by UHT. Since next release you'll have to provide declaration on your own. Please update your code before upgrading to the next release, otherwise your project will no longer compile.\")\r\n"), FunctionData.MacroLine, *ClassName, *FunctionName, *FunctionData.CppValidationImplName, *FunctionData.CppValidationImplName);
 		}
-		RPCWrappers.Logf(TEXT("%sif (!this->%s(%s))%s"), FCString::Spc(8), *FunctionData.CppValidationImplName, *ParameterList, LINE_TERMINATOR);
-		RPCWrappers.Logf(TEXT("%s{%s"), FCString::Spc(8), LINE_TERMINATOR);
-		RPCWrappers.Logf(TEXT("%sRPC_ValidateFailed(TEXT(\"%s\"));%s"), FCString::Spc(8 + 4), *FunctionData.CppValidationImplName, LINE_TERMINATOR);
-		RPCWrappers.Logf(TEXT("%sreturn;%s"), FCString::Spc(8 + 4), LINE_TERMINATOR);	// If we got here, the validation function check failed
-		RPCWrappers.Logf(TEXT("%s}%s"), FCString::Spc(8), LINE_TERMINATOR);
+		RPCWrappers.Logf(TEXT("\t\tif (!this->%s(%s))") LINE_TERMINATOR, *FunctionData.CppValidationImplName, *ParameterList);
+		RPCWrappers.Logf(TEXT("\t\t{") LINE_TERMINATOR);
+		RPCWrappers.Logf(TEXT("\t\t\tRPC_ValidateFailed(TEXT(\"%s\"));") LINE_TERMINATOR, *FunctionData.CppValidationImplName);
+		RPCWrappers.Logf(TEXT("\t\t\treturn;") LINE_TERMINATOR);	// If we got here, the validation function check failed
+		RPCWrappers.Logf(TEXT("\t\t}") LINE_TERMINATOR);
 	}
 
 	if (bShouldEnableImplementationDeprecation)
@@ -3824,13 +3864,14 @@ void FNativeClassHeaderGenerator::ExportFunctionThunk(FStringOutputDevice& RPCWr
 		DeprecationWarningOutputDevice.Logf(TEXT("EMIT_CUSTOM_WARNING_AT_LINE(%d, \"Function %s::%s needs additional native declaration %s due to its properties. Currently %s declaration is autogenerated by UHT. Since next release you'll have to provide declaration on your own. Please update your code before upgrading to the next release, otherwise your project will no longer compile.\")\r\n"), FunctionData.MacroLine, *ClassName, *FunctionName, *FunctionData.CppImplName, *FunctionData.CppImplName);
 	}
 	// write out the return value
-	RPCWrappers.Log(FCString::Spc(8));
+	RPCWrappers.Log(TEXT("\t\t"));
 	if (Return != NULL)
 	{
 		FString ReturnType, ReturnExtendedType;
 		ForwardDeclarations.Add(Return);
 		ReturnType = Return->GetCPPType(&ReturnExtendedType);
-		FStringOutputDevice ReplacementText(*ReturnType);
+		FUHTStringBuilder ReplacementText;
+		ReplacementText += ReturnType;
 		ApplyAlternatePropertyExportText(Return, ReplacementText);
 		ReturnType = ReplacementText;
 		RPCWrappers.Logf(TEXT("*(%s%s*)Result="), *ReturnType, *ReturnExtendedType);
@@ -3839,18 +3880,18 @@ void FNativeClassHeaderGenerator::ExportFunctionThunk(FStringOutputDevice& RPCWr
 	// export the call to the C++ version
 	if (FunctionData.FunctionExportFlags & FUNCEXPORT_CppStatic)
 	{
-		RPCWrappers.Logf(TEXT("%s::%s(%s);%s"), NameLookupCPP.GetNameCPP(Function->GetOwnerClass()), *FunctionData.CppImplName, *ParameterList, LINE_TERMINATOR);
+		RPCWrappers.Logf(TEXT("%s::%s(%s);") LINE_TERMINATOR, NameLookupCPP.GetNameCPP(Function->GetOwnerClass()), *FunctionData.CppImplName, *ParameterList);
 	}
 	else
 	{
-		RPCWrappers.Logf(TEXT("this->%s(%s);%s"), *FunctionData.CppImplName, *ParameterList, LINE_TERMINATOR);
+		RPCWrappers.Logf(TEXT("this->%s(%s);") LINE_TERMINATOR, *FunctionData.CppImplName, *ParameterList);
 	}
 }
 
 FString FNativeClassHeaderGenerator::GetFunctionParameterString(UFunction* Function)
 {
 	FString ParameterList;
-	FStringOutputDevice PropertyText;
+	FUHTStringBuilder PropertyText;
 
 	for (auto Property : TFieldRange<UProperty>(Function))
 	{
@@ -3871,7 +3912,7 @@ FString FNativeClassHeaderGenerator::GetFunctionParameterString(UFunction* Funct
 		ApplyAlternatePropertyExportText(Property, PropertyText);
 
 		ParameterList += PropertyText;
-		PropertyText.Empty();
+		PropertyText.Reset();
 	}
 
 	return ParameterList;
@@ -3879,11 +3920,11 @@ FString FNativeClassHeaderGenerator::GetFunctionParameterString(UFunction* Funct
 
 void FNativeClassHeaderGenerator::ExportNativeFunctions(FUnrealSourceFile& SourceFile, UClass* Class, FClassMetaData* ClassData)
 {
-	FStringOutputDevice RPCWrappers;
-	FStringOutputDevice	AutogeneratedBlueprintFunctionDeclarations;
-	FStringOutputDevice	AutogeneratedBlueprintFunctionDeclarationsOnlyNotDeclared;
-	FStringOutputDevice	MemberFunctionCheckClasses;
-	FStringOutputDevice StaticChecks;
+	FUHTStringBuilder RPCWrappers;
+	FUHTStringBuilder	AutogeneratedBlueprintFunctionDeclarations;
+	FUHTStringBuilder	AutogeneratedBlueprintFunctionDeclarationsOnlyNotDeclared;
+	FUHTStringBuilder	MemberFunctionCheckClasses;
+	FUHTStringBuilder StaticChecks;
 	StaticChecks.Logf(TEXT("\tstatic inline void StaticChecks_Implementation_Validate()\r\n"));
 	StaticChecks.Logf(TEXT("\t{\r\n"));
 
@@ -3914,7 +3955,7 @@ void FNativeClassHeaderGenerator::ExportNativeFunctions(FUnrealSourceFile& Sourc
 			continue;
 		}
 
-		FStringOutputDevice& DestinationForDecl = RPCWrappers;
+		FUHTStringBuilder& DestinationForDecl = RPCWrappers;
 
 		// Should we emit these to RPC wrappers or just ignore them?
 		const bool bWillBeProgrammerTyped = FunctionData.CppImplName == Function->GetName();
@@ -3931,7 +3972,7 @@ void FNativeClassHeaderGenerator::ExportNativeFunctions(FUnrealSourceFile& Sourc
 			bool bHasValidate = ClassDefinition.Contains(FunctionData.CppValidationImplName, ESearchCase::CaseSensitive);
 
 			//Emit warning here if necessary
-			FStringOutputDevice FunctionDeclaration;
+			FUHTStringBuilder FunctionDeclaration;
 			ExportNativeFunctionHeader(FunctionData, FunctionDeclaration, EExportFunctionType::Function, EExportFunctionHeaderStyle::Declaration);
 			FunctionDeclaration.Log(TEXT(";\r\n"));
 
@@ -3972,25 +4013,25 @@ void FNativeClassHeaderGenerator::ExportNativeFunctions(FUnrealSourceFile& Sourc
 		}
 
 		// export the script wrappers
-		RPCWrappers.Logf(TEXT("%sDECLARE_FUNCTION(%s)"), FCString::Spc(4), *FunctionData.UnMarshallAndCallName);
-		RPCWrappers.Logf(TEXT("%s%s{%s"), LINE_TERMINATOR, FCString::Spc(4), LINE_TERMINATOR);
+		RPCWrappers.Logf(TEXT("\tDECLARE_FUNCTION(%s)"), *FunctionData.UnMarshallAndCallName);
+		RPCWrappers += LINE_TERMINATOR TEXT("\t{") LINE_TERMINATOR;
 
 		auto Parameters = GetFunctionParmsAndReturn(FunctionData.FunctionReference);
 		ExportFunctionThunk(RPCWrappers, Function, FunctionData, Parameters.Parms, Parameters.Return, AutogeneratedBlueprintFunctionDeclarationsOnlyNotDeclared);
 
-		RPCWrappers.Logf(TEXT("%s}%s"), FCString::Spc(4), LINE_TERMINATOR);
+		RPCWrappers += TEXT("\t}") LINE_TERMINATOR;
 	}
 
 	StaticChecks.Logf(TEXT("\t}\r\n"));
 
 	FString MacroName = SourceFile.GetGeneratedMacroName(ClassData, TEXT("_RPC_WRAPPERS"));
 	WriteMacro(GeneratedHeaderText, MacroName, AutogeneratedBlueprintFunctionDeclarations + RPCWrappers);
-	InClassMacroCalls.Logf(TEXT("%s%s\r\n"), FCString::Spc(4), *MacroName);
+	InClassMacroCalls.Logf(TEXT("\t%s\r\n"), *MacroName);
 
 	// Put static checks before RPCWrappers to get proper messages from static asserts before compiler errors.
 	FString NoPureDeclsMacroName = SourceFile.GetGeneratedMacroName(ClassData, TEXT("_RPC_WRAPPERS_NO_PURE_DECLS"));
 	WriteMacro(GeneratedHeaderText, NoPureDeclsMacroName, AutogeneratedBlueprintFunctionDeclarationsOnlyNotDeclared + MemberFunctionCheckClasses + StaticChecks + RPCWrappers);
-	InClassNoPureDeclsMacroCalls.Logf(TEXT("%s%s\r\n"), FCString::Spc(4), *NoPureDeclsMacroName);
+	InClassNoPureDeclsMacroCalls.Logf(TEXT("\t%s\r\n"), *NoPureDeclsMacroName);
 }
 
 /**
@@ -4012,8 +4053,8 @@ TArray<UFunction*> FNativeClassHeaderGenerator::ExportCallbackFunctions(FUnrealS
 		}
 	}
 
-	FStringOutputDevice RPCWrappers;
-	FStringOutputDevice RPCWrappersCPP;
+	FUHTStringBuilder RPCWrappers;
+	FUHTStringBuilder RPCWrappersCPP;
 
 	if (CallbackFunctions.Num() == 0)
 	{
@@ -4023,14 +4064,14 @@ TArray<UFunction*> FNativeClassHeaderGenerator::ExportCallbackFunctions(FUnrealS
 
 	CallbackFunctions.Sort();
 
-	FStringOutputDevice UClassMacroContent;
+	FUHTStringBuilder UClassMacroContent;
 
 	// export parameters structs for all events and delegates
-	ExportEventParms(SourceFile.GetScope().Get(), CallbackFunctions, UClassMacroContent, 4, true);
+	ExportEventParms(SourceFile.GetScope().Get(), CallbackFunctions, UClassMacroContent, /*Indent=*/ 1, /*bOutputConstructor=*/ true);
 
 	FString MacroName = SourceFile.GetGeneratedMacroName(ClassData, TEXT("_EVENT_PARMS"));
 	WriteMacro(GeneratedHeaderText, MacroName, UClassMacroContent);
-	PrologMacroCalls.Logf(TEXT("%s%s\r\n"), FCString::Spc(4), *MacroName);
+	PrologMacroCalls.Logf(TEXT("\t%s\r\n"), *MacroName);
 
 	// export .proto files for any net service functions
 	ExportProtoMessage(CallbackFunctions, ClassData);
@@ -4084,36 +4125,35 @@ TArray<UFunction*> FNativeClassHeaderGenerator::ExportCallbackFunctions(FUnrealS
 
 		if (!(Class->ClassFlags & CLASS_Interface))
 		{
-			WriteEventFunctionPrologue(RPCWrappersCPP, 4, Parameters, Function->GetOuter(), *FunctionName);
+			WriteEventFunctionPrologue(RPCWrappersCPP, /*Indent=*/ 1, Parameters, Function->GetOuter(), *FunctionName);
 			{
 				// Cast away const just in case, because ProcessEvent isn't const
 				RPCWrappersCPP.Logf(
-					TEXT("%s%sProcessEvent(FindFunctionChecked(%s_%s),%s);\r\n"),
-					FCString::Spc(8),
+					TEXT("\t\t%sProcessEvent(FindFunctionChecked(%s_%s),%s);\r\n"),
 					(Function->HasAllFunctionFlags(FUNC_Const)) ? *FString::Printf(TEXT("const_cast<%s*>(this)->"), NameLookupCPP.GetNameCPP(Cast<UClass>(Function->GetOuter()))) : TEXT(""),
 					*API,
 					*FunctionName,
 					Parameters.HasParms() ? TEXT("&Parms") : TEXT("NULL")
 					);
 			}
-			WriteEventFunctionEpilogue(RPCWrappersCPP, 4, Parameters, *FunctionName);
+			WriteEventFunctionEpilogue(RPCWrappersCPP, /*Indent=*/ 1, Parameters, *FunctionName);
 		}
 		else
 		{
-			RPCWrappersCPP.Logf(LINE_TERMINATOR);
-			RPCWrappersCPP.Logf(TEXT("%s{%s"), FCString::Spc(4), LINE_TERMINATOR);
+			RPCWrappersCPP += LINE_TERMINATOR;
+			RPCWrappersCPP += TEXT("\t{") LINE_TERMINATOR;
 
 			// assert if this is ever called directly
-			RPCWrappersCPP.Logf(TEXT("%scheck(0 && \"Do not directly call Event functions in Interfaces. Call Execute_%s instead.\");%s"), FCString::Spc(8), *FunctionName, LINE_TERMINATOR);
+			RPCWrappersCPP.Logf(TEXT("\t\tcheck(0 && \"Do not directly call Event functions in Interfaces. Call Execute_%s instead.\");") LINE_TERMINATOR, *FunctionName);
 
 			// satisfy compiler if it's expecting a return value
 			if (Parameters.Return)
 			{
 				FString EventParmStructName = GetEventStructParamsName(Function->GetOuter(), *FunctionName);
-				RPCWrappersCPP.Logf(TEXT("%s%s Parms;%s"), FCString::Spc(8), *EventParmStructName, LINE_TERMINATOR);
-				RPCWrappersCPP.Logf(TEXT("%sreturn Parms.ReturnValue;%s"), FCString::Spc(8), LINE_TERMINATOR);
+				RPCWrappersCPP.Logf(TEXT("\t\t%s Parms;") LINE_TERMINATOR, *EventParmStructName);
+				RPCWrappersCPP.Logf(TEXT("\t\treturn Parms.ReturnValue;") LINE_TERMINATOR);
 			}
-			RPCWrappersCPP.Logf(TEXT("%s}%s"), FCString::Spc(4), LINE_TERMINATOR);
+			RPCWrappersCPP += TEXT("\t}") LINE_TERMINATOR;
 		}
 	}
 
@@ -4126,8 +4166,8 @@ TArray<UFunction*> FNativeClassHeaderGenerator::ExportCallbackFunctions(FUnrealS
 	FString CallbackWrappersMacroName = SourceFile.GetGeneratedMacroName(ClassData, TEXT("_CALLBACK_WRAPPERS"));
 	WriteMacro(GeneratedHeaderText, CallbackWrappersMacroName, RPCWrappers);
 
-	InClassMacroCalls.Logf(TEXT("%s%s\r\n"), FCString::Spc(4), *CallbackWrappersMacroName);
-	InClassNoPureDeclsMacroCalls.Logf(TEXT("%s%s\r\n"), FCString::Spc(4), *CallbackWrappersMacroName);
+	InClassMacroCalls.Logf(TEXT("\t%s\r\n"), *CallbackWrappersMacroName);
+	InClassNoPureDeclsMacroCalls.Logf(TEXT("\t%s\r\n"), *CallbackWrappersMacroName);
 
 	return CallbackFunctions;
 }
@@ -4141,7 +4181,7 @@ TArray<UFunction*> FNativeClassHeaderGenerator::ExportCallbackFunctions(FUnrealS
  * @param	Prop			the property that is being exported
  * @param	PropertyText	the string containing the text exported from ExportCppDeclaration
  */
-void FNativeClassHeaderGenerator::ApplyAlternatePropertyExportText( UProperty* Prop, FStringOutputDevice& PropertyText )
+void FNativeClassHeaderGenerator::ApplyAlternatePropertyExportText(UProperty* Prop, FUHTStringBuilder& PropertyText)
 {
 	if (!bIsExportingForOffsetDeterminationOnly)
 	{
@@ -4255,13 +4295,13 @@ bool IsExportClass(FClass* Class)
 	return Class->HasAnyClassFlags(CLASS_Native) && !Class->HasAnyClassFlags(CLASS_NoExport | CLASS_Intrinsic);
 }
 
-FStringOutputDevice& FNativeClassHeaderGenerator::GetGeneratedFunctionTextDevice()
+FUHTStringBuilder& FNativeClassHeaderGenerator::GetGeneratedFunctionTextDevice()
 {
 	int32 MaxLinesPerCpp = 30000;
 
 	if ((GeneratedFunctionBodyTextSplit.Num() == 0) || (GeneratedFunctionBodyTextSplit[GeneratedFunctionBodyTextSplit.Num() - 1].GetLineCount() > MaxLinesPerCpp))
 	{
-		new(GeneratedFunctionBodyTextSplit)FStringOutputDeviceCountLines;
+		new (GeneratedFunctionBodyTextSplit) FUHTStringBuilderLineCounter;
 	}
 
 	return GeneratedFunctionBodyTextSplit[GeneratedFunctionBodyTextSplit.Num() - 1];
@@ -4295,18 +4335,18 @@ FNativeClassHeaderGenerator::FNativeClassHeaderGenerator(
 
 	bool bPackageHasAnyExportClasses = AllClasses.GetClassesInPackage(Package).ContainsByPredicate(IsExportClass);
 	bool bHasNamesForExport = false;
-	TempHeaderPaths.Empty();
-	PackageHeaderPaths.Empty();
+	TempHeaderPaths.Reset();
+	PackageHeaderPaths.Reset();
 
 	// Reset header generation output strings
-	GeneratedPackageCPP                           .Empty();
-	GeneratedProtoText                            .Empty();
-	GeneratedMCPText                              .Empty();
-	CrossModuleGeneratedFunctionDeclarations      .Empty();
-	UniqueCrossModuleReferences                   .Empty();
-	GeneratedFunctionDeclarations                 .Empty();
-	GeneratedFunctionBodyTextSplit                .Empty();
-	ListOfPublicClassesUObjectHeaderModuleIncludes.Empty();
+	GeneratedPackageCPP.Reset();
+	GeneratedProtoText.Reset();
+	GeneratedMCPText.Reset();
+	CrossModuleGeneratedFunctionDeclarations.Reset();
+	UniqueCrossModuleReferences.Empty(UniqueCrossModuleReferences.Num());
+	GeneratedFunctionDeclarations.Reset();
+	GeneratedFunctionBodyTextSplit.Reset();
+	ListOfPublicClassesUObjectHeaderModuleIncludes.Reset();
 
 	if (bPackageHasAnyExportClasses)
 	{
@@ -4344,10 +4384,10 @@ FNativeClassHeaderGenerator::FNativeClassHeaderGenerator(
 		{
 			ClassesHeaders.Add(ClassesHeaderName);
 
-			ListOfPublicHeaderGroupIncludes.Empty();
-			ListOfAllUObjectHeaderIncludes.Empty();
-			OriginalHeader.Empty();
-			PreHeaderText.Empty();
+			ListOfPublicHeaderGroupIncludes.Reset();
+			ListOfAllUObjectHeaderIncludes.Reset();
+			OriginalHeader.Reset();
+			PreHeaderText.Reset();
 
 			// Load the original header file into memory
 			FFileHelper::LoadFileToString(OriginalHeader, *ClassesHeaderPath);
@@ -4576,7 +4616,7 @@ bool FNativeClassHeaderGenerator::SaveHeaderIfChanged(const TCHAR* HeaderPath, c
 	}
 
 	// Remember this header filename to be able to check for any old (unused) headers later.
-	PackageHeaderPaths.Add( FString(HeaderPath).Replace( TEXT( "\\" ), TEXT( "/" ) ) );
+	PackageHeaderPaths.Add(FString(HeaderPath).Replace(TEXT("\\"), TEXT("/"), ESearchCase::CaseSensitive));
 
 	return bHasChanged;
 }
@@ -4629,7 +4669,7 @@ void FNativeClassHeaderGenerator::ExportGeneratedProto()
 	{
 		UE_LOG(LogCompile, Log,  TEXT("Autogenerating boilerplate proto: %s.proto"), *GeneratedProtoFilenameBase );
 
-		FStringOutputDevice ProtoPreamble;
+		FUHTStringBuilder ProtoPreamble;
 
 		ProtoPreamble.Logf(
 			TEXT("// Copyright 1998-2015 Epic Games, Inc. All Rights Reserved.")					LINE_TERMINATOR
@@ -4670,7 +4710,7 @@ void FNativeClassHeaderGenerator::ExportGeneratedMCP()
 	{
 		UE_LOG(LogCompile, Log,  TEXT("Autogenerating boilerplate MCP: %s.java"), *GeneratedMCPFilenameBase );
 
-		FStringOutputDevice MCPPreamble;
+		FUHTStringBuilder MCPPreamble;
 
 		MCPPreamble.Logf(
 			TEXT("// Copyright 1998-2015 Epic Games, Inc. All Rights Reserved.")					LINE_TERMINATOR
@@ -4700,13 +4740,13 @@ void FNativeClassHeaderGenerator::ExportGeneratedMCP()
  */
 void FNativeClassHeaderGenerator::ExportGeneratedCPP()
 {
-	FStringOutputDevice GeneratedCPPPreamble;
-	FStringOutputDevice GeneratedCPPClassesIncludes;
-	FStringOutputDevice GeneratedCPPEpilogue;
-	FStringOutputDevice GeneratedCPPText;
-	FStringOutputDevice GeneratedLinkerFixupFunction;
+	FUHTStringBuilder GeneratedCPPPreamble;
+	FUHTStringBuilder GeneratedCPPClassesIncludes;
+	FUHTStringBuilder GeneratedCPPEpilogue;
+	FUHTStringBuilder GeneratedCPPText;
+	FUHTStringBuilder GeneratedLinkerFixupFunction;
 
-	TArray<FStringOutputDevice> GeneratedCPPFiles;
+	TArray<FUHTStringBuilder> GeneratedCPPFiles;
 
 	UE_LOG(LogCompile, Log,  TEXT("Autogenerating boilerplate cpp: %s.cpp"), *GeneratedCPPFilenameBase );
 
@@ -4773,7 +4813,7 @@ void FNativeClassHeaderGenerator::ExportGeneratedCPP()
 	// Generate each of the .generated.cpp files
 	for( int32 FileIdx=0;FileIdx<GeneratedFunctionBodyTextSplit.Num();FileIdx++ )
 	{
-		FStringOutputDevice FileText;
+		FUHTStringBuilder FileText;
 		// The first file has all of the GeneratedCPPText, only the functions are split.
 		if( FileIdx == 0)
 		{
@@ -4811,16 +4851,8 @@ void FNativeClassHeaderGenerator::ExportGeneratedCPP()
  		FString CppPath = ModuleInfo->GeneratedCPPFilenameBase + TEXT(".cpp");
 		IFileManager::Get().Delete(*CppPath);
  	}
-
-	if( FParse::Param( FCommandLine::Get(), TEXT("GenerateConstructors") ) && AllConstructors.Len())
-	{
-		const FString PrivateLoc = PkgDir / PkgName / TEXT("Private");
-		FString ConstructorPath = PrivateLoc / PkgName + TEXT("Constructors.cpp");
-		FString ConstructorsWithIfdef(TEXT("// Copyright 1998-2015 Epic Games, Inc. All Rights Reserved.\r\n"));
-		ConstructorsWithIfdef += *Tabify(*AllConstructors);
-		SaveHeaderIfChanged(*ConstructorPath,*ConstructorsWithIfdef);
-	}
 }
+
 /** Get all script plugins based on ini setting */
 void GetScriptPlugins(TArray<IScriptGeneratorPluginInterface*>& ScriptPlugins)
 {
@@ -5329,7 +5361,7 @@ TSharedRef<FUnrealSourceFile> PerformInitialParseOnHeader(UPackage* InParent, co
 	TArray<FHeaderProvider> DependsOn;
 
 	// Parse the header to extract the information needed
-	FStringOutputDevice ClassHeaderTextStrippedOfCppText;
+	FUHTStringBuilder ClassHeaderTextStrippedOfCppText;
 	TArray<FSimplifiedParsingClassInfo> ParsedClassArray;
 	FHeaderParser::SimplifiedClassParse(Buffer, /*out*/ ParsedClassArray, /*out*/ DependsOn, ClassHeaderTextStrippedOfCppText);
 
@@ -5352,4 +5384,3 @@ TSharedRef<FUnrealSourceFile> PerformInitialParseOnHeader(UPackage* InParent, co
 
 	return UnrealSourceFile;
 }
-
diff --git a/Engine/Source/Programs/UnrealHeaderTool/Private/HeaderParser.cpp b/Engine/Source/Programs/UnrealHeaderTool/Private/HeaderParser.cpp
--- a/Engine/Source/Programs/UnrealHeaderTool/Private/HeaderParser.cpp
+++ b/Engine/Source/Programs/UnrealHeaderTool/Private/HeaderParser.cpp
@@ -6580,7 +6580,7 @@ void FHeaderParser::ExportNativeHeaders(
 {
 	// Build a list of header filenames
 	TArray<FString>	ClassHeaderFilenames;
-	new(ClassHeaderFilenames) FString(TEXT(""));
+	new (ClassHeaderFilenames) FString();
 
 	auto SourceFiles = GetSourceFilesWithInheritanceOrdering(CurrentPackage, AllClasses);
 
diff --git a/Engine/Source/Programs/UnrealHeaderTool/Private/NativeClassExporter.h b/Engine/Source/Programs/UnrealHeaderTool/Private/NativeClassExporter.h
--- a/Engine/Source/Programs/UnrealHeaderTool/Private/NativeClassExporter.h
+++ b/Engine/Source/Programs/UnrealHeaderTool/Private/NativeClassExporter.h
@@ -32,10 +32,14 @@ class FClass;
 class FClasses;
 class FModuleClasses;
 
+// These are declared in this way to allow swapping out the classes for something more optimized in the future
+typedef FStringOutputDevice FUHTStringBuilder;
+typedef FStringOutputDeviceCountLines FUHTStringBuilderLineCounter;
+
 struct FNativeClassHeaderGenerator
 {
 private:
-	FString				API;
+	FString API;
 
 	/**
 	 * Gets API string for this header.
@@ -45,41 +49,40 @@ private:
 		return FString::Printf(TEXT("%s_API "), *API);
 	}
 
-	FString			ClassesHeaderPath;
+	FString ClassesHeaderPath;
 	/** the name of the cpp file where the implementation functions are generated, without the .cpp extension */
-	FString				GeneratedCPPFilenameBase;
+	FString GeneratedCPPFilenameBase;
 	/** the name of the proto file where proto definitions are generated */
-	FString				GeneratedProtoFilenameBase;
+	FString GeneratedProtoFilenameBase;
 	/** the name of the java file where mcp definitions are generated */
-	FString				GeneratedMCPFilenameBase;
-	UPackage*			Package;
-	FStringOutputDevice	PreHeaderText;
-	FStringOutputDevice	EnumForeachText;
+	FString GeneratedMCPFilenameBase;
+	UPackage* Package;
+	FUHTStringBuilder PreHeaderText;
+	FUHTStringBuilder EnumForeachText;
 	TArray<FUnrealSourceFile*> ListOfPublicHeaderGroupIncludes;			// This is built up each time from scratch each time for each header groups
-	FStringOutputDevice ListOfPublicClassesUObjectHeaderModuleIncludes; // This is built up for the entire module, and for all processed headers
-	FStringOutputDevice ListOfAllUObjectHeaderIncludes;
-	FStringOutputDevice	GeneratedPackageCPP;
-	FStringOutputDevice	GeneratedHeaderText;
-	FStringOutputDevice	GeneratedHeaderTextBeforeForwardDeclarations;
-	FStringOutputDevice	GeneratedForwardDeclarations;
+	FUHTStringBuilder ListOfPublicClassesUObjectHeaderModuleIncludes; // This is built up for the entire module, and for all processed headers
+	FUHTStringBuilder ListOfAllUObjectHeaderIncludes;
+	FUHTStringBuilder GeneratedPackageCPP;
+	FUHTStringBuilder GeneratedHeaderText;
+	FUHTStringBuilder GeneratedHeaderTextBeforeForwardDeclarations;
+	FUHTStringBuilder GeneratedForwardDeclarations;
     /** output generated for a .proto file */
-	FStringOutputDevice GeneratedProtoText;
+	FUHTStringBuilder GeneratedProtoText;
 	/** output generated for a mcp .java file */
-	FStringOutputDevice GeneratedMCPText;
-	FStringOutputDevice	PrologMacroCalls;
-	FStringOutputDevice	InClassMacroCalls;
-	FStringOutputDevice	InClassNoPureDeclsMacroCalls;
-	FStringOutputDevice	StandardUObjectConstructorsMacroCall;
-	FStringOutputDevice	EnhancedUObjectConstructorsMacroCall;
-	FStringOutputDevice	AllConstructors;
-	FStringOutputDevice StaticChecks;
+	FUHTStringBuilder GeneratedMCPText;
+	FUHTStringBuilder PrologMacroCalls;
+	FUHTStringBuilder InClassMacroCalls;
+	FUHTStringBuilder InClassNoPureDeclsMacroCalls;
+	FUHTStringBuilder StandardUObjectConstructorsMacroCall;
+	FUHTStringBuilder EnhancedUObjectConstructorsMacroCall;
+	FUHTStringBuilder StaticChecks;
 
 	/** Generated function implementations that belong in the cpp file, split into multiple files base on line count **/
-	TArray<FStringOutputDeviceCountLines> GeneratedFunctionBodyTextSplit;
+	TArray<FUHTStringBuilderLineCounter> GeneratedFunctionBodyTextSplit;
 	/** Declarations of generated functions for this module**/
-	FStringOutputDevice	GeneratedFunctionDeclarations;
+	FUHTStringBuilder GeneratedFunctionDeclarations;
 	/** Declarations of generated functions for cross module**/
-	FStringOutputDevice	CrossModuleGeneratedFunctionDeclarations;
+	FUHTStringBuilder CrossModuleGeneratedFunctionDeclarations;
 	/** Set of already exported cross-module references, to prevent duplicates */
 	TSet<FString> UniqueCrossModuleReferences;
 
@@ -87,18 +90,18 @@ private:
 	TSet<FName> ReferencedNames;
 
 	/** the existing disk version of the header for this package's names */
-	FString				OriginalNamesHeader;
+	FString OriginalNamesHeader;
 	/** the existing disk version of this header */
-	FString				OriginalHeader;
+	FString OriginalHeader;
 
 	/** Array of temp filenames that for files to overwrite headers */
-	TArray<FString>		TempHeaderPaths;
+	TArray<FString> TempHeaderPaths;
 
 	/** Array of all header filenames from the current package. */
-	TArray<FString>		PackageHeaderPaths;
+	TArray<FString> PackageHeaderPaths;
 
 	/** Array of all generated Classes headers for the current package */
-	TArray<FString>		ClassesHeaders;
+	TArray<FString> ClassesHeaders;
 	
 	/** if we are exporting a struct for offset determination only, replace some types with simpler ones (delegates) */
 	bool bIsExportingForOffsetDeterminationOnly;
@@ -125,7 +128,7 @@ private:
 	/**
 	 * Gets generated function text device.
 	 */
-	FStringOutputDevice& GetGeneratedFunctionTextDevice();
+	FUHTStringBuilder& GetGeneratedFunctionTextDevice();
 
 	/**
 	 * Sorts the list of header files being exported from a package according to their dependency on each other.
@@ -177,7 +180,7 @@ private:
 	 * @param	TextIndent			Current text indentation
 	 * @param	Output	alternate output device
 	 */
-	void ExportProperties( UStruct* Struct, int32 TextIndent, bool bAccessSpecifiers = true, class FStringOutputDevice* Output = NULL );
+	void ExportProperties(UStruct* Struct, int32 TextIndent, bool bAccessSpecifiers = true, FUHTStringBuilder* Output = NULL);
 
 	/** Return the name of the singleton function that returns the UObject for Item */
 	FString GetSingletonName(UField* Item, bool bRequiresValidObject=true);
@@ -187,7 +190,7 @@ private:
 	/** 
 	 * Export functions used to find and call C++ or script implementation of a script function in the interface 
 	 */
-	void ExportInterfaceCallFunctions(const TArray<UFunction*>& CallbackFunctions, UClass* Class, FClassMetaData* ClassData, FStringOutputDevice& HeaderOutput);
+	void ExportInterfaceCallFunctions(const TArray<UFunction*>& CallbackFunctions, UClass* Class, FClassMetaData* ClassData, FUHTStringBuilder& HeaderOutput);
 
 	/** 
 	 * Export UInterface boilerplate.
@@ -196,7 +199,7 @@ private:
 	 * @param Class Interface to export.
 	 * @param FriendText Friend text for this boilerplate.
 	 */
-	void ExportUInterfaceBoilerplate(FStringOutputDevice &UInterfaceBoilerplate, FClass* Class, const FString& FriendText);
+	void ExportUInterfaceBoilerplate(FUHTStringBuilder& UInterfaceBoilerplate, FClass* Class, const FString& FriendText);
 
 	/**
 	 * Appends the header definition for an inheritance hierarchy of classes to the header.
@@ -266,7 +269,7 @@ private:
 	 * @param	TextIndent	the current indentation of the header exporter
 	 * @param	HeaderOutput	the output device for the mirror struct
 	 */
-	void ExportMirrorsForNoexportStructs( const TArray<UScriptStruct*>& NativeStructs, int32 TextIndent, FStringOutputDevice& HeaderOutput);
+	void ExportMirrorsForNoexportStructs(const TArray<UScriptStruct*>& NativeStructs, int32 TextIndent, FUHTStringBuilder& HeaderOutput);
 
 	/**
 	 * Exports the parameter struct declarations for the list of functions specified
@@ -294,7 +297,7 @@ private:
 	 * @param	Indent				number of spaces to put before each line
 	 * @param	bOutputConstructor	If true, output a constructor for the parm struct
 	 */
-	void ExportEventParms(FScope& Scope, const TArray<UFunction*>& CallbackFunctions, FStringOutputDevice& Output, int32 Indent = 0, bool bOutputConstructor = true);
+	void ExportEventParms(FScope& Scope, const TArray<UFunction*>& CallbackFunctions, FUHTStringBuilder& Output, int32 Indent = 0, bool bOutputConstructor = true);
 
 	/**
 	 * Exports the parameter struct declarations for the given function.
@@ -304,7 +307,7 @@ private:
 	 * @param	Indent				number of spaces to put before each line
 	 * @param	bOutputConstructor	If true, output a constructor for the param struct
 	 */
-	void ExportEventParm(UFunction* Function, FStringOutputDevice &HeaderOutput, int32 Indent = 0, bool bOutputConstructor = true);
+	void ExportEventParm(UFunction* Function, FUHTStringBuilder& HeaderOutput, int32 Indent = 0, bool bOutputConstructor = true);
 
 	/**
 	 * Generate a .proto message declaration for any functions marked as requiring one
@@ -314,7 +317,7 @@ private:
 	 * @param Indent starting indentation level
 	 * @param Output optional output redirect
 	 */
-	void ExportProtoMessage(const TArray<UFunction*>& InCallbackFunctions, FClassMetaData* ClassData, int32 Indent = 0, class FStringOutputDevice* Output = NULL);
+	void ExportProtoMessage(const TArray<UFunction*>& InCallbackFunctions, FClassMetaData* ClassData, int32 Indent = 0, FUHTStringBuilder* Output = NULL);
 
 	/**
 	 * Generate a .java message declaration for any functions marked as requiring one
@@ -324,7 +327,7 @@ private:
 	 * @param Indent starting indentation level
 	 * @param Output optional output redirect
 	 */
-	void ExportMCPMessage(const TArray<UFunction*>& InCallbackFunctions, FClassMetaData* ClassData, int32 Indent = 0, class FStringOutputDevice* Output = NULL);
+	void ExportMCPMessage(const TArray<UFunction*>& InCallbackFunctions, FClassMetaData* ClassData, int32 Indent = 0, FUHTStringBuilder* Output = NULL);
 
 	/** 
 	* Exports the temp header files into the .h files, then deletes the temp files.
@@ -370,7 +373,7 @@ private:
 	 * @param	FunctionHeaderStyle	Whether we're outputting a declaration or definition.
 	 * @param	ExtraParam			Optional extra parameter that will be added to the declaration as the first argument
 	 */
-	void ExportNativeFunctionHeader( const FFuncInfo& FunctionData, FStringOutputDevice& HeaderOutput, EExportFunctionType::Type FunctionType, EExportFunctionHeaderStyle::Type FunctionHeaderStyle, const TCHAR* ExtraParam = NULL );
+	void ExportNativeFunctionHeader(const FFuncInfo& FunctionData, FUHTStringBuilder& HeaderOutput, EExportFunctionType::Type FunctionType, EExportFunctionHeaderStyle::Type FunctionHeaderStyle, const TCHAR* ExtraParam = NULL);
 
 	/**
 	* Exports checks if function exists
@@ -380,7 +383,7 @@ private:
 	* @param	StaticChecks			Where to write the static asserts throwing errors when function doesn't exist.
 	* @param	ClassName				Name of currently parsed class.
 	*/
-	void ExportFunctionChecks(const FFuncInfo& FunctionData, FStringOutputDevice& CheckClasses, FStringOutputDevice& StaticChecks, const FString& ClassName);
+	void ExportFunctionChecks(const FFuncInfo& FunctionData, FUHTStringBuilder& CheckClasses, FUHTStringBuilder& StaticChecks, const FString& ClassName);
 
 	/**
 	 * Exports the native stubs for the list of functions specified
@@ -401,7 +404,7 @@ private:
 	 * @param Return return parameter for the function
 	 * @param DeprecationWarningOutputDevice Device to output deprecation warnings for _Validate and _Implementation functions.
 	 */
-	void ExportFunctionThunk(FStringOutputDevice& RPCWrappers, UFunction* Function, const FFuncInfo& FunctionData, const TArray<UProperty*>& Parameters, UProperty* Return, FStringOutputDevice& DeprecationWarningOutputDevice);
+	void ExportFunctionThunk(FUHTStringBuilder& RPCWrappers, UFunction* Function, const FFuncInfo& FunctionData, const TArray<UProperty*>& Parameters, UProperty* Return, FUHTStringBuilder& DeprecationWarningOutputDevice);
 
 	/** Exports the native function registration code for the given class. */
 	void ExportNatives(FClass* Class);
@@ -412,7 +415,7 @@ private:
 	 * @param	Class			Class to export
 	 * @param	OutFriendText	(Output parameter) Friend text
 	 */
-	void ExportNativeGeneratedInitCode(FClass* Class, FStringOutputDevice& OutFriendText);
+	void ExportNativeGeneratedInitCode(FClass* Class, FUHTStringBuilder& OutFriendText);
 
 	/**
 	 * Export given function.
@@ -480,7 +483,7 @@ private:
 	 * @param	Prop			the property that is being exported
 	 * @param	PropertyText	the string containing the text exported from ExportCppDeclaration
 	 */
-	void ApplyAlternatePropertyExportText( UProperty* Prop, FStringOutputDevice& PropertyText );
+	void ApplyAlternatePropertyExportText(UProperty* Prop, FUHTStringBuilder& PropertyText);
 
 	/**
 	* Create a temp header file name from the header name
diff --git a/Engine/Source/Programs/UnrealHeaderTool/Private/StringUtils.cpp b/Engine/Source/Programs/UnrealHeaderTool/Private/StringUtils.cpp
--- a/Engine/Source/Programs/UnrealHeaderTool/Private/StringUtils.cpp
+++ b/Engine/Source/Programs/UnrealHeaderTool/Private/StringUtils.cpp
@@ -16,7 +16,7 @@ FString GetClassNameWithPrefixRemoved(const FString InClassName)
 FString GetClassNameWithoutPrefix( const FString& InClassNameOrFilename )
 {
 	// Check for header names (they don't come with a full path so we only search for the first dot)
-	const int32 DotIndex = InClassNameOrFilename.Find(TEXT("."));
+	const int32 DotIndex = InClassNameOrFilename.Find(TEXT("."), ESearchCase::CaseSensitive);
 	if (DotIndex == INDEX_NONE)
 	{
 		const FString ClassPrefix = GetClassPrefix( InClassNameOrFilename );
@@ -39,23 +39,34 @@ FString GetClassPrefix( const FString InClassName, bool& bIsLabeledDeprecated )
 	FString ClassPrefix = InClassName.Left(1);
 
 	bIsLabeledDeprecated = false;
-	// First make sure the class name starts with a valid prefix
-	if( ClassPrefix == TEXT("I") ||
-		ClassPrefix == TEXT("A") || 
-		ClassPrefix == TEXT("U") )
+
+	if (!ClassPrefix.IsEmpty())
 	{
-		// If it is a class prefix, check for deprecated class prefix
-		if( InClassName.Mid(1, 11) == TEXT("DEPRECATED_") )
+		const TCHAR ClassPrefixChar = ClassPrefix[0];
+		switch (ClassPrefixChar)
 		{
-			bIsLabeledDeprecated = true;
-			return InClassName.Left(12);
+		case TEXT('I'):
+		case TEXT('A'):
+		case TEXT('U'):
+			// If it is a class prefix, check for deprecated class prefix also
+			if (InClassName.Mid(1, 11).Compare(TEXT("DEPRECATED_"), ESearchCase::CaseSensitive) == 0)
+			{
+				bIsLabeledDeprecated = true;
+				ClassPrefix = InClassName.Left(12);
+			}
+			break;
+
+		case TEXT('F'):
+		case TEXT('T'):
+			// Struct prefixes are also fine.
+			break;
+
+		default:
+			// If it's not a class or struct prefix, it's invalid
+			ClassPrefix.Empty();
+			break;
 		}
 	}
-	// Otherwise check for struct prefix. If it's not a class or struct prefix, it's invalid
-	else if( ClassPrefix != TEXT("F") && ClassPrefix != TEXT("T") )
-	{
-		ClassPrefix = TEXT("");
-	}
 	return ClassPrefix;
 }
 
