UHT can skip inline function declarations.

diff --git a/Engine/Source/Programs/UnrealHeaderTool/Private/HeaderParser.cpp b/Engine/Source/Programs/UnrealHeaderTool/Private/HeaderParser.cpp
--- a/Engine/Source/Programs/UnrealHeaderTool/Private/HeaderParser.cpp
+++ b/Engine/Source/Programs/UnrealHeaderTool/Private/HeaderParser.cpp
@@ -5511,7 +5511,14 @@ void FHeaderParser::CompileFunctionDeclaration(FUnrealSourceFile& SourceFile, FC
 	AddFormattedPrevCommentAsTooltipMetaData(MetaData);
 
 	AddMetaDataToClassData(TopFunction, MetaData);
-	
+
+	// 'final' and 'override' can appear in any order before an optional '= 0' pure virtual specifier
+	bool bFoundFinal    = MatchIdentifier(TEXT("final"));
+	bool bFoundOverride = MatchIdentifier(TEXT("override"));
+	if (!bFoundFinal && bFoundOverride)
+	{
+		bFoundFinal = MatchIdentifier(TEXT("final"));
+	}
 
 	// Handle C++ style functions being declared as abstract
 	if (MatchSymbol(TEXT("=")))
@@ -5527,7 +5534,7 @@ void FHeaderParser::CompileFunctionDeclaration(FUnrealSourceFile& SourceFile, FC
 	}
 
 	// Look for the final keyword to indicate this function is sealed
-	if (MatchIdentifier(TEXT("final")))
+	if (bFoundFinal)
 	{
 		// This is a final (prebinding, non-overridable) function
 		FuncInfo.FunctionFlags |= FUNC_Final;
@@ -5694,9 +5701,30 @@ void FHeaderParser::CompileFunctionDeclaration(FUnrealSourceFile& SourceFile, FC
 	// Just declaring a function, so end the nesting.
 	PostPopFunctionDeclaration(AllClasses, TopFunction);
 
+	// See what's coming next
+	FToken Token;
+	if (!GetToken(Token))
+	{
+		FError::Throwf(TEXT("Unexpected end of file"));
+	}
+
 	// Optionally consume a semicolon
-	// This is optional to allow inline function definitions, as long as the function body is inside #if CPP ... #endif
-	MatchSymbol(TEXT(";"));
+	// This is optional to allow inline function definitions
+	if (Token.TokenType == TOKEN_Symbol && !FCString::Stricmp(Token.Identifier, TEXT(";")))
+	{
+		// Do nothing (consume it)
+	}
+	else if (Token.TokenType == TOKEN_Symbol && !FCString::Stricmp(Token.Identifier, TEXT("{")))
+	{
+		// Skip inline function bodies
+		UngetToken(Token);
+		SkipDeclaration(Token);
+	}
+	else
+	{
+		// Put the token back so we can continue parsing as normal
+		UngetToken(Token);
+	}
 }
 
 /**
diff --git a/Engine/Source/Programs/UnrealHeaderTool/Resources/UHTDebugging/TestObject.h b/Engine/Source/Programs/UnrealHeaderTool/Resources/UHTDebugging/TestObject.h
--- a/Engine/Source/Programs/UnrealHeaderTool/Resources/UHTDebugging/TestObject.h
+++ b/Engine/Source/Programs/UnrealHeaderTool/Resources/UHTDebugging/TestObject.h
@@ -24,6 +24,20 @@ public:
 	UFUNCTION()
 	void CodeGenTestForEnumClasses(ECppEnum Val);
 
+	UFUNCTION()
+	int32 InlineFunc()
+	{
+		return FString("Hello").Len();
+	}
+
+	UFUNCTION()
+	int32 InlineFuncWithCppMacros()
+#if CPP
+	{
+		return FString("Hello").Len();
+	}
+#endif
+
 #if 0
 	UPROPERTY()
 	ITestInterface* ShouldGiveAMeaningfulErrorAboutTScriptInterface;
