Property creation moved to its own function.
DoesAnythingInHierarchyHaveDefaultToInstanced moved out of FHeaderParser.

diff --git a/Engine/Source/Programs/UnrealHeaderTool/Private/HeaderParser.cpp b/Engine/Source/Programs/UnrealHeaderTool/Private/HeaderParser.cpp
--- a/Engine/Source/Programs/UnrealHeaderTool/Private/HeaderParser.cpp
+++ b/Engine/Source/Programs/UnrealHeaderTool/Private/HeaderParser.cpp
@@ -412,6 +412,208 @@ namespace
 		}
 
 		// Unreachable
+		check(false);
+		return nullptr;
+	}
+
+	// Check to see if anything in the class hierarchy passed in has CLASS_DefaultToInstanced
+	bool DoesAnythingInHierarchyHaveDefaultToInstanced(UClass* TestClass)
+	{
+		bool bDefaultToInstanced = false;
+
+		UClass* Search = TestClass;
+		while (!bDefaultToInstanced && (Search != NULL))
+		{
+			bDefaultToInstanced = Search->HasAnyClassFlags(CLASS_DefaultToInstanced);
+			if (!bDefaultToInstanced && !Search->HasAnyClassFlags(CLASS_Intrinsic | CLASS_Parsed))
+			{
+				// The class might not have been parsed yet, look for declaration data.
+				auto ClassDeclarationDataPtr = GClassDeclarations.Find(Search->GetFName());
+				if (ClassDeclarationDataPtr)
+				{
+					bDefaultToInstanced = !!((*ClassDeclarationDataPtr)->ClassFlags & CLASS_DefaultToInstanced);
+				}
+			}
+			Search = Search->GetSuperClass();
+		}
+
+		return bDefaultToInstanced;
+	}
+
+	UProperty* CreateVariableProperty(FPropertyBase& VarProperty, UObject* Scope, FName Name, EObjectFlags ObjectFlags, EVariableCategory::Type VariableCategory)
+	{
+		switch (VarProperty.Type)
+		{
+			case CPT_Byte:
+			{
+				auto* Result = new (EC_InternalUseOnlyConstructor, Scope, Name, ObjectFlags) UByteProperty(FObjectInitializer());
+				Result->Enum = VarProperty.Enum;
+				return Result;
+			}
+
+			case CPT_Int8:
+				return new (EC_InternalUseOnlyConstructor, Scope, Name, ObjectFlags) UInt8Property(FObjectInitializer());
+
+			case CPT_Int16:
+				return new (EC_InternalUseOnlyConstructor, Scope, Name, ObjectFlags) UInt16Property(FObjectInitializer());
+
+			case CPT_Int:
+				return new (EC_InternalUseOnlyConstructor, Scope, Name, ObjectFlags) UIntProperty(FObjectInitializer());
+
+			case CPT_Int64:
+				return new (EC_InternalUseOnlyConstructor, Scope, Name, ObjectFlags) UInt64Property(FObjectInitializer());
+
+			case CPT_UInt16:
+				return new (EC_InternalUseOnlyConstructor, Scope, Name, ObjectFlags) UUInt16Property(FObjectInitializer());
+
+			case CPT_UInt32:
+				return new (EC_InternalUseOnlyConstructor, Scope, Name, ObjectFlags) UUInt32Property(FObjectInitializer());
+
+			case CPT_UInt64:
+				return new (EC_InternalUseOnlyConstructor, Scope, Name, ObjectFlags) UUInt64Property(FObjectInitializer());
+
+			case CPT_Bool:
+			{
+				auto* Result = new (EC_InternalUseOnlyConstructor, Scope, Name, ObjectFlags) UBoolProperty(FObjectInitializer());
+				Result->SetBoolSize(sizeof(bool), true);
+				return Result;
+			}
+
+			case CPT_Bool8:
+			{
+				auto* Result = new (EC_InternalUseOnlyConstructor, Scope, Name, ObjectFlags) UBoolProperty(FObjectInitializer());
+				Result->SetBoolSize((VariableCategory == EVariableCategory::Return) ? sizeof(bool) : sizeof(uint8), VariableCategory == EVariableCategory::Return);
+				return Result;
+			}
+
+			case CPT_Bool16:
+			{
+				auto* Result = new (EC_InternalUseOnlyConstructor, Scope, Name, ObjectFlags) UBoolProperty(FObjectInitializer());
+				Result->SetBoolSize((VariableCategory == EVariableCategory::Return) ? sizeof(bool) : sizeof(uint16), VariableCategory == EVariableCategory::Return);
+				return Result;
+			}
+
+			case CPT_Bool32:
+			{
+				auto* Result = new (EC_InternalUseOnlyConstructor, Scope, Name, ObjectFlags) UBoolProperty(FObjectInitializer());
+				Result->SetBoolSize((VariableCategory == EVariableCategory::Return) ? sizeof(bool) : sizeof(uint32), VariableCategory == EVariableCategory::Return);
+				return Result;
+			}
+
+			case CPT_Bool64:
+			{
+				auto* Result = new (EC_InternalUseOnlyConstructor, Scope, Name, ObjectFlags) UBoolProperty(FObjectInitializer());
+				Result->SetBoolSize((VariableCategory == EVariableCategory::Return) ? sizeof(bool) : sizeof(uint64), VariableCategory == EVariableCategory::Return);
+				return Result;
+			}
+
+			case CPT_Float:
+				return new (EC_InternalUseOnlyConstructor, Scope, Name, ObjectFlags) UFloatProperty(FObjectInitializer());
+
+			case CPT_Double:
+				return new (EC_InternalUseOnlyConstructor, Scope, Name, ObjectFlags) UDoubleProperty(FObjectInitializer());
+
+			case CPT_ObjectReference:
+				check(VarProperty.PropertyClass);
+
+				if (VarProperty.PropertyClass->IsChildOf(UClass::StaticClass()))
+				{
+					auto* Result = new (EC_InternalUseOnlyConstructor, Scope, Name, ObjectFlags) UClassProperty(FObjectInitializer());
+					Result->MetaClass     = VarProperty.MetaClass;
+					Result->PropertyClass = VarProperty.PropertyClass;
+					return Result;
+				}
+				else
+				{
+					if (DoesAnythingInHierarchyHaveDefaultToInstanced(VarProperty.PropertyClass))
+					{
+						VarProperty.PropertyFlags |= CPF_InstancedReference;
+						AddEditInlineMetaData(VarProperty.MetaData);
+					}
+
+					auto* Result = new (EC_InternalUseOnlyConstructor, Scope, Name, ObjectFlags) UObjectProperty(FObjectInitializer());
+					Result->PropertyClass = VarProperty.PropertyClass;
+					return Result;
+				}
+
+			case CPT_WeakObjectReference:
+			{
+				check(VarProperty.PropertyClass);
+
+				auto* Result = new (EC_InternalUseOnlyConstructor, Scope, Name, ObjectFlags) UWeakObjectProperty(FObjectInitializer());
+				Result->PropertyClass = VarProperty.PropertyClass;
+				return Result;
+			}
+
+			case CPT_LazyObjectReference:
+			{
+				check(VarProperty.PropertyClass);
+
+				auto* Result = new (EC_InternalUseOnlyConstructor, Scope, Name, ObjectFlags) ULazyObjectProperty(FObjectInitializer());
+				Result->PropertyClass = VarProperty.PropertyClass;
+				return Result;
+			}
+
+			case CPT_AssetObjectReference:
+				check(VarProperty.PropertyClass);
+
+				if (VarProperty.PropertyClass->IsChildOf(UClass::StaticClass()))
+				{
+					auto* Result = new (EC_InternalUseOnlyConstructor, Scope, Name, ObjectFlags) UAssetClassProperty(FObjectInitializer());
+					Result->MetaClass     = VarProperty.MetaClass;
+					Result->PropertyClass = VarProperty.PropertyClass;
+					return Result;
+				}
+				else
+				{
+					auto* Result = new (EC_InternalUseOnlyConstructor, Scope, Name, ObjectFlags) UAssetObjectProperty(FObjectInitializer());
+					Result->PropertyClass = VarProperty.PropertyClass;
+					return Result;
+				}
+
+			case CPT_Interface:
+			{
+				check(VarProperty.PropertyClass);
+				check(VarProperty.PropertyClass->HasAnyClassFlags(CLASS_Interface));
+
+				auto* Result = new (EC_InternalUseOnlyConstructor, Scope, Name, ObjectFlags)  UInterfaceProperty(FObjectInitializer());
+				Result->InterfaceClass = VarProperty.PropertyClass;
+				return Result;
+			}
+
+			case CPT_Name:
+				return new (EC_InternalUseOnlyConstructor, Scope, Name, ObjectFlags) UNameProperty(FObjectInitializer());
+
+			case CPT_String:
+				return new (EC_InternalUseOnlyConstructor, Scope, Name, ObjectFlags) UStrProperty(FObjectInitializer());
+
+			case CPT_Text:
+				return new (EC_InternalUseOnlyConstructor, Scope, Name, ObjectFlags) UTextProperty(FObjectInitializer());
+
+			case CPT_Struct:
+			{
+				if (VarProperty.Struct->StructFlags & STRUCT_HasInstancedReference)
+				{
+					VarProperty.PropertyFlags |= CPF_ContainsInstancedReference;
+				}
+
+				auto* Result = new (EC_InternalUseOnlyConstructor, Scope, Name, ObjectFlags) UStructProperty(FObjectInitializer());
+				Result->Struct = VarProperty.Struct;
+				return Result;
+			}
+
+			case CPT_Delegate:
+				return new (EC_InternalUseOnlyConstructor, Scope, Name, ObjectFlags) UDelegateProperty(FObjectInitializer());
+
+			case CPT_MulticastDelegate:
+				return new (EC_InternalUseOnlyConstructor, Scope, Name, ObjectFlags) UMulticastDelegateProperty(FObjectInitializer());
+
+			default:
+				FError::Throwf(TEXT("Unknown property type %i"), (uint8)VarProperty.Type);
+		}
+
+		// Unreachable
+		check(false);
 		return nullptr;
 	}
 
@@ -3447,160 +3649,7 @@ UProperty* FHeaderParser::GetVarNameAndDim
 			ArrayDim = 2;
 		}
 
-		switch (VarProperty.Type)
-		{
-			case CPT_Byte:
-				NewProperty = new(EC_InternalUseOnlyConstructor, NewScope, PropertyName, ObjectFlags)UByteProperty(FObjectInitializer());
-				((UByteProperty*)NewProperty)->Enum = VarProperty.Enum;
-				break;
-
-			case CPT_Int8:
-				NewProperty = new(EC_InternalUseOnlyConstructor, NewScope, PropertyName, ObjectFlags)UInt8Property(FObjectInitializer());
-				break;
-
-			case CPT_Int16:
-				NewProperty = new(EC_InternalUseOnlyConstructor, NewScope, PropertyName, ObjectFlags)UInt16Property(FObjectInitializer());
-				break;
-
-			case CPT_Int:
-				NewProperty = new(EC_InternalUseOnlyConstructor, NewScope, PropertyName, ObjectFlags)UIntProperty(FObjectInitializer());
-				break;
-
-			case CPT_Int64:
-				NewProperty = new(EC_InternalUseOnlyConstructor, NewScope, PropertyName, ObjectFlags)UInt64Property(FObjectInitializer());
-				break;
-
-			case CPT_UInt16:
-				NewProperty = new(EC_InternalUseOnlyConstructor, NewScope, PropertyName, ObjectFlags)UUInt16Property(FObjectInitializer());
-				break;
-
-			case CPT_UInt32:
-				NewProperty = new(EC_InternalUseOnlyConstructor, NewScope, PropertyName, ObjectFlags)UUInt32Property(FObjectInitializer());
-				break;
-
-			case CPT_UInt64:
-				NewProperty = new(EC_InternalUseOnlyConstructor, NewScope, PropertyName, ObjectFlags)UUInt64Property(FObjectInitializer());
-				break;
-
-			case CPT_Bool:
-				NewProperty = new(EC_InternalUseOnlyConstructor, NewScope, PropertyName, ObjectFlags)UBoolProperty(FObjectInitializer());
-				((UBoolProperty*)NewProperty)->SetBoolSize(sizeof(bool), true);
-				break;
-
-			case CPT_Bool8:
-				NewProperty = new(EC_InternalUseOnlyConstructor, NewScope, PropertyName, ObjectFlags)UBoolProperty(FObjectInitializer());
-				((UBoolProperty*)NewProperty)->SetBoolSize((VariableCategory == EVariableCategory::Return) ? sizeof(bool) : sizeof(uint8), VariableCategory == EVariableCategory::Return);
-				break;
-
-			case CPT_Bool16:
-				NewProperty = new(EC_InternalUseOnlyConstructor, NewScope, PropertyName, ObjectFlags)UBoolProperty(FObjectInitializer());
-				((UBoolProperty*)NewProperty)->SetBoolSize((VariableCategory == EVariableCategory::Return) ? sizeof(bool) : sizeof(uint16), VariableCategory == EVariableCategory::Return);
-				break;
-
-			case CPT_Bool32:
-				NewProperty = new(EC_InternalUseOnlyConstructor, NewScope, PropertyName, ObjectFlags)UBoolProperty(FObjectInitializer());
-				((UBoolProperty*)NewProperty)->SetBoolSize((VariableCategory == EVariableCategory::Return) ? sizeof(bool) : sizeof(uint32), VariableCategory == EVariableCategory::Return);
-				break;
-
-			case CPT_Bool64:
-				NewProperty = new(EC_InternalUseOnlyConstructor, NewScope, PropertyName, ObjectFlags)UBoolProperty(FObjectInitializer());
-				((UBoolProperty*)NewProperty)->SetBoolSize((VariableCategory == EVariableCategory::Return) ? sizeof(bool) : sizeof(uint64), VariableCategory == EVariableCategory::Return);
-				break;
-
-			case CPT_Float:
-				NewProperty = new(EC_InternalUseOnlyConstructor, NewScope, PropertyName, ObjectFlags)UFloatProperty(FObjectInitializer());
-				break;
-
-			case CPT_Double:
-				NewProperty = new(EC_InternalUseOnlyConstructor, NewScope, PropertyName, ObjectFlags)UDoubleProperty(FObjectInitializer());
-				break;
-
-			case CPT_ObjectReference:
-				check(VarProperty.PropertyClass);
-				if (VarProperty.PropertyClass->IsChildOf(UClass::StaticClass()))
-				{
-					NewProperty = new(EC_InternalUseOnlyConstructor, NewScope, PropertyName, ObjectFlags)UClassProperty(FObjectInitializer());
-					((UClassProperty*)NewProperty)->MetaClass = VarProperty.MetaClass;
-				}
-				else
-				{
-					if ( DoesAnythingInHierarchyHaveDefaultToInstanced( VarProperty.PropertyClass ) )
-					{
-						VarProperty.PropertyFlags |= CPF_InstancedReference;
-						AddEditInlineMetaData(VarProperty.MetaData);
-					}
-					NewProperty = new(EC_InternalUseOnlyConstructor, NewScope, PropertyName, ObjectFlags)UObjectProperty(FObjectInitializer());
-				}
-
-				((UObjectPropertyBase*)NewProperty)->PropertyClass = VarProperty.PropertyClass;
-				break;
-
-			case CPT_WeakObjectReference:
-				check(VarProperty.PropertyClass);
-				NewProperty = new(EC_InternalUseOnlyConstructor, NewScope, PropertyName, ObjectFlags)UWeakObjectProperty(FObjectInitializer());
-				((UObjectPropertyBase*)NewProperty)->PropertyClass = VarProperty.PropertyClass;
-				break;
-
-			case CPT_LazyObjectReference:
-				check(VarProperty.PropertyClass);
-				NewProperty = new(EC_InternalUseOnlyConstructor, NewScope, PropertyName, ObjectFlags)ULazyObjectProperty(FObjectInitializer());
-				((UObjectPropertyBase*)NewProperty)->PropertyClass = VarProperty.PropertyClass;
-				break;
-
-			case CPT_AssetObjectReference:
-				check(VarProperty.PropertyClass);
-				if (VarProperty.PropertyClass->IsChildOf(UClass::StaticClass()))
-				{
-					NewProperty = new(EC_InternalUseOnlyConstructor, NewScope, PropertyName, ObjectFlags)UAssetClassProperty(FObjectInitializer());
-					((UAssetClassProperty*)NewProperty)->MetaClass = VarProperty.MetaClass;
-				}
-				else
-				{
-					NewProperty = new(EC_InternalUseOnlyConstructor, NewScope, PropertyName, ObjectFlags)UAssetObjectProperty(FObjectInitializer());
-				}
-				((UObjectPropertyBase*)NewProperty)->PropertyClass = VarProperty.PropertyClass;
-				break;
-
-			case CPT_Interface:
-				check(VarProperty.PropertyClass);
-				check(VarProperty.PropertyClass->HasAnyClassFlags(CLASS_Interface));
-
-				NewProperty = new(EC_InternalUseOnlyConstructor, NewScope, PropertyName, ObjectFlags) UInterfaceProperty(FObjectInitializer());
-				((UInterfaceProperty*)NewProperty)->InterfaceClass = VarProperty.PropertyClass;
-				break;
-
-			case CPT_Name:
-				NewProperty = new(EC_InternalUseOnlyConstructor, NewScope, PropertyName, ObjectFlags)UNameProperty(FObjectInitializer());
-				break;
-
-			case CPT_String:
-				NewProperty = new(EC_InternalUseOnlyConstructor, NewScope, PropertyName, ObjectFlags)UStrProperty(FObjectInitializer());
-				break;
-
-			case CPT_Text:
-				NewProperty = new(EC_InternalUseOnlyConstructor, NewScope, PropertyName, ObjectFlags)UTextProperty(FObjectInitializer());
-				break;
-
-			case CPT_Struct:
-				NewProperty = new(EC_InternalUseOnlyConstructor, NewScope, PropertyName, ObjectFlags)UStructProperty(FObjectInitializer());
-				if (VarProperty.Struct->StructFlags & STRUCT_HasInstancedReference)
-				{
-					VarProperty.PropertyFlags |= CPF_ContainsInstancedReference;
-				}
-				((UStructProperty*)NewProperty)->Struct = VarProperty.Struct;
-				break;
-
-			case CPT_Delegate:
-				NewProperty = new(EC_InternalUseOnlyConstructor, NewScope, PropertyName, ObjectFlags)UDelegateProperty(FObjectInitializer());
-				break;
-
-			case CPT_MulticastDelegate:
-				NewProperty = new(EC_InternalUseOnlyConstructor, NewScope, PropertyName, ObjectFlags)UMulticastDelegateProperty(FObjectInitializer());
-				break;
-
-			default:
-				FError::Throwf(TEXT("Unknown property type %i"), (uint8)VarProperty.Type);
-		}
+		NewProperty = CreateVariableProperty(VarProperty, NewScope, PropertyName, ObjectFlags, VariableCategory);
 
 		if( Array )
 		{
@@ -7847,26 +7896,3 @@ UDelegateFunction* FHeaderParser::CreateDelegateFunction(const FFuncInfo &FuncIn
 {
 	return CreateFunctionImpl<UDelegateFunction>(FuncInfo, IsInAClass() ? (UObject*)GetCurrentClass() : (UObject*)GetCurrentFileScope()->GetSourceFile()->GetPackage(), GetCurrentScope());
 }
-
-bool FHeaderParser::DoesAnythingInHierarchyHaveDefaultToInstanced(UClass* TestClass)
-{
-	bool bDefaultToInstanced = false;
-
-	UClass* Search = TestClass;
-	while (!bDefaultToInstanced && (Search != NULL))
-	{
-		bDefaultToInstanced = Search->HasAnyClassFlags(CLASS_DefaultToInstanced);		
-		if (!bDefaultToInstanced && !Search->HasAnyClassFlags(CLASS_Intrinsic | CLASS_Parsed))
-		{			
-			// The class might not have been parsed yet, look for declaration data.
-			auto ClassDeclarationDataPtr = GClassDeclarations.Find(Search->GetFName());
-			if (ClassDeclarationDataPtr)
-			{
-				bDefaultToInstanced = !!((*ClassDeclarationDataPtr)->ClassFlags & CLASS_DefaultToInstanced);
-			}
-		}
-		Search = Search->GetSuperClass();
-	}
-
-	return bDefaultToInstanced;
-}
\ No newline at end of file
diff --git a/Engine/Source/Programs/UnrealHeaderTool/Private/HeaderParser.h b/Engine/Source/Programs/UnrealHeaderTool/Private/HeaderParser.h
--- a/Engine/Source/Programs/UnrealHeaderTool/Private/HeaderParser.h
+++ b/Engine/Source/Programs/UnrealHeaderTool/Private/HeaderParser.h
@@ -744,9 +744,6 @@ protected:
 	 */
 	EFindName GetFindFlagForPropertyName(const TCHAR* PropertyName);
 
-	// Check to see if anything in the class hierarchy passed in has CLASS_DefaultToInstanced
-	static bool DoesAnythingInHierarchyHaveDefaultToInstanced(UClass* TestClass);
-
 	static void ValidatePropertyIsDeprecatedIfNecessary(FPropertyBase& VarProperty, FToken* OuterPropertyType);
 
 private:
