ValidateMetaDataFormat overloads and AddMetaDataToClassData moved out of FHeaderParser.

diff --git a/Engine/Source/Programs/UnrealHeaderTool/Private/HeaderParser.cpp b/Engine/Source/Programs/UnrealHeaderTool/Private/HeaderParser.cpp
--- a/Engine/Source/Programs/UnrealHeaderTool/Private/HeaderParser.cpp
+++ b/Engine/Source/Programs/UnrealHeaderTool/Private/HeaderParser.cpp
@@ -617,6 +617,98 @@ namespace
 		return nullptr;
 	}
 
+	/**
+	 * Ensures at script compile time that the metadata formatting is correct
+	 * @param	InKey			the metadata key being added
+	 * @param	InValue			the value string that will be associated with the InKey
+	 */
+	void ValidateMetaDataFormat(UField* Field, const FString& InKey, const FString& InValue)
+	{
+		if ((InKey == TEXT("UIMin")) || (InKey == TEXT("UIMax")) || (InKey == TEXT("ClampMin")) || (InKey == TEXT("ClampMax")))
+		{
+			if (!InValue.IsNumeric())
+			{
+				FError::Throwf(TEXT("Metadata value for '%s' is non-numeric : '%s'"), *InKey, *InValue);
+			}
+		}
+		else if (InKey == /*FBlueprintMetadata::MD_Protected*/ TEXT("BlueprintProtected"))
+		{
+			if (UFunction* Function = Cast<UFunction>(Field))
+			{
+				if (Function->HasAnyFunctionFlags(FUNC_Static))
+				{
+					// Determine if it's a function library
+					UClass* Class = Cast<UClass>(Function->GetOuterUClass());
+					while (Class != nullptr && Class->GetSuperClass() != UObject::StaticClass())
+					{
+						Class = Class->GetSuperClass();
+					}
+
+					if (Class != nullptr && Class->GetName() == TEXT("BlueprintFunctionLibrary"))
+					{
+						FError::Throwf(TEXT("%s doesn't make sense on static method '%s' in a blueprint function library"), *InKey, *Function->GetName());
+					}
+				}
+			}
+		}
+		else if (InKey == TEXT("DevelopmentStatus"))
+		{
+			const FString EarlyAccessValue(TEXT("EarlyAccess"));
+			const FString ExperimentalValue(TEXT("Experimental"));
+			if ((InValue != EarlyAccessValue) && (InValue != ExperimentalValue))
+			{
+				FError::Throwf(TEXT("'%s' metadata was '%s' but it must be %s or %s"), *InKey, *InValue, *ExperimentalValue, *EarlyAccessValue);
+			}
+		}
+		else if (InKey == TEXT("Units"))
+		{
+			// Check for numeric property
+			if (!Cast<UNumericProperty>(Field))
+			{
+				FError::Throwf(TEXT("'Units' meta data can only be applied to numeric properties"));
+			}
+			else if (!FUnitConversion::UnitFromString(*InValue))
+			{
+				FError::Throwf(TEXT("Unrecognized units (%s) specified for numeric property '%s'"), *InValue, *Field->GetDisplayNameText().ToString());
+			}
+		}
+	}
+
+	// Ensures at script compile time that the metadata formatting is correct
+	void ValidateMetaDataFormat(UField* Field, const TMap<FName, FString>& MetaData)
+	{
+		for (const auto& Pair : MetaData)
+		{
+			ValidateMetaDataFormat(Field, Pair.Key.ToString(), Pair.Value);
+		}
+	}
+
+	// Validates the metadata, then adds it to the class data
+	void AddMetaDataToClassData(UField* Field, const TMap<FName, FString>& InMetaData)
+	{
+		// Evaluate any key redirects on the passed in pairs
+		TMap<FName, FString> RemappedPairs;
+		RemappedPairs.Empty(InMetaData.Num());
+
+		for (const auto& Pair : InMetaData)
+		{
+			FName CurrentKey = Pair.Key;
+			FName NewKey = UMetaData::GetRemappedKeyName(CurrentKey);
+
+			if (NewKey != NAME_None)
+			{
+				UE_LOG(LogCompile, Warning, TEXT("Remapping old metadata key '%s' to new key '%s', please update the declaration."), *CurrentKey.ToString(), *NewKey.ToString());
+				CurrentKey = NewKey;
+			}
+
+			RemappedPairs.Add(CurrentKey, Pair.Value);
+		}
+
+		// Finish validating and associate the metadata with the field
+		ValidateMetaDataFormat(Field, RemappedPairs);
+		FClassMetaData::AddMetaData(Field, RemappedPairs);
+	}
+
 	bool IsPropertySupportedByBlueprint(const UProperty* Property)
 	{
 		if (Property == NULL)
@@ -5776,98 +5868,6 @@ void FHeaderParser::CompileFunctionDeclaration(FUnrealSourceFile& SourceFile, FC
 	}
 }
 
-/**
- * Ensures at script compile time that the metadata formatting is correct
- * @param	InKey			the metadata key being added
- * @param	InValue			the value string that will be associated with the InKey
- */
-void FHeaderParser::ValidateMetaDataFormat(UField* Field, const FString& InKey, const FString& InValue)
-{
-	if ((InKey == TEXT("UIMin")) || (InKey == TEXT("UIMax")) || (InKey == TEXT("ClampMin")) || (InKey == TEXT("ClampMax")))
-	{
-		if (!InValue.IsNumeric())
-		{
-			FError::Throwf(TEXT("Metadata value for '%s' is non-numeric : '%s'"), *InKey, *InValue);
-		}
-	}
-	else if (InKey == /*FBlueprintMetadata::MD_Protected*/ TEXT("BlueprintProtected"))
-	{
-		if (UFunction* Function = Cast<UFunction>(Field))
-		{
-			if (Function->HasAnyFunctionFlags(FUNC_Static))
-			{
-				// Determine if it's a function library
-				UClass* Class = Cast<UClass>(Function->GetOuterUClass());
-				while (Class != nullptr && Class->GetSuperClass() != UObject::StaticClass())
-				{
-					Class = Class->GetSuperClass();
-				}
-
-				if (Class != nullptr && Class->GetName() == TEXT("BlueprintFunctionLibrary"))
-				{
-					FError::Throwf(TEXT("%s doesn't make sense on static method '%s' in a blueprint function library"), *InKey, *Function->GetName());
-				}
-			}
-		}
-	}
-	else if (InKey == TEXT("DevelopmentStatus"))
-	{
-		const FString EarlyAccessValue(TEXT("EarlyAccess"));
-		const FString ExperimentalValue(TEXT("Experimental"));
-		if ((InValue != EarlyAccessValue) && (InValue != ExperimentalValue))
-		{
-			FError::Throwf(TEXT("'%s' metadata was '%s' but it must be %s or %s"), *InKey, *InValue, *ExperimentalValue, *EarlyAccessValue);
-		}
-	}
-	else if (InKey == TEXT("Units"))
-	{
-		// Check for numeric property
-		if (!Cast<UNumericProperty>(Field))
-		{
-			FError::Throwf(TEXT("'Units' meta data can only be applied to numeric properties"));
-		}
-		else if (!FUnitConversion::UnitFromString(*InValue))
-		{
-			FError::Throwf(TEXT("Unrecognized units (%s) specified for numeric property '%s'"), *InValue, *Field->GetDisplayNameText().ToString());
-		}
-	}
-}
-
-// Validates the metadata, then adds it to the class data
-void FHeaderParser::AddMetaDataToClassData(UField* Field, const TMap<FName, FString>& InMetaData)
-{
-	// Evaluate any key redirects on the passed in pairs
-	TMap<FName, FString> RemappedPairs;
-	RemappedPairs.Empty(InMetaData.Num());
-
-	for (TMap<FName, FString>::TConstIterator It(InMetaData); It; ++It)
-	{
-		FName CurrentKey = It.Key();
-		FName NewKey = UMetaData::GetRemappedKeyName(CurrentKey);
-
-		if (NewKey != NAME_None)
-		{
-			UE_LOG(LogCompile, Warning, TEXT("Remapping old metadata key '%s' to new key '%s', please update the declaration."), *CurrentKey.ToString(), *NewKey.ToString());
-			CurrentKey = NewKey;
-		}
-
-		RemappedPairs.Add(CurrentKey, It.Value());
-	}
-
-	// Finish validating and associate the metadata with the field
-	ValidateMetaDataFormat(Field, RemappedPairs);
-	FClassMetaData::AddMetaData(Field, RemappedPairs);
-}
-
-// Ensures at script compile time that the metadata formatting is correct
-void FHeaderParser::ValidateMetaDataFormat(UField* Field, const TMap<FName, FString>& MetaData)
-{
-	for (TMap<FName, FString>::TConstIterator It(MetaData); It; ++It)
-	{
-		ValidateMetaDataFormat(Field, It.Key().ToString(), It.Value());
-	}
-}
-
 /** Parses optional metadata text. */
 void FHeaderParser::ParseFieldMetaData(TMap<FName, FString>& MetaData, const TCHAR* FieldName)
 {
diff --git a/Engine/Source/Programs/UnrealHeaderTool/Private/HeaderParser.h b/Engine/Source/Programs/UnrealHeaderTool/Private/HeaderParser.h
--- a/Engine/Source/Programs/UnrealHeaderTool/Private/HeaderParser.h
+++ b/Engine/Source/Programs/UnrealHeaderTool/Private/HeaderParser.h
@@ -630,21 +630,6 @@ protected:
 	 */
 	static bool HasMemberProperties( const UStruct* Scope );
 
-	/**
-	 * Ensures at script compile time that the metadata formatting is correct
-	 * @param	InKey			the metadata key being added
-	 * @param	InValue			the value string that will be associated with the InKey
-	 */
-	void ValidateMetaDataFormat(UField* Field, const FString& InKey, const FString& InValue);
-
-	/**
-	 * Ensures at script compile time that the metadata formatting is correct
-	 */
-	void ValidateMetaDataFormat(UField* Field, const TMap<FName, FString>& MetaData);
-
-	// Validates the metadata, then adds it to the class data
-	void AddMetaDataToClassData(UField* Field, const TMap<FName, FString>& InMetaData);
-
 	/**
 	 * Parses optional metadata text.
 	 *
