Delete old generated .cpp files when generating less code.
Provided by [REDACTED_PERSON] at Black Tusk.

diff --git a/Engine/Source/Programs/UnrealHeaderTool/Private/CodeGenerator.cpp b/Engine/Source/Programs/UnrealHeaderTool/Private/CodeGenerator.cpp
--- a/Engine/Source/Programs/UnrealHeaderTool/Private/CodeGenerator.cpp
+++ b/Engine/Source/Programs/UnrealHeaderTool/Private/CodeGenerator.cpp
@@ -4846,11 +4846,25 @@ void FNativeClassHeaderGenerator::ExportGeneratedCPP()
 	// delete the old .cpp file that will cause link errors if it's left around (Engine.generated.cpp and Engine.generated.1.cpp will 
 	// conflict now that we no longer use Engine.generated.cpp to #include Engine.generated.1.cpp, and UBT would compile all 3)
 	// @todo: This is a temp measure so we don't force everyone to require a Clean
- 	if (GeneratedFunctionBodyTextSplit.Num() > 1)
- 	{
- 		FString CppPath = ModuleInfo->GeneratedCPPFilenameBase + TEXT(".cpp");
+	if (GeneratedFunctionBodyTextSplit.Num() > 1)
+	{
+		FString CppPath = ModuleInfo->GeneratedCPPFilenameBase + TEXT(".cpp");
 		IFileManager::Get().Delete(*CppPath);
- 	}
+	}
+
+	// Delete old generated .cpp files which we don't need because we generated less code than last time.
+	{
+		TArray<FString> FoundFiles;
+		IFileManager::Get().FindFiles(FoundFiles, *(ModuleInfo->GeneratedCPPFilenameBase + TEXT(".*.cpp")), true, false);
+		FString BaseDir = FPaths::GetPath(ModuleInfo->GeneratedCPPFilenameBase);
+		for (FString& File : FoundFiles)
+		{
+			if (!NumberedHeaderNames.Contains(File))
+			{
+				IFileManager::Get().Delete(*FPaths::Combine(*BaseDir, *File));
+			}
+		}
+	}
 }
 
 /** Get all script plugins based on ini setting */
