UE-12214 Calling a function that returns TWeakObjectPtr in blueprints causes crash

Only a simple pointer to UObjects can be a parameter of BlueprintCallable (or BlueprintImplementableEvent.. etc) function. No WeakObjPtr or AssetSubclassOf is accepted. Strict check in UHT.

diff --git a/Engine/Source/Programs/UnrealHeaderTool/Private/HeaderParser.cpp b/Engine/Source/Programs/UnrealHeaderTool/Private/HeaderParser.cpp
--- a/Engine/Source/Programs/UnrealHeaderTool/Private/HeaderParser.cpp
+++ b/Engine/Source/Programs/UnrealHeaderTool/Private/HeaderParser.cpp
@@ -709,7 +709,7 @@ namespace
 		FClassMetaData::AddMetaData(Field, RemappedPairs);
 	}
 
-	bool IsPropertySupportedByBlueprint(const UProperty* Property)
+	bool IsPropertySupportedByBlueprint(const UProperty* Property, bool bMemberVariable)
 	{
 		if (Property == NULL)
 		{
@@ -717,13 +717,13 @@ namespace
 		}
 		if (auto ArrayProperty = Cast<const UArrayProperty>(Property))
 		{
-			return IsPropertySupportedByBlueprint(ArrayProperty->Inner);
+			// Inner Property can be handled as a member variable
+			return IsPropertySupportedByBlueprint(ArrayProperty->Inner, true);
 		}
 
 		const bool bSupportedType = Property->IsA<UInterfaceProperty>()
 			|| Property->IsA<UClassProperty>()
-			|| Property->IsA<UAssetClassProperty>() //?
-			|| Property->IsA<UObjectPropertyBase>()
+			|| Property->IsA<UObjectProperty>()
 			|| Property->IsA<UStructProperty>()
 			|| Property->IsA<UFloatProperty>()
 			|| Property->IsA<UIntProperty>()
@@ -735,7 +735,9 @@ namespace
 			|| Property->IsA<UMulticastDelegateProperty>()
 			|| Property->IsA<UDelegateProperty>();
 
-		return bSupportedType;
+		const bool bIsSupportedMemberVariable = Property->IsA<UObjectPropertyBase>();
+
+		return bSupportedType || (bIsSupportedMemberVariable && bMemberVariable);
 	}
 }
 	
@@ -5832,7 +5834,7 @@ void FHeaderParser::CompileFunctionDeclaration(FUnrealSourceFile& SourceFile, FC
 				FError::Throwf(TEXT("Static array cannot be exposed to blueprint. Function: %s Parameter %s\n"), *TopFunction->GetName(), *Param->GetName());
 			}
 
-			if (!IsPropertySupportedByBlueprint(Param))
+			if (!IsPropertySupportedByBlueprint(Param, false))
 			{
 				FError::Throwf(TEXT("Type '%s' is not supported by blueprint. Function: %s Parameter %s\n"), *Param->GetCPPType(), *TopFunction->GetName(), *Param->GetName());
 			}
@@ -6139,7 +6141,7 @@ void FHeaderParser::CompileVariableDeclaration(FClasses& AllClasses, UStruct* St
 			FError::Throwf(TEXT("Static array cannot be exposed to blueprint %s.%s"), *Struct->GetName(), *NewProperty->GetName());
 		}
 
-		if (NewProperty->HasAnyPropertyFlags(CPF_BlueprintVisible) && !IsPropertySupportedByBlueprint(NewProperty))
+		if (NewProperty->HasAnyPropertyFlags(CPF_BlueprintVisible) && !IsPropertySupportedByBlueprint(NewProperty, true))
 		{
 			FError::Throwf(TEXT("Type '%s' is not supported by blueprint. %s.%s"), *NewProperty->GetCPPType(), *Struct->GetName(), *NewProperty->GetName());
 		}
diff --git a/Engine/Source/Runtime/GameplayAbilities/Public/GameplayCueNotify_Actor.h b/Engine/Source/Runtime/GameplayAbilities/Public/GameplayCueNotify_Actor.h
--- a/Engine/Source/Runtime/GameplayAbilities/Public/GameplayCueNotify_Actor.h
+++ b/Engine/Source/Runtime/GameplayAbilities/Public/GameplayCueNotify_Actor.h
@@ -44,7 +44,7 @@ class GAMEPLAYABILITIES_API AGameplayCueNotify_Actor : public AActor
 
 	/** Generic Event Graph event that will get called for every event type */
 	UFUNCTION(BlueprintImplementableEvent, Category = "GameplayCueNotify", DisplayName = "HandleGameplayCue")
-	void K2_HandleGameplayCue(TWeakObjectPtr<AActor> MyTarget, EGameplayCueEvent::Type EventType, FGameplayCueParameters Parameters);
+	void K2_HandleGameplayCue(AActor* MyTarget, EGameplayCueEvent::Type EventType, FGameplayCueParameters Parameters);
 
 	UFUNCTION(BlueprintNativeEvent, Category = "GameplayCueNotify")
 	bool OnExecute(AActor* MyTarget, FGameplayCueParameters Parameters);
diff --git a/Engine/Source/Runtime/GameplayAbilities/Public/GameplayCueNotify_Static.h b/Engine/Source/Runtime/GameplayAbilities/Public/GameplayCueNotify_Static.h
--- a/Engine/Source/Runtime/GameplayAbilities/Public/GameplayCueNotify_Static.h
+++ b/Engine/Source/Runtime/GameplayAbilities/Public/GameplayCueNotify_Static.h
@@ -41,7 +41,7 @@ class GAMEPLAYABILITIES_API UGameplayCueNotify_Static : public UObject
 
 	/** Generic Event Graph event that will get called for every event type */
 	UFUNCTION(BlueprintImplementableEvent, Category = "GameplayCueNotify", DisplayName = "HandleGameplayCue")
-	void K2_HandleGameplayCue(TWeakObjectPtr<AActor> MyTarget, EGameplayCueEvent::Type EventType, FGameplayCueParameters Parameters) const;
+	void K2_HandleGameplayCue(AActor* MyTarget, EGameplayCueEvent::Type EventType, FGameplayCueParameters Parameters) const;
 
 	UFUNCTION(BlueprintNativeEvent, BlueprintPure, Category = "GameplayCueNotify")
 	bool OnExecute(AActor* MyTarget, FGameplayCueParameters Parameters) const;
