Stop UHT from emiting warnings on missing *_Validate and *_Implementation functions
Introduce versioning system to GENERATED_*BODY macros
Search for existing functions using case sensitive and whole word match.
Move checks for existing of *_Validate and *_Implementation functions to *generated.cpp to allow changing that code in hotfixes.

diff --git a/Engine/Config/BaseEngine.ini b/Engine/Config/BaseEngine.ini
--- a/Engine/Config/BaseEngine.ini
+++ b/Engine/Config/BaseEngine.ini
@@ -2209,4 +2209,7 @@ DeployServerPort=53501
 
 [CrashReportClient]
 ;CrashReportReceiverIP="http://localhost:57005"
-CrashReportReceiverIP="http://crashreporter.epicgames.com:57005"
\ No newline at end of file
+CrashReportReceiverIP="http://crashreporter.epicgames.com:57005"
+
+[UnrealHeaderTool]
+GeneratedCodeVersion=V1
diff --git a/Engine/Source/Programs/UnrealBuildTool/Configuration/UEBuildModule.cs b/Engine/Source/Programs/UnrealBuildTool/Configuration/UEBuildModule.cs
--- a/Engine/Source/Programs/UnrealBuildTool/Configuration/UEBuildModule.cs
+++ b/Engine/Source/Programs/UnrealBuildTool/Configuration/UEBuildModule.cs
@@ -1940,7 +1940,7 @@ namespace UnrealBuildTool
 				bool HasUObjects;
 				var DirectIncludes = CPPEnvironment.GetDirectIncludeDependencies( Target, UObjectHeaderFileItem, BuildPlatform, bOnlyCachedDependencies:false, HasUObjects:out HasUObjects );
 				Debug.Assert( DirectIncludes != null );
-
+				
 				if (HasUObjects)
 				{ 
 					if (UObjectHeaderFileItem.AbsolutePath.StartsWith(ClassesFolder))
@@ -1965,7 +1965,8 @@ namespace UnrealBuildTool
 				ModuleType					= this.Type.ToString(),
 				PublicUObjectClassesHeaders = _AllClassesHeaders    .ToList(),
 				PublicUObjectHeaders        = _PublicUObjectHeaders .ToList(),
-				PrivateUObjectHeaders       = _PrivateUObjectHeaders.ToList()
+				PrivateUObjectHeaders       = _PrivateUObjectHeaders.ToList(),
+				GeneratedCodeVersion = this.Target.Rules.GetGeneratedCodeVersion()
 			};
 
 			return CachedModuleUHTInfo;
diff --git a/Engine/Source/Programs/UnrealBuildTool/System/ExternalExecution.cs b/Engine/Source/Programs/UnrealBuildTool/System/ExternalExecution.cs
--- a/Engine/Source/Programs/UnrealBuildTool/System/ExternalExecution.cs
+++ b/Engine/Source/Programs/UnrealBuildTool/System/ExternalExecution.cs
@@ -71,12 +71,27 @@ namespace UnrealBuildTool
 		/** Base (i.e. extensionless) path+filename of the .generated files */
 		public string GeneratedCPPFilenameBase;
 
+		/** Version of code generated by UHT */
+		public EGeneratedCodeVersion GeneratedCodeVersion;
+
 		public override string ToString()
 		{
 			return ModuleName;
 		}
 	}
 
+	//
+	// This MUST be kept in sync with EGeneratedBodyVersion enum and 
+	// ToGeneratedBodyVersion function in UHT defined in GeneratedCodeVersion.h.
+	//
+	public enum EGeneratedCodeVersion
+	{
+		None,
+		V1,
+		V2,
+		VLatest = V2
+	};
+
 	public struct UHTManifest
 	{
 		public struct Module
@@ -92,6 +107,7 @@ namespace UnrealBuildTool
 			public string       PCH;
 			public string       GeneratedCPPFilenameBase;
 			public bool         SaveExportedHeaders;
+			public EGeneratedCodeVersion UHTGeneratedCodeVersion;
 
 			public override string ToString()
 			{
@@ -106,7 +122,8 @@ namespace UnrealBuildTool
 			RootBuildPath = InRootBuildPath;
 			TargetName    = Target.GetTargetName();
 
-			Modules = ModuleInfo.Select(Info => new Module{
+			Modules = ModuleInfo.Select(Info => new Module
+			{
 				Name                     = Info.ModuleName,
 				ModuleType				 = Info.ModuleType,
 				BaseDirectory            = Info.ModuleDirectory,
@@ -118,7 +135,7 @@ namespace UnrealBuildTool
 				PCH                      = Info.PCH,
 				GeneratedCPPFilenameBase = Info.GeneratedCPPFilenameBase,
 				SaveExportedHeaders      = true, // Rocket just regenerates all headers to the project directory now
-
+				UHTGeneratedCodeVersion = Info.GeneratedCodeVersion,
 			}).ToList();
 		}
 
diff --git a/Engine/Source/Programs/UnrealBuildTool/System/RulesCompiler.cs b/Engine/Source/Programs/UnrealBuildTool/System/RulesCompiler.cs
--- a/Engine/Source/Programs/UnrealBuildTool/System/RulesCompiler.cs
+++ b/Engine/Source/Programs/UnrealBuildTool/System/RulesCompiler.cs
@@ -970,6 +970,14 @@ namespace UnrealBuildTool
 		{
 			return "";
 		}
+
+		/// <summary>
+		/// Allow target module to override UHT code generation version.
+		/// </summary>
+		public virtual EGeneratedCodeVersion GetGeneratedCodeVersion()
+		{
+			return EGeneratedCodeVersion.None;
+		}
 	}
 
 
diff --git a/Engine/Source/Programs/UnrealHeaderTool/Private/CodeGenerator.cpp b/Engine/Source/Programs/UnrealHeaderTool/Private/CodeGenerator.cpp
--- a/Engine/Source/Programs/UnrealHeaderTool/Private/CodeGenerator.cpp
+++ b/Engine/Source/Programs/UnrealHeaderTool/Private/CodeGenerator.cpp
@@ -45,6 +45,38 @@ const TCHAR* GetTabs(int32 NumTabs)
 	return TabString + TabPos;
 }
 
+bool HasIdentifierExactMatch(const FString &ClassDefinition, const FString& Identifier)
+{
+	int32 SearchStartPosition = INDEX_NONE;
+
+	for (;;)
+	{
+		auto IdentifierStartIndex = ClassDefinition.Find(Identifier, ESearchCase::CaseSensitive, ESearchDir::FromStart, SearchStartPosition);
+
+		// Class definition not found
+		if (IdentifierStartIndex == INDEX_NONE
+			// ClassDefinition can't start with CppImplName. If it starts, we'll get compiler error anyway
+			|| IdentifierStartIndex == 0)
+		{
+			return false;
+		}
+
+		// To find exact match, previous and next character can't be parts of identifier.
+		if (!FChar::IsIdentifier(ClassDefinition[IdentifierStartIndex - 1])
+			&& !FChar::IsIdentifier(ClassDefinition[IdentifierStartIndex + Identifier.Len()]))
+		{
+			return true;
+		}
+
+		// Exact match not found, continue search starting just after the end of identifier
+		SearchStartPosition = IdentifierStartIndex + Identifier.Len() + 1;
+	}
+
+	// We should never get here.
+	checkNoEntry();
+	return false;
+}
+
 /////////////////////////////////////////////////////
 // FFlagAudit
 
@@ -1762,17 +1794,12 @@ void FNativeClassHeaderGenerator::ExportClassesFromSourceFileInner(FUnrealSource
 
 			// Replication, add in the declaration for GetLifetimeReplicatedProps() automatically if there are any net flagged properties
 			bool bNeedsRep = false;
-			FUHTStringBuilder StaticChecks;
 			auto ClassName = FString(NameLookupCPP.GetNameCPP(Class));
 			for (TFieldIterator<UProperty> It(Class, EFieldIteratorFlags::ExcludeSuper); It; ++It)
 			{
 				if ((It->PropertyFlags & CPF_Net) != 0)
 				{
 					bNeedsRep = true;
-					StaticChecks.Logf(TEXT("\tstatic inline void StaticChecks()\r\n"));
-					StaticChecks.Logf(TEXT("\t{\r\n"));
-					StaticChecks.Logf(TEXT("\t\tstatic_assert(THasMemberFunction_GetLifetimeReplicatedProps<%s>::Value, \"Class %s has Net flagged properties and should declare member function: void GetLifetimeReplicatedProps( TArray< FLifetimeProperty > & OutLifetimeProps ) const override;\");\r\n"), ClassCPPName, ClassCPPName, ClassCPPName);
-					StaticChecks.Logf(TEXT("\t}\r\n"));
 					break;
 				}
 			}
@@ -1787,11 +1814,18 @@ void FNativeClassHeaderGenerator::ExportClassesFromSourceFileInner(FUnrealSource
 			auto ClassEnd = ClassRange.End;
 			auto ClassDefinition = FString(ClassEnd - ClassStart, ClassStart);
 
-			bool bHasGetLifetimeReplicatedProps = ClassDefinition.Contains(TEXT("GetLifetimeReplicatedProps"), ESearchCase::CaseSensitive);
+			bool bHasGetLifetimeReplicatedProps = HasIdentifierExactMatch(ClassDefinition, TEXT("GetLifetimeReplicatedProps"));
 
 			if (bNeedsRep && !bHasGetLifetimeReplicatedProps)
 			{
-				InterfaceBoilerplate.Logf(TEXT("\void GetLifetimeReplicatedProps(TArray<FLifetimeProperty>& OutLifetimeProps) const;\r\n"));
+				if (SourceFile.GetGeneratedCodeVersionForStruct(Class) == EGeneratedCodeVersion::V1)
+				{
+					InterfaceBoilerplate.Logf(TEXT("\void GetLifetimeReplicatedProps(TArray<FLifetimeProperty>& OutLifetimeProps) const;\r\n"));
+				}
+				else
+				{
+					FError::Throwf(TEXT("Class %s has Net flagged properties and should declare member function: void GetLifetimeReplicatedProps(TArray<FLifetimeProperty>& OutLifetimeProps) const override"), ClassCPPName);
+				}
 			}
 
 			auto NoPureDeclsMacroName = SourceFile.GetGeneratedMacroName(ClassData, TEXT("_INCLASS_IINTERFACE_NO_PURE_DECLS"));
@@ -1799,7 +1833,7 @@ void FNativeClassHeaderGenerator::ExportClassesFromSourceFileInner(FUnrealSource
 			InClassNoPureDeclsMacroCalls.Logf(TEXT("\t%s\r\n"), *NoPureDeclsMacroName);
 
 			auto MacroName = SourceFile.GetGeneratedMacroName(ClassData, TEXT("_INCLASS_IINTERFACE"));
-			WriteMacro(GeneratedHeaderText, MacroName, InterfaceBoilerplate + StaticChecks);
+			WriteMacro(GeneratedHeaderText, MacroName, InterfaceBoilerplate);
 			InClassMacroCalls.Logf(TEXT("\t%s\r\n"), *MacroName);
 		}
 		else
@@ -1861,16 +1895,11 @@ void FNativeClassHeaderGenerator::ExportClassesFromSourceFileInner(FUnrealSource
 
 				// Replication, add in the declaration for GetLifetimeReplicatedProps() automatically if there are any net flagged properties
 				bool bHasNetFlaggedProperties = false;
-				FUHTStringBuilder StaticChecks;
 				for (TFieldIterator<UProperty> It(Class, EFieldIteratorFlags::ExcludeSuper); It; ++It)
 				{
 					if ((It->PropertyFlags & CPF_Net) != 0)
 					{
 						bHasNetFlaggedProperties = true;
-						StaticChecks.Logf(TEXT("\tstatic inline void StaticChecks()\r\n"));
-						StaticChecks.Logf(TEXT("\t{\r\n"));
-						StaticChecks.Logf(TEXT("\t\tstatic_assert(THasMemberFunction_GetLifetimeReplicatedProps<%s>::Value, \"Class %s has Net flagged properties and should have member void GetLifetimeReplicatedProps( TArray< FLifetimeProperty > & OutLifetimeProps ) const override;\");\r\n"), ClassCPPName, ClassCPPName, ClassCPPName);
-						StaticChecks.Logf(TEXT("\t}\r\n"));
 						break;
 					}
 				}
@@ -1885,15 +1914,23 @@ void FNativeClassHeaderGenerator::ExportClassesFromSourceFileInner(FUnrealSource
 				auto ClassEnd = ClassRange.End;
 				auto ClassDefinition = FString(ClassEnd - ClassStart, ClassStart);
 
-				bool bHasGetLifetimeReplicatedProps = ClassDefinition.Contains(TEXT("GetLifetimeReplicatedProps"), ESearchCase::CaseSensitive);
+				bool bHasGetLifetimeReplicatedProps = HasIdentifierExactMatch(ClassDefinition, TEXT("GetLifetimeReplicatedProps"));
 
 				if (bHasNetFlaggedProperties && !bHasGetLifetimeReplicatedProps)
 				{
-					ClassBoilerplate.Logf(TEXT("\tvoid GetLifetimeReplicatedProps(TArray<FLifetimeProperty>& OutLifetimeProps) const override;\r\n"));
+					// Default version autogenerates declarations.
+					if (SourceFile.GetGeneratedCodeVersionForStruct(Class) == EGeneratedCodeVersion::V1)
+					{
+						ClassBoilerplate.Logf(TEXT("\tvoid GetLifetimeReplicatedProps(TArray<FLifetimeProperty>& OutLifetimeProps) const override;\r\n"));
+					}
+					else
+					{
+						FError::Throwf(TEXT("Class %s has Net flagged properties and should declare member function: void GetLifetimeReplicatedProps(TArray<FLifetimeProperty>& OutLifetimeProps) const override"), ClassCPPName);
+					}
 				}
 
 				auto NoPureDeclsMacroName = SourceFile.GetGeneratedMacroName(ClassData, TEXT("_INCLASS_NO_PURE_DECLS"));
-				WriteMacro(GeneratedHeaderText, NoPureDeclsMacroName, ClassBoilerplate + StaticChecks);
+				WriteMacro(GeneratedHeaderText, NoPureDeclsMacroName, ClassBoilerplate);
 				InClassNoPureDeclsMacroCalls.Logf(TEXT("\t%s\r\n"), *NoPureDeclsMacroName);
 
 				FString MacroName = SourceFile.GetGeneratedMacroName(ClassData, TEXT("_INCLASS"));
@@ -3499,11 +3536,11 @@ FString GetFunctionConstModifierString(UFunction* Function)
 	return FString();
 }
 
-void FNativeClassHeaderGenerator::ExportFunctionChecks(const FFuncInfo& FunctionData, FUHTStringBuilder& CheckClasses, FUHTStringBuilder& StaticChecks, const FString& ClassName)
+void FNativeClassHeaderGenerator::CheckRPCFunctions(const FFuncInfo& FunctionData, const FString& ClassName, bool bHasImplementation, bool bHasValidate, const FUnrealSourceFile& SourceFile, int32 GeneratedBodyLine)
 {
 	auto Function = FunctionData.FunctionReference;
 	auto FunctionReturnType = GetFunctionReturnString(Function);
-	auto ConstModifier = GetFunctionConstModifierString(Function);
+	auto ConstModifier = GetFunctionConstModifierString(Function) + TEXT(" ");
 
 	auto bIsNative = Function->HasAllFunctionFlags(FUNC_Native);
 	auto bIsNet = Function->HasAllFunctionFlags(FUNC_Net);
@@ -3545,19 +3582,15 @@ void FNativeClassHeaderGenerator::ExportFunctionChecks(const FFuncInfo& Function
 	//
 	// Add proper checks if necessary.
 	//
-	FString FunctionName;
-	if (bNeedsImplementation)
+	if (bNeedsImplementation && !bHasImplementation)
 	{
-		FunctionName = FunctionData.CppImplName;
-		StaticChecks.Logf(TEXT("\t\tstatic_assert(THasMemberFunction_%s<%s>::Value, \"%s Declare function %s %s(%s) %s in class %s.\");\r\n"), *FunctionName, *ClassName, *AssertMessage, *FunctionReturnType, *FunctionName, *ParameterString, *ConstModifier, *ClassName);
-		CheckClasses.Logf(TEXT("GENERATE_MEMBER_FUNCTION_CHECK(%s, %s, %s, %s)\r\n"), *FunctionName, *FunctionReturnType, *ConstModifier, *ParameterString);
+		// FullFilePath(LineNumber): error:
+		FError::Throwf(TEXT("%s(%d): error: %s Declare function %s %s(%s) %sin class %s."), *SourceFile.GetFilename(), GeneratedBodyLine, *AssertMessage, *FunctionReturnType, *FunctionData.CppImplName, *ParameterString, *ConstModifier, *ClassName);
 	}
 	
-	if (bNeedsValidate)
+	if (bNeedsValidate && !bHasValidate)
 	{
-		FunctionName = FunctionData.CppValidationImplName;
-		StaticChecks.Logf(TEXT("\t\tstatic_assert(THasMemberFunction_%s<%s>::Value, \"%s Declare function bool %s(%s) %s in class %s.\");\r\n"), *FunctionName, *ClassName, *AssertMessage, *FunctionName, *ParameterString, *ConstModifier, *ClassName);
-		CheckClasses.Logf(TEXT("GENERATE_MEMBER_FUNCTION_CHECK(%s, bool, %s, %s)\r\n"), *FunctionName, *ConstModifier, *ParameterString);
+		FError::Throwf(TEXT("%s(%d): error: %s Declare function bool %s(%s) %sin class %s."), *SourceFile.GetFilename(), GeneratedBodyLine, *AssertMessage, *FunctionData.CppValidationImplName, *ParameterString, *ConstModifier, *ClassName);
 	}
 }
 
@@ -3821,8 +3854,8 @@ void FNativeClassHeaderGenerator::ExportFunctionThunk(FUHTStringBuilder& RPCWrap
 
 	auto FunctionName = Function->GetName();
 
-	bool bHasImplementation = ClassDefinition.Contains(FunctionData.CppImplName, ESearchCase::CaseSensitive);
-	bool bHasValidate = ClassDefinition.Contains(FunctionData.CppValidationImplName, ESearchCase::CaseSensitive);
+	bool bHasImplementation = HasIdentifierExactMatch(ClassDefinition, FunctionData.CppImplName);
+	bool bHasValidate = HasIdentifierExactMatch(ClassDefinition, FunctionData.CppValidationImplName);
 
 	auto bShouldEnableImplementationDeprecation =
 		// Enable deprecation warnings only if GENERATED_BODY is used inside class or interface (not GENERATED_UCLASS_BODY etc.)
@@ -3844,12 +3877,6 @@ void FNativeClassHeaderGenerator::ExportFunctionThunk(FUHTStringBuilder& RPCWrap
 	// Call the validate function if there is one
 	if (!(FunctionData.FunctionExportFlags & FUNCEXPORT_CppStatic) && (FunctionData.FunctionFlags & FUNC_NetValidate))
 	{
-		// Call the validate function, and if it fails, return (which will NOT execute the _Implementation function below)
-		// NOTE - We don't need to set "Result", since RPC calls don't return any values...
-		if (bShouldEnableValidateDeprecation)
-		{
-			DeprecationWarningOutputDevice.Logf(TEXT("EMIT_CUSTOM_WARNING_AT_LINE(%d, \"Function %s::%s needs additional native validation function declaration %s due to its properties. Currently %s declaration is autogenerated by UHT. Since next release you'll have to provide declaration on your own. Please update your code before upgrading to the next release, otherwise your project will no longer compile.\")\r\n"), FunctionData.MacroLine, *ClassName, *FunctionName, *FunctionData.CppValidationImplName, *FunctionData.CppValidationImplName);
-		}
 		RPCWrappers.Logf(TEXT("\t\tif (!this->%s(%s))") LINE_TERMINATOR, *FunctionData.CppValidationImplName, *ParameterList);
 		RPCWrappers.Logf(TEXT("\t\t{") LINE_TERMINATOR);
 		RPCWrappers.Logf(TEXT("\t\t\tRPC_ValidateFailed(TEXT(\"%s\"));") LINE_TERMINATOR, *FunctionData.CppValidationImplName);
@@ -3857,10 +3884,6 @@ void FNativeClassHeaderGenerator::ExportFunctionThunk(FUHTStringBuilder& RPCWrap
 		RPCWrappers.Logf(TEXT("\t\t}") LINE_TERMINATOR);
 	}
 
-	if (bShouldEnableImplementationDeprecation)
-	{
-		DeprecationWarningOutputDevice.Logf(TEXT("EMIT_CUSTOM_WARNING_AT_LINE(%d, \"Function %s::%s needs additional native declaration %s due to its properties. Currently %s declaration is autogenerated by UHT. Since next release you'll have to provide declaration on your own. Please update your code before upgrading to the next release, otherwise your project will no longer compile.\")\r\n"), FunctionData.MacroLine, *ClassName, *FunctionName, *FunctionData.CppImplName, *FunctionData.CppImplName);
-	}
 	// write out the return value
 	RPCWrappers.Log(TEXT("\t\t"));
 	if (Return != NULL)
@@ -3919,12 +3942,10 @@ FString FNativeClassHeaderGenerator::GetFunctionParameterString(UFunction* Funct
 void FNativeClassHeaderGenerator::ExportNativeFunctions(FUnrealSourceFile& SourceFile, UClass* Class, FClassMetaData* ClassData)
 {
 	FUHTStringBuilder RPCWrappers;
-	FUHTStringBuilder	AutogeneratedBlueprintFunctionDeclarations;
-	FUHTStringBuilder	AutogeneratedBlueprintFunctionDeclarationsOnlyNotDeclared;
-	FUHTStringBuilder	MemberFunctionCheckClasses;
-	FUHTStringBuilder StaticChecks;
-	StaticChecks.Logf(TEXT("\tstatic inline void StaticChecks_Implementation_Validate()\r\n"));
-	StaticChecks.Logf(TEXT("\t{\r\n"));
+	FUHTStringBuilder AutogeneratedBlueprintFunctionDeclarations;
+	FUHTStringBuilder AutogeneratedBlueprintFunctionDeclarationsOnlyNotDeclared;
+
+	auto ClassName = Class->GetName();
 
 	auto ClassRange = ClassDefinitionRange();
 	if (ClassDefinitionRanges.Contains(Class))
@@ -3966,8 +3987,8 @@ void FNativeClassHeaderGenerator::ExportNativeFunctions(FUnrealSourceFile& Sourc
 
 			auto FunctionName = Function->GetName();
 
-			bool bHasImplementation = ClassDefinition.Contains(FunctionData.CppImplName, ESearchCase::CaseSensitive);
-			bool bHasValidate = ClassDefinition.Contains(FunctionData.CppValidationImplName, ESearchCase::CaseSensitive);
+			bool bHasImplementation = HasIdentifierExactMatch(ClassDefinition, FunctionData.CppImplName);
+			bool bHasValidate = HasIdentifierExactMatch(ClassDefinition, FunctionData.CppValidationImplName);
 
 			//Emit warning here if necessary
 			FUHTStringBuilder FunctionDeclaration;
@@ -3995,7 +4016,12 @@ void FNativeClassHeaderGenerator::ExportNativeFunctions(FUnrealSourceFile& Sourc
 				AutogeneratedBlueprintFunctionDeclarationsOnlyNotDeclared.Log(*FunctionDeclaration);
 			}
 
-			ExportFunctionChecks(FunctionData, MemberFunctionCheckClasses, StaticChecks, Class->HasAnyClassFlags(CLASS_Interface) ? *(FString(TEXT("I")) + Class->GetName()) : NameLookupCPP.GetNameCPP(Class));
+			// Versions that skip function autodeclaration throw an error when a function is missing.
+			if ((SourceFile.GetGeneratedCodeVersionForStruct(Class) > EGeneratedCodeVersion::V1) && ClassRange.bHasGeneratedBody)
+			{
+				CheckRPCFunctions(FunctionData, Class->HasAnyClassFlags(CLASS_Interface) ? *(FString(TEXT("I")) + ClassName) : NameLookupCPP.GetNameCPP(Class), bHasImplementation, bHasValidate, SourceFile, Class->HasAnyClassFlags(CLASS_Interface) ? ClassData->GetInterfaceGeneratedBodyLine() : ClassData->GetGeneratedBodyLine());
+			}
+			
 		}
 
 		if (bMultiLineUFUNCTION)
@@ -4020,15 +4046,20 @@ void FNativeClassHeaderGenerator::ExportNativeFunctions(FUnrealSourceFile& Sourc
 		RPCWrappers += TEXT("\t}") LINE_TERMINATOR;
 	}
 
-	StaticChecks.Logf(TEXT("\t}\r\n"));
-
 	FString MacroName = SourceFile.GetGeneratedMacroName(ClassData, TEXT("_RPC_WRAPPERS"));
 	WriteMacro(GeneratedHeaderText, MacroName, AutogeneratedBlueprintFunctionDeclarations + RPCWrappers);
 	InClassMacroCalls.Logf(TEXT("\t%s\r\n"), *MacroName);
 
 	// Put static checks before RPCWrappers to get proper messages from static asserts before compiler errors.
 	FString NoPureDeclsMacroName = SourceFile.GetGeneratedMacroName(ClassData, TEXT("_RPC_WRAPPERS_NO_PURE_DECLS"));
-	WriteMacro(GeneratedHeaderText, NoPureDeclsMacroName, AutogeneratedBlueprintFunctionDeclarationsOnlyNotDeclared + MemberFunctionCheckClasses + StaticChecks + RPCWrappers);
+	if (SourceFile.GetGeneratedCodeVersionForStruct(Class) > EGeneratedCodeVersion::V1)
+	{
+		WriteMacro(GeneratedHeaderText, NoPureDeclsMacroName, RPCWrappers);
+	}
+	else
+	{
+		WriteMacro(GeneratedHeaderText, NoPureDeclsMacroName, AutogeneratedBlueprintFunctionDeclarationsOnlyNotDeclared + RPCWrappers);
+	}
 	InClassNoPureDeclsMacroCalls.Logf(TEXT("\t%s\r\n"), *NoPureDeclsMacroName);
 }
 
@@ -4778,6 +4809,7 @@ void FNativeClassHeaderGenerator::ExportGeneratedCPP()
 	GeneratedLinkerFixupFunction.Logf(TEXT("void EmptyLinkFunctionForGeneratedCode%s() {}") LINE_TERMINATOR, *ModuleInfo->Name);
 
 	GeneratedCPPText.Log(*GeneratedPackageCPP);
+
 	// Add all names marked for export to a list (for sorting)
 	TArray<FString> Names;
 	for (TSet<FName>::TIterator It(ReferencedNames); It; ++It)
diff --git a/Engine/Source/Programs/UnrealHeaderTool/Private/HeaderParser.cpp b/Engine/Source/Programs/UnrealHeaderTool/Private/HeaderParser.cpp
--- a/Engine/Source/Programs/UnrealHeaderTool/Private/HeaderParser.cpp
+++ b/Engine/Source/Programs/UnrealHeaderTool/Private/HeaderParser.cpp
@@ -12,6 +12,7 @@
 #include "IScriptGeneratorPluginInterface.h"
 #include "Manifest.h"
 #include "UnitConversion.h"
+#include "GeneratedCodeVersion.h"
 
 double GPluginOverheadTime = 0.0;
 double GHeaderCodeGenTime = 0.0;
@@ -50,6 +51,7 @@ private:
 enum {MAX_ARRAY_SIZE=2048};
 
 static const FName NAME_ToolTip(TEXT("ToolTip"));
+EGeneratedCodeVersion FHeaderParser::DefaultGeneratedCodeVersion = EGeneratedCodeVersion::V1;
 TMap<UClass*, ClassDefinitionRange> ClassDefinitionRanges;
 /**
  * Dirty hack global variable to allow different result codes passed through
@@ -1798,14 +1800,7 @@ UScriptStruct* FHeaderParser::CompileStructDeclaration(FClasses& AllClasses, FUn
 			Struct->StructMacroDeclaredLineNumber = InputLine;
 			RequireSymbol(TEXT("("), TEXT("'struct'"));
 
-			FToken DuplicateStructName;
-			if (GetIdentifier(DuplicateStructName))
-			{
-				if (!DuplicateStructName.Matches(*StructNameInScript))
-				{
-					FError::Throwf(TEXT("The argument to %s must match the struct name '%s' if present.  However, the argument can be omitted entirely."), Token.Identifier, *StructNameInScript);
-				}
-			}
+			CompileVersionDeclaration(SourceFile, Struct);
 
 			RequireSymbol(TEXT(")"), TEXT("'struct'"));
 
@@ -3834,6 +3829,7 @@ bool FHeaderParser::CompileDeclaration(FClasses& AllClasses, FUnrealSourceFile&
 			FError::Throwf(TEXT("%s must occur inside the native interface definition"), Token.Identifier);
 		}
 		RequireSymbol(TEXT("("), Token.Identifier);
+		CompileVersionDeclaration(SourceFile, GetCurrentClass());
 		RequireSymbol(TEXT(")"), Token.Identifier);
 
 		auto* ClassData = GetCurrentClassData();
@@ -3847,6 +3843,11 @@ bool FHeaderParser::CompileDeclaration(FClasses& AllClasses, FUnrealSourceFile&
 		{
 			CurrentAccessSpecifier = ACCESS_Public;
 		}
+
+		if (Token.Matches(TEXT("GENERATED_BODY")))
+		{
+			ClassDefinitionRanges[GetCurrentClass()].bHasGeneratedBody = true;
+		}
 	}
 	else if (Token.Matches(TEXT("GENERATED_UINTERFACE_BODY")) || (Token.Matches(TEXT("GENERATED_BODY")) && TopNest->NestType == NEST_Interface))
 	{
@@ -3855,6 +3856,7 @@ bool FHeaderParser::CompileDeclaration(FClasses& AllClasses, FUnrealSourceFile&
 			FError::Throwf(TEXT("%s must occur inside the interface definition"), Token.Identifier);
 		}
 		RequireSymbol(TEXT("("), Token.Identifier);
+		CompileVersionDeclaration(SourceFile, GetCurrentClass());
 		RequireSymbol(TEXT(")"), Token.Identifier);
 
 		auto* ClassData = GetCurrentClassData();
@@ -3895,6 +3897,7 @@ bool FHeaderParser::CompileDeclaration(FClasses& AllClasses, FUnrealSourceFile&
 		}
 
 		RequireSymbol(TEXT("("), Token.Identifier);
+		CompileVersionDeclaration(SourceFile, GetCurrentClass());
 		RequireSymbol(TEXT(")"), Token.Identifier);
 
 		ClassData->SetGeneratedBodyLine(InputLine);
@@ -6727,6 +6730,12 @@ FHeaderParser::FHeaderParser(FFeedbackContext* InWarn)
 	DelegateParameterCountStrings.Add(TEXT("_SixParams"));
 	DelegateParameterCountStrings.Add(TEXT("_SevenParams"));
 	DelegateParameterCountStrings.Add(TEXT("_EightParams"));
+
+	FString Version;
+	if (GConfig->GetString(TEXT("GeneratedCodeVersion"), TEXT("UnrealHeaderTool"), Version, GEngineIni))
+	{
+		DefaultGeneratedCodeVersion = ToGeneratedCodeVersion(Version);
+	}
 }
 
 // Throws if a specifier value wasn't provided
@@ -6799,6 +6808,7 @@ ECompilationResult::Type FHeaderParser::ParseAllHeadersInside(
 	// Create the header parser and register it as the warning context.
 	// Note: This must be declared outside the try block, since the catch block will log into it.
 	FHeaderParser HeaderParser(Warn);
+	HeaderParser.CurrentlyParsedModule = &Module;
 	Warn->SetContext(&HeaderParser);
 
 	// Set up a filename for the error context if we don't even get as far parsing a class
@@ -7712,6 +7722,38 @@ void FHeaderParser::SkipDeprecatedMacroIfNecessary()
 	RequireSymbol(TEXT(")"), TEXT("DEPRECATED macro"));
 }
 
+void FHeaderParser::CompileVersionDeclaration(FUnrealSourceFile& SourceFile, UStruct* Struct)
+{
+	// Do nothing if we're at the end of file.
+	FToken Token;
+	if (!GetToken(Token, true, ESymbolParseOption::Normal))
+	{
+		return;
+	}
+
+	// Default version based on config file.
+	auto Version = DefaultGeneratedCodeVersion;
+
+	// Overwrite with module-specific value if one was specified.
+	if (CurrentlyParsedModule->GeneratedCodeVersion != EGeneratedCodeVersion::None)
+	{
+		Version = CurrentlyParsedModule->GeneratedCodeVersion;
+	}
+
+	if (Token.TokenType == ETokenType::TOKEN_Symbol
+		&& !FCString::Stricmp(Token.Identifier, TEXT(")")))
+	{
+		SourceFile.GetGeneratedCodeVersions().FindOrAdd(Struct) = Version;
+		UngetToken(Token);
+		return;
+	}
+
+	// Overwrite with version specified by macro.
+	Version = ToGeneratedCodeVersion(Token.Identifier);
+
+	SourceFile.GetGeneratedCodeVersions().FindOrAdd(Struct) = Version;
+}
+
 void FHeaderParser::ResetClassData()
 {
 	UClass* CurrentClass = GetCurrentClass();
diff --git a/Engine/Source/Programs/UnrealHeaderTool/Private/HeaderParser.h b/Engine/Source/Programs/UnrealHeaderTool/Private/HeaderParser.h
--- a/Engine/Source/Programs/UnrealHeaderTool/Private/HeaderParser.h
+++ b/Engine/Source/Programs/UnrealHeaderTool/Private/HeaderParser.h
@@ -6,9 +6,13 @@
 #include "BaseParser.h"
 #include "CompilationResult.h"
 #include "ScopedTimers.h"
+#include "Scope.h"
+#include "GeneratedCodeVersion.h"
 
 class FClass;
 class FClasses;
+class FScope;
+class FHeaderProvider;
 
 extern double GPluginOverheadTime;
 extern double GHeaderCodeGenTime;
@@ -185,6 +189,9 @@ private:
 class FHeaderParser : public FBaseParser, public FContextSupplier
 {
 public:
+	// Default version of generated code. Defaults to oldest possible, unless specified otherwise in config.
+	static EGeneratedCodeVersion DefaultGeneratedCodeVersion;
+
 	// Compute the function parameter size and save the return offset
 	static void ComputeFunctionParametersSize(UClass* InClass);
 
@@ -732,6 +739,9 @@ protected:
 	static void ValidatePropertyIsDeprecatedIfNecessary(FPropertyBase& VarProperty, FToken* OuterPropertyType);
 
 private:
+	// Module currently parsed by UHT.
+	const FManifestModule* CurrentlyParsedModule;
+
 	/**
 	 * Tries to match constructor parameter list. Assumes that constructor
 	 * name is already matched.
@@ -745,6 +755,8 @@ private:
 	bool TryToMatchConstructorParameterList(FToken Token);
 	void SkipDeprecatedMacroIfNecessary();
 
+	// Parses possible version declaration in generated code, e.g. GENERATED_BODY(<some_version>).
+	void CompileVersionDeclaration(FUnrealSourceFile& SourceFile, UStruct* Struct);
 };
 
 /////////////////////////////////////////////////////
diff --git a/Engine/Source/Programs/UnrealHeaderTool/Private/Manifest.cpp b/Engine/Source/Programs/UnrealHeaderTool/Private/Manifest.cpp
--- a/Engine/Source/Programs/UnrealHeaderTool/Private/Manifest.cpp
+++ b/Engine/Source/Programs/UnrealHeaderTool/Private/Manifest.cpp
@@ -5,7 +5,7 @@
 #include "Manifest.h"
 #include "Json.h"
 #include "Manifest.h"
-
+#include "GeneratedCodeVersion.h"
 
 namespace
 {
@@ -117,7 +117,7 @@ FManifest FManifest::LoadFromFile(const FString& Filename)
 		auto& KnownModule = Result.Modules.Last();
 
 		FString Outer = FString::Printf(TEXT("Modules[%d]"), ModuleIndex);
-
+		FString GeneratedCodeVersionString;
 		GetJsonFieldValue(KnownModule.Name,                      ModuleObj, TEXT("Name"),                     *Outer);
 		GetJsonFieldValue(KnownModule.BaseDirectory,             ModuleObj, TEXT("BaseDirectory"),            *Outer);
 		GetJsonFieldValue(KnownModule.IncludeBase,               ModuleObj, TEXT("IncludeBase"),              *Outer);
@@ -128,6 +128,8 @@ FManifest FManifest::LoadFromFile(const FString& Filename)
 		GetJsonFieldValue(PrivateHeaders,                        ModuleObj, TEXT("PrivateHeaders"),           *Outer);
 		GetJsonFieldValue(KnownModule.PCH,                       ModuleObj, TEXT("PCH"),                      *Outer);
 		GetJsonFieldValue(KnownModule.GeneratedCPPFilenameBase,  ModuleObj, TEXT("GeneratedCPPFilenameBase"), *Outer);
+		GetJsonFieldValue(GeneratedCodeVersionString, ModuleObj, TEXT("UHTGeneratedCodeVersion"), *Outer);
+		KnownModule.GeneratedCodeVersion = ToGeneratedCodeVersion(GeneratedCodeVersionString);
 
 		FString ModuleTypeText;
 		GetJsonFieldValue(ModuleTypeText, ModuleObj, TEXT("ModuleType"), *Outer);
diff --git a/Engine/Source/Programs/UnrealHeaderTool/Private/Manifest.h b/Engine/Source/Programs/UnrealHeaderTool/Private/Manifest.h
--- a/Engine/Source/Programs/UnrealHeaderTool/Private/Manifest.h
+++ b/Engine/Source/Programs/UnrealHeaderTool/Private/Manifest.h
@@ -1,6 +1,7 @@
 // Copyright 1998-2015 Epic Games, Inc. All Rights Reserved.
 
 #pragma once
+#include "UnrealSourceFile.h"
 
 
 struct FManifestModule
@@ -40,6 +41,9 @@ struct FManifestModule
 
 	/** Whether or not to write out headers that have changed */
 	bool SaveExportedHeaders;
+
+	/** Version of generated code. */
+	EGeneratedCodeVersion GeneratedCodeVersion;
 };
 
 
diff --git a/Engine/Source/Programs/UnrealHeaderTool/Private/NativeClassExporter.h b/Engine/Source/Programs/UnrealHeaderTool/Private/NativeClassExporter.h
--- a/Engine/Source/Programs/UnrealHeaderTool/Private/NativeClassExporter.h
+++ b/Engine/Source/Programs/UnrealHeaderTool/Private/NativeClassExporter.h
@@ -31,6 +31,7 @@ namespace EExportFunctionType
 class FClass;
 class FClasses;
 class FModuleClasses;
+class FScope;
 
 // These are declared in this way to allow swapping out the classes for something more optimized in the future
 typedef FStringOutputDevice FUHTStringBuilder;
@@ -75,7 +76,6 @@ private:
 	FUHTStringBuilder InClassNoPureDeclsMacroCalls;
 	FUHTStringBuilder StandardUObjectConstructorsMacroCall;
 	FUHTStringBuilder EnhancedUObjectConstructorsMacroCall;
-	FUHTStringBuilder StaticChecks;
 
 	/** Generated function implementations that belong in the cpp file, split into multiple files base on line count **/
 	TArray<FUHTStringBuilderLineCounter> GeneratedFunctionBodyTextSplit;
@@ -379,11 +379,11 @@ private:
 	* Exports checks if function exists
 	*
 	* @param	FunctionData			Data representing the function to export.
-	* @param	CheckClasses			Where to write the member check classes.
-	* @param	StaticChecks			Where to write the static asserts throwing errors when function doesn't exist.
 	* @param	ClassName				Name of currently parsed class.
+	* @param	bHasImplementation		Whether class has declared a version of FunctionData with _Implementation suffix.
+	* @param	bHasValidate			Whether class has declared a version of FunctionData with _Validate suffix.
 	*/
-	void ExportFunctionChecks(const FFuncInfo& FunctionData, FUHTStringBuilder& CheckClasses, FUHTStringBuilder& StaticChecks, const FString& ClassName);
+	void CheckRPCFunctions(const FFuncInfo& FunctionData, const FString& ClassName, bool bHasImplementation, bool bHasValidate, const FUnrealSourceFile& SourceFile, int32 GeneratedBodyLine);
 
 	/**
 	 * Exports the native stubs for the list of functions specified
diff --git a/Engine/Source/Programs/UnrealHeaderTool/Private/UnrealSourceFile.cpp b/Engine/Source/Programs/UnrealHeaderTool/Private/UnrealSourceFile.cpp
--- a/Engine/Source/Programs/UnrealHeaderTool/Private/UnrealSourceFile.cpp
+++ b/Engine/Source/Programs/UnrealHeaderTool/Private/UnrealSourceFile.cpp
@@ -3,6 +3,8 @@
 #include "UnrealHeaderTool.h"
 #include "UnrealSourceFile.h"
 #include "ParserHelper.h"
+#include "HeaderParser.h"
+#include "GeneratedCodeVersion.h"
 
 void FUnrealSourceFile::AddDefinedClass(UClass* Class)
 {
@@ -98,6 +100,17 @@ const FString& FUnrealSourceFile::GetContent() const
 	return Content;
 }
 
+EGeneratedCodeVersion FUnrealSourceFile::GetGeneratedCodeVersionForStruct(UStruct* Struct)
+{
+	auto Version = GeneratedCodeVersions.Find(Struct);
+	if (Version)
+	{
+		return *Version;
+	}
+
+	return FHeaderParser::DefaultGeneratedCodeVersion;
+}
+
 void FUnrealSourceFile::MarkDependenciesResolved()
 {
 	bDependenciesResolved = true;
diff --git a/Engine/Source/Programs/UnrealHeaderTool/Private/UnrealSourceFile.h b/Engine/Source/Programs/UnrealHeaderTool/Private/UnrealSourceFile.h
--- a/Engine/Source/Programs/UnrealHeaderTool/Private/UnrealSourceFile.h
+++ b/Engine/Source/Programs/UnrealHeaderTool/Private/UnrealSourceFile.h
@@ -5,6 +5,7 @@
 #include "Scope.h"
 #include "HeaderProvider.h"
 #include "UnrealTypeDefinitionInfo.h"
+#include "GeneratedCodeVersion.h"
 
 class FClassMetaData;
 
@@ -184,6 +185,27 @@ public:
 		return Includes;
 	}
 
+	/**
+	 * Gets generated code version for given UStruct.
+	 */
+	EGeneratedCodeVersion GetGeneratedCodeVersionForStruct(UStruct* Struct);
+
+	/**
+	 * Gets generated code versions.
+	 */
+	TMap<UStruct*, EGeneratedCodeVersion>& GetGeneratedCodeVersions()
+	{
+		return GeneratedCodeVersions;
+	}
+
+	/**
+	 * Gets generated code versions. Const version.
+	 */
+	const TMap<UStruct*, EGeneratedCodeVersion>& GetGeneratedCodeVersions() const
+	{
+		return GeneratedCodeVersions;
+	}
+
 	/**
 	 * Sets generated filename.
 	 */
@@ -265,4 +287,7 @@ private:
 
 	// List of classes defined in this source file.
 	TArray<UClass*> DefinedClasses;
+
+	// Mapping of UStructs to versions, according to which their code should be generated.
+	TMap<UStruct*, EGeneratedCodeVersion> GeneratedCodeVersions;
 };
\ No newline at end of file
diff --git a/Engine/Source/Runtime/Core/Public/Misc/Char.h b/Engine/Source/Runtime/Core/Public/Misc/Char.h
--- a/Engine/Source/Runtime/Core/Public/Misc/Char.h
+++ b/Engine/Source/Runtime/Core/Public/Misc/Char.h
@@ -106,7 +106,10 @@ public:
 	static inline bool IsDigit(CharType Char);
 	static inline bool IsHexDigit(CharType Char);
 	static inline bool IsWhitespace(CharType Char);
-
+	static inline bool IsIdentifier(CharType Char)
+	{
+		return IsAlnum(Char) || IsUnderscore(Char);
+	}
 	static inline bool IsUnderscore(CharType Char)		{ return Char == LITERAL(CharType, '_'); }
 
 public:
