Add missing file.

diff --git a/Engine/Source/Programs/UnrealHeaderTool/Private/GeneratedCodeVersion.h b/Engine/Source/Programs/UnrealHeaderTool/Private/GeneratedCodeVersion.h
new file mode 100644
--- /dev/null
+++ b/Engine/Source/Programs/UnrealHeaderTool/Private/GeneratedCodeVersion.h
@@ -0,0 +1,34 @@
+// Copyright 1998-2015 Epic Games, Inc. All Rights Reserved.
+#pragma once
+
+//
+// This MUST be kept in sync with EGeneratedBodyVersion in UBT defined in ExternalExecution.cs
+// and with ToGeneratedBodyVersion function below
+//
+enum class EGeneratedCodeVersion : int32
+{
+	None,
+	V1,
+	V2,
+	VLatest = V2
+};
+
+inline EGeneratedCodeVersion ToGeneratedCodeVersion(const FString& InString)
+{
+	if (InString.Compare(TEXT("V1")) == 0)
+	{
+		return EGeneratedCodeVersion::V1;
+	}
+
+	if (InString.Compare(TEXT("V2")) == 0)
+	{
+		return EGeneratedCodeVersion::V2;
+	}
+
+	if (InString.Compare(TEXT("VLatest")) == 0)
+	{
+		return EGeneratedCodeVersion::VLatest;
+	}
+
+	return EGeneratedCodeVersion::None;
+}
\ No newline at end of file
