New CPF_UObjectWrapper flag to indicate that a UObjectProperty is a wrapper or a raw pointer.
UClassProperty::GetCPPType made to use the flag to generate a UClass* rather than a TSubobjectOf<UObject>.
Test code for UHT.

diff --git a/Engine/Source/Programs/UnrealHeaderTool/Private/HeaderParser.cpp b/Engine/Source/Programs/UnrealHeaderTool/Private/HeaderParser.cpp
--- a/Engine/Source/Programs/UnrealHeaderTool/Private/HeaderParser.cpp
+++ b/Engine/Source/Programs/UnrealHeaderTool/Private/HeaderParser.cpp
@@ -3151,7 +3151,6 @@ void FHeaderParser::GetVarType
 		{
 			// An object reference of some type (maybe a restricted class?)
 			UClass* TempClass = NULL;
-			bool bExpectStar = true;
 
 			const bool bIsLazyPtrTemplate        = VarType.Matches(TEXT("TLazyObjectPtr"));
 			const bool bIsAssetPtrTemplate       = VarType.Matches(TEXT("TAssetPtr"));
@@ -3173,7 +3172,7 @@ void FHeaderParser::GetVarType
 			else if (VarType.Matches(TEXT("FScriptInterface")))
 			{
 				TempClass = UInterface::StaticClass();
-				bExpectStar = false;
+				Flags |= CPF_UObjectWrapper;
 			}
 			else if (bIsAssetClassTemplate)
 			{
@@ -3222,7 +3221,7 @@ void FHeaderParser::GetVarType
 						Flags |= CPF_SubobjectReference | CPF_InstancedReference;
 					}
 
-					bExpectStar = false;
+					Flags |= CPF_UObjectWrapper;
 				}
 				else
 				{
@@ -3246,7 +3245,7 @@ void FHeaderParser::GetVarType
 				{
 					if ( MatchSymbol(TEXT("<")) )
 					{
-						bExpectStar = false;
+						Flags |= CPF_UObjectWrapper;
 
 						// Consume a forward class declaration 'class' if present
 						MatchIdentifier(TEXT("class"));
@@ -3290,7 +3289,7 @@ void FHeaderParser::GetVarType
 				}
 
 				// Eat the star that indicates this is a pointer to the UObject
-				if (bExpectStar)
+				if (!(Flags & CPF_UObjectWrapper))
 				{
 					// Const after variable type but before pointer symbol
 					MatchIdentifier(TEXT("const"));
@@ -5006,7 +5005,7 @@ void FHeaderParser::ParseParameterList(FClasses& AllClasses, UFunction* Function
 		// Get parameter type.
 		FToken Property(CPT_None);
 		EVariableCategory::Type VariableCategory = (Function->FunctionFlags & FUNC_Net) ? EVariableCategory::ReplicatedParameter : EVariableCategory::RegularParameter;
-		GetVarType(AllClasses, GetCurrentScope(), Property, ~(CPF_ParmFlags | CPF_AutoWeak | CPF_RepSkip), NULL, EPropertyDeclarationStyle::None, VariableCategory);
+		GetVarType(AllClasses, GetCurrentScope(), Property, ~(CPF_ParmFlags | CPF_AutoWeak | CPF_RepSkip | CPF_UObjectWrapper), NULL, EPropertyDeclarationStyle::None, VariableCategory);
 		Property.PropertyFlags |= CPF_Parm;
 
 		if (bExpectCommaBeforeName)
diff --git a/Engine/Source/Programs/UnrealHeaderTool/Resources/UHTDebugging/TestObject.h b/Engine/Source/Programs/UnrealHeaderTool/Resources/UHTDebugging/TestObject.h
--- a/Engine/Source/Programs/UnrealHeaderTool/Resources/UHTDebugging/TestObject.h
+++ b/Engine/Source/Programs/UnrealHeaderTool/Resources/UHTDebugging/TestObject.h
@@ -38,6 +38,9 @@ public:
 	}
 #endif
 
+	UFUNCTION(BlueprintNativeEvent, Category="Game")
+	UClass* BrokenReturnTypeForFunction();
+
 #if 0
 	UPROPERTY()
 	ITestInterface* ShouldGiveAMeaningfulErrorAboutTScriptInterface;
diff --git a/Engine/Source/Runtime/CoreUObject/Private/UObject/PropertyClass.cpp b/Engine/Source/Runtime/CoreUObject/Private/UObject/PropertyClass.cpp
--- a/Engine/Source/Runtime/CoreUObject/Private/UObject/PropertyClass.cpp
+++ b/Engine/Source/Runtime/CoreUObject/Private/UObject/PropertyClass.cpp
@@ -91,7 +91,14 @@ const TCHAR* UClassProperty::ImportText_Internal( const TCHAR* Buffer, void* Dat
 
 FString UClassProperty::GetCPPType( FString* ExtendedTypeText/*=NULL*/, uint32 CPPExportFlags/*=0*/ ) const
 {
-	return FString::Printf(TEXT("TSubclassOf<%s%s> "),MetaClass->GetPrefixCPP(),*MetaClass->GetName());
+	if (PropertyFlags & CPF_UObjectWrapper)
+	{
+		return FString::Printf(TEXT("TSubclassOf<%s%s> "),MetaClass->GetPrefixCPP(),*MetaClass->GetName());
+	}
+	else
+	{
+		return TEXT("UClass*");
+	}
 }
 
 FString UClassProperty::GetCPPTypeForwardDeclaration() const
diff --git a/Engine/Source/Runtime/CoreUObject/Public/UObject/ObjectBase.h b/Engine/Source/Runtime/CoreUObject/Public/UObject/ObjectBase.h
--- a/Engine/Source/Runtime/CoreUObject/Public/UObject/ObjectBase.h
+++ b/Engine/Source/Runtime/CoreUObject/Public/UObject/ObjectBase.h
@@ -359,6 +359,7 @@ typedef uint64 EClassCastFlags;
 #define CPF_NonPIEDuplicateTransient		DECLARE_UINT64(0x0000800000000000)		// Property should only be copied in PIE
 #define CPF_ExposeOnSpawn					DECLARE_UINT64(0x0001000000000000)		// Property is exposed on spawn
 #define CPF_PersistentInstance				DECLARE_UINT64(0x0002000000000000)		// A object referenced by the property is duplicated like a component. (Each actor should have an own instance.)
+#define CPF_UObjectWrapper					DECLARE_UINT64(0x0004000000000000)		// Property was parsed as a wrapper class like TSubobjectOf<T>, FScriptInterface etc., rather than a USomething*
 
 #define CPF_NonPIETransient \
 	EMIT_DEPRECATED_WARNING_MESSAGE("CPF_NonPIETransient is deprecated. Please use CPF_NonPIEDuplicateTransient instead.") \
@@ -367,7 +368,7 @@ typedef uint64 EClassCastFlags;
 /** @name Combinations flags */
 //@{
 #define CPF_ParmFlags				(CPF_Parm | CPF_OutParm | CPF_ReturnParm | CPF_ReferenceParm | CPF_ConstParm)
-#define CPF_PropagateToArrayInner	(CPF_ExportObject | CPF_PersistentInstance | CPF_InstancedReference | CPF_ContainsInstancedReference | CPF_Config | CPF_EditConst | CPF_Deprecated | CPF_EditorOnly | CPF_AutoWeak )
+#define CPF_PropagateToArrayInner	(CPF_ExportObject | CPF_PersistentInstance | CPF_InstancedReference | CPF_ContainsInstancedReference | CPF_Config | CPF_EditConst | CPF_Deprecated | CPF_EditorOnly | CPF_AutoWeak | CPF_UObjectWrapper )
 
 /** the flags that should never be set on interface properties */
 #define CPF_InterfaceClearMask		(CPF_ExportObject|CPF_InstancedReference|CPF_ContainsInstancedReference)
