UnrealHeaderTool: Added limited support for parsing character literals to avoid incorrectly erroring on things like '\"'

diff --git a/Engine/Source/Programs/UnrealHeaderTool/Private/BaseParser.cpp b/Engine/Source/Programs/UnrealHeaderTool/Private/BaseParser.cpp
--- a/Engine/Source/Programs/UnrealHeaderTool/Private/BaseParser.cpp
+++ b/Engine/Source/Programs/UnrealHeaderTool/Private/BaseParser.cpp
@@ -407,47 +407,48 @@ bool FBaseParser::GetToken( FToken& Token, bool bNoConsts/*=false*/, ESymbolPars
 		}
 		return true;
 	}
-//@TODO: 'z' is a character literal in C++, not a FName like it was in UnrealScript - do we need this code? 
-// 	else if( !bNoConsts && c=='\'' )
-// 	{
-// 		// Name constant.
-// 		int32 Length=0;
-// 		c = GetChar();
-// 		while( (c>='A' && c<='Z') || (c>='a' && c<='z') || (c>='0' && c<='9') || (c=='_') || (c=='-') || (c==' ') ) //@FIXME: space in names should be illegal!
-// 		{
-// 			Token.Identifier[Length++] = c;
-// 			if( Length >= NAME_SIZE )
-// 			{
-// 				FError::Throwf(TEXT("Name length exceeds maximum of %i"), (int32)NAME_SIZE );
-// 				// trick the error a few lines down
-// 				c = TEXT('\'');
-// 				Length = ((int32)NAME_SIZE) - 1;
-// 				break;
-// 			}
-// 			c = GetChar();
-// 		}
-// 		if( c != '\'' )
-// 		{
-// 			UngetChar();
-// 			FError::Throwf(TEXT("Illegal character in name") );
-// 		}
-// 		Token.Identifier[Length]=0;
-// 
-// 		// Make constant name.
-// 		Token.SetConstName( FName(Token.Identifier) );
-// 		return true;
-// 	}
-	else if( c=='"' )
+	else if (c == '\'')
+	{
+		TCHAR ActualCharLiteral = GetChar(/*bLiteral=*/ true);
+
+		if (ActualCharLiteral == '\\')
+		{
+			ActualCharLiteral = GetChar(/*bLiteral=*/ true);
+			switch (ActualCharLiteral)
+			{
+			case TCHAR('t'):
+				ActualCharLiteral = '\t';
+				break;
+			case TCHAR('n'):
+				ActualCharLiteral = '\n';
+				break;
+			case TCHAR('r'):
+				ActualCharLiteral = '\r';
+				break;
+			}
+		}
+
+		c = GetChar(/*bLiteral=*/ true);
+		if (c != '\'')
+		{
+			FError::Throwf(TEXT("Unterminated character constant"));
+			UngetChar();
+		}
+
+		Token.SetConstChar(ActualCharLiteral);
+		return true;
+	}
+	else if (c == '"')
 	{
 		// String constant.
 		TCHAR Temp[MAX_STRING_CONST_SIZE];
 		int32 Length=0;
-		c = GetChar(1);
+		c = GetChar(/*bLiteral=*/ true);
 		while( (c!='"') && !IsEOL(c) )
 		{
 			if( c=='\\' )
 			{
-				c = GetChar(1);
+				c = GetChar(/*bLiteral=*/ true);
 				if( IsEOL(c) )
 				{
 					break;
@@ -466,7 +467,7 @@ bool FBaseParser::GetToken( FToken& Token, bool bNoConsts/*=false*/, ESymbolPars
 				Length = ((int32)MAX_STRING_CONST_SIZE) - 1;
 				break;
 			}
-			c = GetChar(1);
+			c = GetChar(/*bLiteral=*/ true);
 		}
 		Temp[Length]=0;
 
diff --git a/Engine/Source/Programs/UnrealHeaderTool/Private/ParserHelper.h b/Engine/Source/Programs/UnrealHeaderTool/Private/ParserHelper.h
--- a/Engine/Source/Programs/UnrealHeaderTool/Private/ParserHelper.h
+++ b/Engine/Source/Programs/UnrealHeaderTool/Private/ParserHelper.h
@@ -792,6 +792,14 @@ public:
 		}
 		TokenType = TOKEN_Const;
 	}
+	void SetConstChar(TCHAR InChar)
+	{
+		//@TODO: Treating this like a string for now, nothing consumes it
+		(FPropertyBase&)*this = FPropertyBase(CPT_String);
+		String[0] = InChar;
+		String[1] = 0;
+		TokenType = TOKEN_Const;
+	}
 	//!!struct constants
 
 	// Getters.
