Add UHTGeneratedCodeVersion to UHTDebugging.manifest.

diff --git a/Engine/Source/Programs/UnrealHeaderTool/Resources/UHTDebugging.manifest b/Engine/Source/Programs/UnrealHeaderTool/Resources/UHTDebugging.manifest
--- a/Engine/Source/Programs/UnrealHeaderTool/Resources/UHTDebugging.manifest
+++ b/Engine/Source/Programs/UnrealHeaderTool/Resources/UHTDebugging.manifest
@@ -18,7 +18,8 @@
 			],
 			"PCH": "",
 			"GeneratedCPPFilenameBase": "..\\..\\..\\..\\Intermediate/Build/Inc\\CoreUObject\\CoreUObject.generated",
-			"SaveExportedHeaders": true
+			"SaveExportedHeaders": true,
+			"UHTGeneratedCodeVersion":"None"
 		}, {
 			"Name": "UHTDebugging",
 			"ModuleType": "Game",
@@ -39,7 +40,8 @@
 			],
 			"PCH": "",
 			"GeneratedCPPFilenameBase": "..\\..\\..\\..\\Intermediate/Build/Inc\\UHTDebugging\\UHTDebugging.generated",
-			"SaveExportedHeaders": true
+			"SaveExportedHeaders": true,
+			"UHTGeneratedCodeVersion":"None"
 		}
 	]
 }
