UHT Fix

diff --git a/Engine/Source/Programs/UnrealHeaderTool/Private/CodeGenerator.cpp b/Engine/Source/Programs/UnrealHeaderTool/Private/CodeGenerator.cpp
--- a/Engine/Source/Programs/UnrealHeaderTool/Private/CodeGenerator.cpp
+++ b/Engine/Source/Programs/UnrealHeaderTool/Private/CodeGenerator.cpp
@@ -1037,7 +1037,7 @@ void FNativeClassHeaderGenerator::ExportGeneratedPackageInitCode(UPackage* Packa
 		uint32 CombinedCRC = 0;
 		for (auto& Split : GeneratedFunctionBodyTextSplit)
 		{
-			uint32 SplitCRC = GenerateTextCRC(*Split.ToUpper());
+			uint32 SplitCRC = GenerateTextCRC(*Split->ToUpper());
 			if (CombinedCRC == 0)
 			{
 				// Don't combine in the first case because it keeps GUID backwards compatibility
@@ -4328,12 +4328,12 @@ FUHTStringBuilder& FNativeClassHeaderGenerator::GetGeneratedFunctionTextDevice()
 {
 	int32 MaxLinesPerCpp = 30000;
 
-	if ((GeneratedFunctionBodyTextSplit.Num() == 0) || (GeneratedFunctionBodyTextSplit[GeneratedFunctionBodyTextSplit.Num() - 1].GetLineCount() > MaxLinesPerCpp))
+	if ((GeneratedFunctionBodyTextSplit.Num() == 0) || (GeneratedFunctionBodyTextSplit[GeneratedFunctionBodyTextSplit.Num() - 1]->GetLineCount() > MaxLinesPerCpp))
 	{
-		new (GeneratedFunctionBodyTextSplit) FUHTStringBuilderLineCounter;
+		GeneratedFunctionBodyTextSplit.Add( TUniqueObj<FUHTStringBuilderLineCounter>() );
 	}
 
-	return GeneratedFunctionBodyTextSplit[GeneratedFunctionBodyTextSplit.Num() - 1];
+	return GeneratedFunctionBodyTextSplit[GeneratedFunctionBodyTextSplit.Num() - 1].Get();
 }
 
 // Constructor.
@@ -4860,7 +4860,7 @@ void FNativeClassHeaderGenerator::ExportGeneratedCPP()
 				FileText.Logf(TEXT("\r\n"));
 			}
 			FileText.Log(GeneratedFunctionDeclarations);
-			FileText.Log(GeneratedFunctionBodyTextSplit[FileIdx]);
+			FileText.Log(GeneratedFunctionBodyTextSplit[FileIdx].Get());
 			FileText.Logf(TEXT("#endif\r\n"));
 		}
 
diff --git a/Engine/Source/Programs/UnrealHeaderTool/Private/NativeClassExporter.h b/Engine/Source/Programs/UnrealHeaderTool/Private/NativeClassExporter.h
--- a/Engine/Source/Programs/UnrealHeaderTool/Private/NativeClassExporter.h
+++ b/Engine/Source/Programs/UnrealHeaderTool/Private/NativeClassExporter.h
@@ -78,7 +78,7 @@ private:
 	FUHTStringBuilder EnhancedUObjectConstructorsMacroCall;
 
 	/** Generated function implementations that belong in the cpp file, split into multiple files base on line count **/
-	TArray<FUHTStringBuilderLineCounter> GeneratedFunctionBodyTextSplit;
+	TArray<TUniqueObj<FUHTStringBuilderLineCounter>> GeneratedFunctionBodyTextSplit;
 	/** Declarations of generated functions for this module**/
 	FUHTStringBuilder GeneratedFunctionDeclarations;
 	/** Declarations of generated functions for cross module**/
