UnrealHeaderTool: Added support for FColor default values in reflected methods

diff --git a/Engine/Source/Programs/UnrealHeaderTool/Private/HeaderParser.cpp b/Engine/Source/Programs/UnrealHeaderTool/Private/HeaderParser.cpp
--- a/Engine/Source/Programs/UnrealHeaderTool/Private/HeaderParser.cpp
+++ b/Engine/Source/Programs/UnrealHeaderTool/Private/HeaderParser.cpp
@@ -7437,6 +7437,7 @@ bool FHeaderParser::DefaultValueStringCppFormatToInnerFormat(const UProperty* Pr
 		static const UScriptStruct* Vector2DStruct = FindObjectChecked<UScriptStruct>(CoreUObjectPackage, TEXT("Vector2D"));
 		static const UScriptStruct* RotatorStruct = FindObjectChecked<UScriptStruct>(CoreUObjectPackage, TEXT("Rotator"));
 		static const UScriptStruct* LinearColorStruct = FindObjectChecked<UScriptStruct>(CoreUObjectPackage, TEXT("LinearColor"));
+		static const UScriptStruct* ColorStruct = FindObjectChecked<UScriptStruct>(CoreUObjectPackage, TEXT("Color"));
 
 		const UStructProperty* StructProperty = CastChecked<UStructProperty>(Property);
 		if( StructProperty->Struct == VectorStruct )
@@ -7563,6 +7564,57 @@ bool FHeaderParser::DefaultValueStringCppFormatToInnerFormat(const UProperty* Pr
 				}
 			}
 		}
+		else if( StructProperty->Struct == ColorStruct )
+		{
+			if( FDefaultValueHelper::Is( CppForm, TEXT("FColor::White") ) )
+			{
+				OutForm = FColor::White.ToString();
+			}
+			else if ( FDefaultValueHelper::Is( CppForm, TEXT("FColor::Black") ) )
+			{
+				OutForm = FColor::Black.ToString();
+			}
+			else if ( FDefaultValueHelper::Is( CppForm, TEXT("FColor::Red") ) )
+			{
+				OutForm = FColor::Red.ToString();
+			}
+			else if ( FDefaultValueHelper::Is( CppForm, TEXT("FColor::Green") ) )
+			{
+				OutForm = FColor::Green.ToString();
+			}
+			else if ( FDefaultValueHelper::Is( CppForm, TEXT("FColor::Blue") ) )
+			{
+				OutForm = FColor::Blue.ToString();
+			}
+			else if (FDefaultValueHelper::Is(CppForm, TEXT("FColor::Yellow")))
+			{
+				OutForm = FColor::Yellow.ToString();
+			}
+			else if ( FDefaultValueHelper::Is( CppForm, TEXT("FColor::Cyan") ) )
+			{
+				OutForm = FColor::Cyan.ToString();
+			}
+			else if ( FDefaultValueHelper::Is( CppForm, TEXT("FColor::Magenta") ) )
+			{
+				OutForm = FColor::Magenta.ToString();
+			}
+			else
+			{
+				FString Parameters;
+				if( FDefaultValueHelper::GetParameters(CppForm, TEXT("FColor"), Parameters) )
+				{
+					if( FDefaultValueHelper::Is(Parameters, TEXT("ForceInit")) )
+					{
+						return true;
+					}
+					FColor Color;
+					if( FDefaultValueHelper::ParseColor(Parameters, Color) )
+					{
+						OutForm = Color.ToString();
+					}
+				}
+			}
+		}
 	}
 
 	return !OutForm.IsEmpty();
diff --git a/Engine/Source/Runtime/Core/Private/Misc/DefaultValueHelper.cpp b/Engine/Source/Runtime/Core/Private/Misc/DefaultValueHelper.cpp
--- a/Engine/Source/Runtime/Core/Private/Misc/DefaultValueHelper.cpp
+++ b/Engine/Source/Runtime/Core/Private/Misc/DefaultValueHelper.cpp
@@ -234,6 +234,13 @@ bool FDefaultValueHelper::IsStringValidInteger(const TCHAR* Start, const TCHAR*
 }
 
 
+bool FDefaultValueHelper::IsStringValidInteger(const TCHAR* Start, const TCHAR* End)
+{
+	int32 Base;
+	return IsStringValidInteger(Start, End, /*out*/ Base);
+}
+	
+	
 bool FDefaultValueHelper::IsStringValidInteger(const FString& Source)
 {
 	if(!Source.IsEmpty())
@@ -664,3 +671,50 @@ bool FDefaultValueHelper::ParseLinearColor(const FString& Source, FLinearColor&
 	
 	return true;
 }
+
+bool FDefaultValueHelper::ParseColor(const FString& Source, FColor& OutVal)
+{
+	const bool bHasWhitespace = HasWhitespaces(Source);
+	const FString NoWhitespace = bHasWhitespace ? RemoveWhitespaces(Source) : FString();
+	const FString& ProperSource = bHasWhitespace ? NoWhitespace : Source;
+	if(ProperSource.IsEmpty())
+	{
+		return false;
+	}
+
+	const TCHAR* Start = StartOf(ProperSource);
+	const TCHAR* FirstComma = FCString::Strstr( Start, TEXT(",") );
+	if(!FirstComma)
+	{
+		return false;
+	}
+
+	const TCHAR* SecondComma = FCString::Strstr( FirstComma + 1,  TEXT(",") );
+	if(!SecondComma)
+	{
+		return false;
+	}
+
+	const TCHAR* ThirdComma = FCString::Strstr( SecondComma + 1,  TEXT(",") );
+	const TCHAR* End = EndOf(ProperSource);
+	if( ( NULL != ThirdComma ) && !IsStringValidInteger( ThirdComma + 1, End ) )
+	{
+		return false;
+	}
+	
+	if(	!IsStringValidInteger( Start, FirstComma ) ||
+		!IsStringValidInteger( FirstComma + 1, SecondComma ) ||
+		!IsStringValidInteger( SecondComma + 1, ( NULL != ThirdComma ) ? ThirdComma : End ) )
+	{
+		return false;
+	}
+
+	const float Alpha = ( NULL != ThirdComma ) ? FCString::Atoi( ThirdComma + 1 ) : 1.0f;
+	OutVal = FColor( 
+		FCString::Atoi( Start ), 
+		FCString::Atoi( FirstComma + 1 ), 
+		FCString::Atoi( SecondComma + 1 ),
+		Alpha );
+	
+	return true;
+}
\ No newline at end of file
diff --git a/Engine/Source/Runtime/Core/Public/Misc/DefaultValueHelper.h b/Engine/Source/Runtime/Core/Public/Misc/DefaultValueHelper.h
--- a/Engine/Source/Runtime/Core/Public/Misc/DefaultValueHelper.h
+++ b/Engine/Source/Runtime/Core/Public/Misc/DefaultValueHelper.h
@@ -40,6 +40,7 @@ class CORE_API FDefaultValueHelper
 
 	/** return if the string can be parse to an int. OutBase is the base of the integer (8, 10, 16)*/
 	static bool IsStringValidInteger(const TCHAR* Start, const TCHAR* End, int32& OutBase);
+	static bool IsStringValidInteger(const TCHAR* Start, const TCHAR* End);
 
 public:
 
@@ -149,4 +150,15 @@ public:
 	 * @return true if the conversion happened, false otherwise
 	 */
 	static bool ParseLinearColor(const FString& Source, FLinearColor& OutVal);
+
+	/**
+	 * Converts a string into a FLinearColor. 
+	 * Accepted forms: " %d, %d, %d " or " %d, %d, %d, %d " (alpha is optional).
+	 *
+	 * @param Source the input string to try to convert
+	 * @param OutVal the output color
+	 *
+	 * @return true if the conversion happened, false otherwise
+	 */
+	static bool ParseColor(const FString& Source, FColor& OutVal);
 };
