UHT enforces user defined _Validate and _Implementation functions to be "virtual".

diff --git a/Engine/Source/Programs/UnrealHeaderTool/Private/CodeGenerator.cpp b/Engine/Source/Programs/UnrealHeaderTool/Private/CodeGenerator.cpp
--- a/Engine/Source/Programs/UnrealHeaderTool/Private/CodeGenerator.cpp
+++ b/Engine/Source/Programs/UnrealHeaderTool/Private/CodeGenerator.cpp
@@ -45,36 +45,111 @@ const TCHAR* GetTabs(int32 NumTabs)
 	return TabString + TabPos;
 }
 
-bool HasIdentifierExactMatch(const FString &ClassDefinition, const FString& Identifier)
+/**
+ * Finds exact match of Identifier in string. Returns nullptr if none is found.
+ *
+ * @param StringBegin Start of string to search.
+ * @param StringEnd End of string to search.
+ * @param Identifier Identifier to find.
+ * @return Pointer to Identifier match within string. nullptr if none found.
+ */
+const TCHAR* FindIdentifierExactMatch(const TCHAR* StringBegin, const TCHAR* StringEnd, const FString& Identifier)
 {
-	int32 SearchStartPosition = INDEX_NONE;
+	int32 StringLen = StringEnd - StringBegin;
+
+	// Check for exact match first.
+	if (FCString::Strncmp(StringBegin, *Identifier, StringLen) == 0)
+	{
+		return StringBegin;
+	}
+
+	auto FindLen = Identifier.Len();
+	auto StringToSearch = StringBegin;
 
 	for (;;)
 	{
-		auto IdentifierStartIndex = ClassDefinition.Find(Identifier, ESearchCase::CaseSensitive, ESearchDir::FromStart, SearchStartPosition);
+		auto IdentifierStart = FCString::Strstr(StringToSearch, *Identifier);
+		if (IdentifierStart == nullptr)
+		{
+			// Not found.
+			return nullptr;
+		}
 
-		// Class definition not found
-		if (IdentifierStartIndex == INDEX_NONE
-			// ClassDefinition can't start with CppImplName. If it starts, we'll get compiler error anyway
-			|| IdentifierStartIndex == 0)
+		if ((IdentifierStart > StringEnd) || (IdentifierStart + FindLen + 1 > StringEnd))
 		{
-			return false;
+			// Found match is out of string range.
+			return nullptr;
 		}
 
-		// To find exact match, previous and next character can't be parts of identifier.
-		if (!FChar::IsIdentifier(ClassDefinition[IdentifierStartIndex - 1])
-			&& !FChar::IsIdentifier(ClassDefinition[IdentifierStartIndex + Identifier.Len()]))
+		if ((IdentifierStart == StringBegin) && (!FChar::IsIdentifier(*(IdentifierStart + FindLen + 1))))
 		{
-			return true;
+			// Found match is at the beginning of string.
+			return IdentifierStart;
 		}
 
-		// Exact match not found, continue search starting just after the end of identifier
-		SearchStartPosition = IdentifierStartIndex + Identifier.Len() + 1;
+		if ((IdentifierStart + FindLen == StringEnd) && (!FChar::IsIdentifier(*(IdentifierStart - 1))))
+		{
+			// Found match ends with end of string.
+			return IdentifierStart;
+		}
+
+		if ((!FChar::IsIdentifier(*(IdentifierStart + FindLen)))
+			&& (!FChar::IsIdentifier(*(IdentifierStart - 1))))
+		{
+			// Found match is in the middle of string
+			return IdentifierStart;
+		}
+
+		// Didn't find exact match, nor got to end of search string. Keep on searching.
+		StringToSearch = IdentifierStart;
 	}
 
 	// We should never get here.
 	checkNoEntry();
-	return false;
+	return nullptr;
+}
+
+/**
+ * Finds exact match of Identifier in string. Returns nullptr if none is found.
+ *
+ * @param String String to search.
+ * @param Identifier Identifier to find.
+ * @return Index to Identifier match within String. INDEX_NONE if none found.
+ */
+int32 FindIdentifierExactMatch(const FString& String, const FString& Identifier)
+{
+	auto IdentifierPtr = FindIdentifierExactMatch(*String, *String + String.Len(), Identifier);
+	if (IdentifierPtr == nullptr)
+	{
+		return INDEX_NONE;
+	}
+
+	return IdentifierPtr - *String;
+}
+
+/**
+* Checks if exact match of Identifier is in String.
+*
+* @param StringBegin Start of string to search.
+* @param StringEnd End of string to search.
+* @param Identifier Identifier to find.
+* @return true if Identifier is within string, false otherwise.
+*/
+bool HasIdentifierExactMatch(const TCHAR* StringBegin, const TCHAR* StringEnd, const FString& Find)
+{
+	return FindIdentifierExactMatch(StringBegin, StringEnd, Find) != nullptr;
+}
+
+/**
+* Checks if exact match of Identifier is in String.
+*
+* @param String String to search.
+* @param Identifier Identifier to find.
+* @return true if Identifier is within String, false otherwise.
+*/
+bool HasIdentifierExactMatch(const FString &String, const FString& Identifier)
+{
+	return FindIdentifierExactMatch(String, Identifier) != INDEX_NONE;
 }
 
 /////////////////////////////////////////////////////
@@ -3536,8 +3611,56 @@ FString GetFunctionConstModifierString(UFunction* Function)
 	return FString();
 }
 
-void FNativeClassHeaderGenerator::CheckRPCFunctions(const FFuncInfo& FunctionData, const FString& ClassName, bool bHasImplementation, bool bHasValidate, const FUnrealSourceFile& SourceFile, int32 GeneratedBodyLine)
+/**
+ * Converts Position within File to Line and Column.
+ *
+ * @param File File contents.
+ * @param Position Position in string to convert.
+ * @param OutLine Result line.
+ * @param OutColumn Result column.
+ */
+void GetLineAndColumnFromPositionInFile(const FString& File, int32 Position, int32& OutLine, int32& OutColumn)
 {
+	OutLine = 1;
+	OutColumn = 1;
+
+	int32 i;
+	for (i = 1; i <= Position; ++i)
+	{
+		if (File[i] == '\n')
+		{
+			++OutLine;
+			OutColumn = 0;
+		}
+		else
+		{
+			++OutColumn;
+		}
+	}
+}
+
+bool FNativeClassHeaderGenerator::IsMissingVirtualSpecifier(const FString& SourceFile, int32 FunctionNamePosition) const
+{
+	auto IsEndOfSearchChar = [](TCHAR C) { return (C == TEXT('}')) || (C == TEXT('{')) || (C == TEXT(';')); };
+
+	// Find first occurrence of "}", ";", "{" going backwards from ImplementationPosition.
+	int32 EndOfSearchCharIndex = SourceFile.FindLastCharByPredicate(IsEndOfSearchChar, FunctionNamePosition);
+	check(EndOfSearchCharIndex != INDEX_NONE);
+
+	// Then find if there is "virtual" keyword starting from position of found character to ImplementationPosition
+	return !HasIdentifierExactMatch(&SourceFile[EndOfSearchCharIndex], &SourceFile[FunctionNamePosition], TEXT("virtual"));
+}
+
+FString CreateClickableErrorMessage(const FString& Filename, int32 Line, int32 Column)
+{
+	return FString::Printf(TEXT("%s(%d,%d): error: "), *Filename, Line, Column);
+}
+
+void FNativeClassHeaderGenerator::CheckRPCFunctions(const FFuncInfo& FunctionData, const FString& ClassName, int32 ImplementationPosition, int32 ValidatePosition, const FUnrealSourceFile& SourceFile)
+{
+	bool bHasImplementation = ImplementationPosition != INDEX_NONE;
+	bool bHasValidate = ValidatePosition != INDEX_NONE;
+
 	auto Function = FunctionData.FunctionReference;
 	auto FunctionReturnType = GetFunctionReturnString(Function);
 	auto ConstModifier = GetFunctionConstModifierString(Function) + TEXT(" ");
@@ -3554,6 +3677,8 @@ void FNativeClassHeaderGenerator::CheckRPCFunctions(const FFuncInfo& FunctionDat
 	check(bNeedsImplementation || bNeedsValidate);
 
 	auto ParameterString = GetFunctionParameterString(Function);
+	const auto& Filename = SourceFile.GetFilename();
+	const auto& FileContent = SourceFile.GetContent();
 
 	//
 	// Get string with function specifiers, listing why we need _Implementation or _Validate functions.
@@ -3580,17 +3705,42 @@ void FNativeClassHeaderGenerator::CheckRPCFunctions(const FFuncInfo& FunctionDat
 	AssertMessage.Logf(TEXT("."));
 
 	//
-	// Add proper checks if necessary.
+	// Check if functions are missing.
 	//
+	int32 Line;
+	int32 Column;
+	GetLineAndColumnFromPositionInFile(FileContent, FunctionData.InputPos, Line, Column);
 	if (bNeedsImplementation && !bHasImplementation)
 	{
-		// FullFilePath(LineNumber): error:
-		FError::Throwf(TEXT("%s(%d): error: %s Declare function %s %s(%s) %sin class %s."), *SourceFile.GetFilename(), GeneratedBodyLine, *AssertMessage, *FunctionReturnType, *FunctionData.CppImplName, *ParameterString, *ConstModifier, *ClassName);
+		FString ErrorPosition = CreateClickableErrorMessage(Filename, Line, Column);
+		FString FunctionDecl = FString::Printf(TEXT("virtual %s %s::%s(%s) %s"), *FunctionReturnType, *ClassName, *FunctionData.CppImplName, *ParameterString, *ConstModifier);
+		FError::Throwf(TEXT("%s%s Declare function %s"), *ErrorPosition, *AssertMessage, *FunctionDecl);
 	}
 	
 	if (bNeedsValidate && !bHasValidate)
 	{
-		FError::Throwf(TEXT("%s(%d): error: %s Declare function bool %s(%s) %sin class %s."), *SourceFile.GetFilename(), GeneratedBodyLine, *AssertMessage, *FunctionData.CppValidationImplName, *ParameterString, *ConstModifier, *ClassName);
+		FString ErrorPosition = CreateClickableErrorMessage(Filename, Line, Column);
+		FString FunctionDecl = FString::Printf(TEXT("virtual bool %s::%s(%s) %s"), *ClassName, *FunctionData.CppValidationImplName, *ParameterString, *ConstModifier);
+		FError::Throwf(TEXT("%s%s Declare function %s"), *ErrorPosition, *AssertMessage, *FunctionDecl);
+	}
+
+	//
+	// If all needed functions are declared, check if they have virtual specifiers.
+	//
+	if (bNeedsImplementation && bHasImplementation && IsMissingVirtualSpecifier(FileContent, ImplementationPosition, Line, Column))
+	{
+		GetLineAndColumnFromPositionInFile(FileContent, ImplementationPosition, Line, Column);
+		FString ErrorPosition = CreateClickableErrorMessage(Filename, Line, Column);
+		FString FunctionDecl = FString::Printf(TEXT("%s %s::%s(%s) %s"), *FunctionReturnType, *ClassName, *FunctionData.CppImplName, *ParameterString, *ConstModifier);
+		FError::Throwf(TEXT("%sDeclared function %sis not marked as virtual."), *ErrorPosition, *FunctionDecl);
+	}
+
+	if (bNeedsValidate && bHasValidate && IsMissingVirtualSpecifier(FileContent, ValidatePosition, Line, Column))
+	{
+		GetLineAndColumnFromPositionInFile(FileContent, ValidatePosition, Line, Column);
+		FString ErrorPosition = CreateClickableErrorMessage(Filename, Line, Column);
+		FString FunctionDecl = FString::Printf(TEXT("bool %s::%s(%s) %s"), *ClassName, *FunctionData.CppValidationImplName, *ParameterString, *ConstModifier);
+		FError::Throwf(TEXT("%sDeclared function %sis not marked as virtual."), *ErrorPosition, *FunctionDecl);
 	}
 }
 
@@ -3986,9 +4136,20 @@ void FNativeClassHeaderGenerator::ExportNativeFunctions(FUnrealSourceFile& Sourc
 			auto ClassDefinition = FString(ClassEnd - ClassStart, ClassStart);
 
 			auto FunctionName = Function->GetName();
+			int32 ClassDefinitionStartPosition = ClassStart - *SourceFile.GetContent();
 
-			bool bHasImplementation = HasIdentifierExactMatch(ClassDefinition, FunctionData.CppImplName);
-			bool bHasValidate = HasIdentifierExactMatch(ClassDefinition, FunctionData.CppValidationImplName);
+			int32 ImplementationPosition = FindIdentifierExactMatch(ClassDefinition, FunctionData.CppImplName);
+			if (ImplementationPosition != INDEX_NONE)
+			{
+				ImplementationPosition += ClassDefinitionStartPosition;
+			}
+			int32 ValidatePosition = FindIdentifierExactMatch(ClassDefinition, FunctionData.CppValidationImplName);
+			if (ValidatePosition != INDEX_NONE)
+			{
+				ValidatePosition += ClassDefinitionStartPosition;
+			}
+			bool bHasImplementation = ImplementationPosition != INDEX_NONE;
+			bool bHasValidate = ValidatePosition != INDEX_NONE;
 
 			//Emit warning here if necessary
 			FUHTStringBuilder FunctionDeclaration;
@@ -4019,9 +4180,9 @@ void FNativeClassHeaderGenerator::ExportNativeFunctions(FUnrealSourceFile& Sourc
 			// Versions that skip function autodeclaration throw an error when a function is missing.
 			if ((SourceFile.GetGeneratedCodeVersionForStruct(Class) > EGeneratedCodeVersion::V1) && ClassRange.bHasGeneratedBody)
 			{
-				CheckRPCFunctions(FunctionData, Class->HasAnyClassFlags(CLASS_Interface) ? *(FString(TEXT("I")) + ClassName) : NameLookupCPP.GetNameCPP(Class), bHasImplementation, bHasValidate, SourceFile, Class->HasAnyClassFlags(CLASS_Interface) ? ClassData->GetInterfaceGeneratedBodyLine() : ClassData->GetGeneratedBodyLine());
+				auto Name = Class->HasAnyClassFlags(CLASS_Interface) ? *(FString(TEXT("I")) + ClassName) : NameLookupCPP.GetNameCPP(Class);
+				CheckRPCFunctions(FunctionData, Name, ImplementationPosition, ValidatePosition, SourceFile);
 			}
-			
 		}
 
 		if (bMultiLineUFUNCTION)
diff --git a/Engine/Source/Programs/UnrealHeaderTool/Private/HeaderParser.cpp b/Engine/Source/Programs/UnrealHeaderTool/Private/HeaderParser.cpp
--- a/Engine/Source/Programs/UnrealHeaderTool/Private/HeaderParser.cpp
+++ b/Engine/Source/Programs/UnrealHeaderTool/Private/HeaderParser.cpp
@@ -5198,6 +5198,11 @@ void FHeaderParser::CompileDelegateDeclaration(FUnrealSourceFile& SourceFile, FC
 		RequireSymbol(TEXT(","), CurrentScopeName);
 	}
 
+	// Skip whitespaces to get InputPos exactly on beginning of function name.
+	while (FChar::IsWhitespace(PeekChar())) { GetChar(); }
+
+	FuncInfo.InputPos = InputPos;
+
 	// Get the delegate name
 	if (!GetIdentifier(FuncInfo.Function))
 	{
@@ -5527,6 +5532,11 @@ void FHeaderParser::CompileFunctionDeclaration(FUnrealSourceFile& SourceFile, FC
 		GetVarType(AllClasses, GetCurrentScope(), ReturnType, 0, NULL, EPropertyDeclarationStyle::None, EVariableCategory::Return);
 	}
 
+	// Skip whitespaces to get InputPos exactly on beginning of function name.
+	while (FChar::IsWhitespace(PeekChar())) { GetChar(); }
+
+	FuncInfo.InputPos = InputPos;
+
 	// Get function or operator name.
 	if (!GetIdentifier(FuncInfo.Function))
 	{
diff --git a/Engine/Source/Programs/UnrealHeaderTool/Private/NativeClassExporter.h b/Engine/Source/Programs/UnrealHeaderTool/Private/NativeClassExporter.h
--- a/Engine/Source/Programs/UnrealHeaderTool/Private/NativeClassExporter.h
+++ b/Engine/Source/Programs/UnrealHeaderTool/Private/NativeClassExporter.h
@@ -376,14 +376,15 @@ private:
 	void ExportNativeFunctionHeader(const FFuncInfo& FunctionData, FUHTStringBuilder& HeaderOutput, EExportFunctionType::Type FunctionType, EExportFunctionHeaderStyle::Type FunctionHeaderStyle, const TCHAR* ExtraParam = NULL);
 
 	/**
-	* Exports checks if function exists
+	* Runs checks whether necessary RPC functions exist for function described by FunctionData.
 	*
 	* @param	FunctionData			Data representing the function to export.
 	* @param	ClassName				Name of currently parsed class.
-	* @param	bHasImplementation		Whether class has declared a version of FunctionData with _Implementation suffix.
-	* @param	bHasValidate			Whether class has declared a version of FunctionData with _Validate suffix.
+	* @param	ImplementationPosition	Position in source file of _Implementation function for function described by FunctionData.
+	* @param	ValidatePosition		Position in source file of _Validate function for function described by FunctionData.
+	* @param	SourceFile				Currently analyzed source file.
 	*/
-	void CheckRPCFunctions(const FFuncInfo& FunctionData, const FString& ClassName, bool bHasImplementation, bool bHasValidate, const FUnrealSourceFile& SourceFile, int32 GeneratedBodyLine);
+	void CheckRPCFunctions(const FFuncInfo& FunctionData, const FString& ClassName, int32 ImplementationPosition, int32 ValidatePosition, const FUnrealSourceFile& SourceFile);
 
 	/**
 	 * Exports the native stubs for the list of functions specified
@@ -550,4 +551,13 @@ public:
 	* @return FString with function parameters.
 	*/
 	FString GetFunctionParameterString(UFunction* Function);
+
+	/**
+	 * Checks if function is missing "virtual" specifier.
+	 * 
+	 * @param SourceFile SourceFile where function is declared.
+	 * @param FunctionNamePosition Position of name of function in SourceFile.
+	 * @return true if function misses "virtual" specifier, false otherwise.
+	 */
+	bool IsMissingVirtualSpecifier(const FString& SourceFile, int32 FunctionNamePosition) const;
 };
diff --git a/Engine/Source/Programs/UnrealHeaderTool/Private/ParserHelper.h b/Engine/Source/Programs/UnrealHeaderTool/Private/ParserHelper.h
--- a/Engine/Source/Programs/UnrealHeaderTool/Private/ParserHelper.h
+++ b/Engine/Source/Programs/UnrealHeaderTool/Private/ParserHelper.h
@@ -887,6 +887,8 @@ struct FFuncInfo
 	bool		bSealedEvent;
 	/** Delegate macro line in header. */
 	int32		MacroLine;
+	/** Position in file where this function was declared. Points to first char of function name. */
+	int32 InputPos;
 	//@}
 
 	/** Constructor. */
@@ -902,6 +904,7 @@ struct FFuncInfo
 		, RPCResponseId(0)
 		, bSealedEvent(false)
 		, MacroLine(-1)
+		, InputPos(-1)
 	{}
 
 	FFuncInfo(const FFuncInfo& Other)
@@ -914,6 +917,7 @@ struct FFuncInfo
 		, RPCId(Other.RPCId)
 		, RPCResponseId(Other.RPCResponseId)
 		, MacroLine(Other.MacroLine)
+		, InputPos(Other.InputPos)
 	{
 		Function.Clone(Other.Function);
 		if (FunctionReference)
diff --git a/Engine/Source/Runtime/Core/Public/Containers/Array.h b/Engine/Source/Runtime/Core/Public/Containers/Array.h
--- a/Engine/Source/Runtime/Core/Public/Containers/Array.h
+++ b/Engine/Source/Runtime/Core/Public/Containers/Array.h
@@ -909,6 +909,29 @@ public:
 		return INDEX_NONE;
 	}
 
+	/**
+	 * Finds element within the array starting from the end. Uses predicate to match element.
+	 *
+	 * @param Pred Predicate taking array element and returns true if element matches search criteria, false otherwise.
+	 * @param StartIndex Index of element from which to start searching. Defaults to index of last element in array.
+	 *
+	 * @returns Index of the found element. INDEX_NONE otherwise.
+	 */
+	template <typename Predicate>
+	int32 FindLastByPredicate(Predicate Pred, int32 StartIndex = ArrayNum) const
+	{
+		const ElementType* RESTRICT End = GetData() + StartIndex;
+		for (const ElementType* RESTRICT Data = End, *RESTRICT DataStart = Data - StartIndex; Data != DataStart;)
+		{
+			--Data;
+			if (Pred(*Data))
+			{
+				return static_cast<int32>(Data - DataStart);
+			}
+		}
+		return INDEX_NONE;
+	}
+
 	/**
 	 * Finds element within the array that fulfills given predicate.
 	 *
diff --git a/Engine/Source/Runtime/Core/Public/Containers/UnrealString.h b/Engine/Source/Runtime/Core/Public/Containers/UnrealString.h
--- a/Engine/Source/Runtime/Core/Public/Containers/UnrealString.h
+++ b/Engine/Source/Runtime/Core/Public/Containers/UnrealString.h
@@ -1048,6 +1048,20 @@ public:
 		return Data.FindLast(InChar, Index);
 	}
 
+	/**
+	* Searches the string for the last occurrence of a character
+	*
+	* @param Pred Predicate that takes TCHAR and returns true if TCHAR matches search criteria, false otherwise.
+	* @param StartIndex Index of element from which to start searching. Defaults to last TCHAR in string.
+	*
+	* @return Index of found TCHAR, INDEX_NONE otherwise.
+	*/
+	template <typename Predicate>
+	FORCEINLINE int32 FindLastCharByPredicate(Predicate Pred, int32 StartIndex = Data.Num() - 1) const
+	{
+		return Data.FindLastByPredicate(Pred, StartIndex);
+	}
+
 	/**
 	 * Lexicographically tests whether this string is equivalent to the Other given string
 	 * 
