Add missing file.

diff --git a/Engine/Source/Programs/UnrealHeaderTool/Private/CodeGenerator.cpp b/Engine/Source/Programs/UnrealHeaderTool/Private/CodeGenerator.cpp
--- a/Engine/Source/Programs/UnrealHeaderTool/Private/CodeGenerator.cpp
+++ b/Engine/Source/Programs/UnrealHeaderTool/Private/CodeGenerator.cpp
@@ -3727,7 +3727,7 @@ void FNativeClassHeaderGenerator::CheckRPCFunctions(const FFuncInfo& FunctionDat
 	//
 	// If all needed functions are declared, check if they have virtual specifiers.
 	//
-	if (bNeedsImplementation && bHasImplementation && IsMissingVirtualSpecifier(FileContent, ImplementationPosition, Line, Column))
+	if (bNeedsImplementation && bHasImplementation && IsMissingVirtualSpecifier(FileContent, ImplementationPosition))
 	{
 		GetLineAndColumnFromPositionInFile(FileContent, ImplementationPosition, Line, Column);
 		FString ErrorPosition = CreateClickableErrorMessage(Filename, Line, Column);
@@ -3735,7 +3735,7 @@ void FNativeClassHeaderGenerator::CheckRPCFunctions(const FFuncInfo& FunctionDat
 		FError::Throwf(TEXT("%sDeclared function %sis not marked as virtual."), *ErrorPosition, *FunctionDecl);
 	}
 
-	if (bNeedsValidate && bHasValidate && IsMissingVirtualSpecifier(FileContent, ValidatePosition, Line, Column))
+	if (bNeedsValidate && bHasValidate && IsMissingVirtualSpecifier(FileContent, ValidatePosition))
 	{
 		GetLineAndColumnFromPositionInFile(FileContent, ValidatePosition, Line, Column);
 		FString ErrorPosition = CreateClickableErrorMessage(Filename, Line, Column);
