Fixing UHT not propagating class flags/properties properly

- Small refactor of how and where class flags are parsed
- 'Placeable' is no longer deprecated (it never really was)
- Restoring 'Placeable' keyword on ALandscape

UE-14035

diff --git a/Engine/Source/Programs/UnrealHeaderTool/Private/ClassDeclarationMetaData.cpp b/Engine/Source/Programs/UnrealHeaderTool/Private/ClassDeclarationMetaData.cpp
new file mode 100644
--- /dev/null
+++ b/Engine/Source/Programs/UnrealHeaderTool/Private/ClassDeclarationMetaData.cpp
@@ -0,0 +1,512 @@
+// Copyright 1998-2015 Epic Games, Inc. All Rights Reserved.
+
+#include "UnrealHeaderTool.h"
+#include "ClassDeclarationMetaData.h"
+#include "ParserClass.h"
+#include "HeaderParser.h"
+#include "Classes.h"
+
+// Utility functions
+namespace
+{
+	bool IsActorClass(UClass *Class)
+	{
+		static FName ActorName = FName(TEXT("Actor"));
+		while (Class)
+		{
+			if (Class->GetFName() == ActorName)
+			{
+				return true;
+			}
+			Class = Class->GetSuperClass();
+		}
+		return false;
+	}
+}
+
+FClassDeclarationMetaData::FClassDeclarationMetaData()
+	: ClassFlags(0)
+	, WantsToBePlaceable(false)
+{
+}
+
+void FClassDeclarationMetaData::ParseClassProperties(const TArray<FPropertySpecifier>& InClassSpecifiers, const FString& InRequiredAPIMacroIfPresent)
+{
+	ClassFlags = 0;
+	// Record that this class is RequiredAPI if the CORE_API style macro was present
+	if (!InRequiredAPIMacroIfPresent.IsEmpty())
+	{
+		ClassFlags |= CLASS_RequiredAPI;
+	}
+	ClassFlags |= CLASS_Native;
+
+	// Process all of the class specifiers
+
+	bool bWithinSpecified = false;
+	bool bDeclaresConfigFile = false;
+	for (const FPropertySpecifier& PropSpecifier : InClassSpecifiers)
+	{
+		const FString& Specifier = PropSpecifier.Key;
+
+		if (Specifier == TEXT("noexport"))
+		{
+			// Don't export to C++ header.
+			ClassFlags |= CLASS_NoExport;
+		}
+		else if (Specifier == TEXT("intrinsic"))
+		{
+			ClassFlags |= CLASS_Intrinsic;
+		}
+		else if (Specifier == TEXT("ComponentWrapperClass"))
+		{
+			MetaData.Add(TEXT("IgnoreCategoryKeywordsInSubclasses"), TEXT("true"));
+		}
+		else if (Specifier == TEXT("within"))
+		{
+			ClassWithin = FHeaderParser::RequireExactlyOneSpecifierValue(PropSpecifier);
+		}
+		else if (Specifier == TEXT("editinlinenew"))
+		{
+			// Class can be constructed from the New button in editinline
+			ClassFlags |= CLASS_EditInlineNew;
+		}
+		else if (Specifier == TEXT("noteditinlinenew"))
+		{
+			// Class cannot be constructed from the New button in editinline
+			ClassFlags &= ~CLASS_EditInlineNew;
+		}
+		else if (Specifier == TEXT("placeable"))
+		{
+			WantsToBePlaceable = true;
+			ClassFlags &= ~CLASS_NotPlaceable;
+		}
+		else if (Specifier == TEXT("defaulttoinstanced"))
+		{
+			// these classed default to instanced.
+			ClassFlags |= CLASS_DefaultToInstanced;
+		}
+		else if (Specifier == TEXT("notplaceable"))
+		{
+			// Don't allow the class to be placed in the editor.
+			ClassFlags |= CLASS_NotPlaceable;
+		}
+		else if (Specifier == TEXT("hidedropdown"))
+		{
+			// Prevents class from appearing in class comboboxes in the property window
+			ClassFlags |= CLASS_HideDropDown;
+		}
+		else if (Specifier == TEXT("dependsOn"))
+		{
+			FError::Throwf(TEXT("The dependsOn specifier is deprecated. Please use #include \"ClassHeaderFilename.h\" instead."));
+		}
+		else if (Specifier == TEXT("MinimalAPI"))
+		{
+			ClassFlags |= CLASS_MinimalAPI;
+		}
+		else if (Specifier == TEXT("const"))
+		{
+			ClassFlags |= CLASS_Const;
+		}
+		else if (Specifier == TEXT("perObjectConfig"))
+		{
+			ClassFlags |= CLASS_PerObjectConfig;
+		}
+		else if (Specifier == TEXT("configdonotcheckdefaults"))
+		{
+			ClassFlags |= CLASS_ConfigDoNotCheckDefaults;
+		}
+		else if (Specifier == TEXT("abstract"))
+		{
+			// Hide all editable properties.
+			ClassFlags |= CLASS_Abstract;
+		}
+		else if (Specifier == TEXT("deprecated"))
+		{
+			ClassFlags |= CLASS_Deprecated;
+
+			// Don't allow the class to be placed in the editor.
+			ClassFlags |= CLASS_NotPlaceable;
+		}
+		else if (Specifier == TEXT("transient"))
+		{
+			// Transient class.
+			ClassFlags |= CLASS_Transient;
+		}
+		else if (Specifier == TEXT("nonTransient"))
+		{
+			// this child of a transient class is not transient - remove the transient flag
+			ClassFlags &= ~CLASS_Transient;
+		}
+		else if (Specifier == TEXT("customConstructor"))
+		{
+			// we will not export a constructor for this class, assuming it is in the CPP block
+			ClassFlags |= CLASS_CustomConstructor;
+		}
+		else if (Specifier == TEXT("config"))
+		{
+			// Class containing config properties - parse the name of the config file to use
+			ConfigName = FHeaderParser::RequireExactlyOneSpecifierValue(PropSpecifier);
+		}
+		else if (Specifier == TEXT("defaultconfig"))
+		{
+			// Save object config only to Default INIs, never to local INIs.
+			ClassFlags |= CLASS_DefaultConfig;
+		}
+		else if (Specifier == TEXT("globaluserconfig"))
+		{
+			// Save object config only to global user overrides, never to local INIs
+			ClassFlags |= CLASS_GlobalUserConfig;
+		}
+		else if (Specifier == TEXT("showCategories"))
+		{
+			FHeaderParser::RequireSpecifierValue(PropSpecifier);
+
+			for (const FString& Value : PropSpecifier.Values)
+			{
+				// if we didn't find this specific category path in the HideCategories metadata
+				if (HideCategories.Remove(Value) == 0)
+				{
+					TArray<FString> SubCategoryList;
+					Value.ParseIntoArray(SubCategoryList, TEXT("|"), true);
+
+					FString SubCategoryPath;
+					// look to see if any of the parent paths are excluded in the HideCategories list
+					for (int32 CategoryPathIndex = 0; CategoryPathIndex < SubCategoryList.Num() - 1; ++CategoryPathIndex)
+					{
+						SubCategoryPath += SubCategoryList[CategoryPathIndex];
+						// if we're hiding a parent category, then we need to flag this sub category for show
+						if (HideCategories.Contains(SubCategoryPath))
+						{
+							ShowSubCatgories.AddUnique(Value);
+							break;
+						}
+						SubCategoryPath += "|";
+					}
+				}
+			}
+		}
+		else if (Specifier == TEXT("hideCategories"))
+		{
+			FHeaderParser::RequireSpecifierValue(PropSpecifier);
+
+			for (const FString& Value : PropSpecifier.Values)
+			{
+				HideCategories.AddUnique(Value);
+			}
+		}
+		else if (Specifier == TEXT("showFunctions"))
+		{
+			FHeaderParser::RequireSpecifierValue(PropSpecifier);
+
+			for (const FString& Value : PropSpecifier.Values)
+			{
+				HideFunctions.Remove(Value);
+			}
+		}
+		else if (Specifier == TEXT("hideFunctions"))
+		{
+			FHeaderParser::RequireSpecifierValue(PropSpecifier);
+
+			for (const FString& Value : PropSpecifier.Values)
+			{
+				HideFunctions.AddUnique(Value);
+			}
+		}
+		else if (Specifier == TEXT("classGroup"))
+		{
+			FHeaderParser::RequireSpecifierValue(PropSpecifier);
+
+			for (const FString& Value : PropSpecifier.Values)
+			{
+				ClassGroupNames.Add(Value);
+			}
+		}
+		else if (Specifier == TEXT("autoExpandCategories"))
+		{
+			FHeaderParser::RequireSpecifierValue(PropSpecifier);
+
+			for (const FString& Value : PropSpecifier.Values)
+			{
+				AutoCollapseCategories.Remove(Value);
+				AutoExpandCategories.AddUnique(Value);
+			}
+		}
+		else if (Specifier == TEXT("autoCollapseCategories"))
+		{
+			FHeaderParser::RequireSpecifierValue(PropSpecifier);
+
+			for (const FString& Value : PropSpecifier.Values)
+			{
+				AutoExpandCategories.Remove(Value);
+				AutoCollapseCategories.AddUnique(Value);
+			}
+		}
+		else if (Specifier == TEXT("dontAutoCollapseCategories"))
+		{
+			FHeaderParser::RequireSpecifierValue(PropSpecifier);
+
+			for (const FString& Value : PropSpecifier.Values)
+			{
+				AutoCollapseCategories.Remove(Value);
+			}
+		}
+		else if (Specifier == TEXT("collapseCategories"))
+		{
+			// Class' properties should not be shown categorized in the editor.
+			ClassFlags |= CLASS_CollapseCategories;
+		}
+		else if (Specifier == TEXT("dontCollapseCategories"))
+		{
+			// Class' properties should be shown categorized in the editor.
+			ClassFlags &= ~CLASS_CollapseCategories;
+		}
+		else if (Specifier == TEXT("AdvancedClassDisplay"))
+		{
+			// By default the class properties are shown in advanced sections in UI
+			ClassFlags |= CLASS_AdvancedDisplay;
+		}
+		else if (Specifier == TEXT("ConversionRoot"))
+		{
+			MetaData.Add(FName(TEXT("IsConversionRoot")), "true");
+		}
+		else
+		{
+			FError::Throwf(TEXT("Unknown class specifier '%s'"), *Specifier);
+		}
+	}
+}
+
+void FClassDeclarationMetaData::MergeShowCategories()
+{
+	for (const FString& Value : ShowCategories)
+	{
+		// if we didn't find this specific category path in the HideCategories metadata
+		if (HideCategories.Remove(Value) == 0)
+		{
+			TArray<FString> SubCategoryList;
+			Value.ParseIntoArray(SubCategoryList, TEXT("|"), true);
+
+			FString SubCategoryPath;
+			// look to see if any of the parent paths are excluded in the HideCategories list
+			for (int32 CategoryPathIndex = 0; CategoryPathIndex < SubCategoryList.Num() - 1; ++CategoryPathIndex)
+			{
+				SubCategoryPath += SubCategoryList[CategoryPathIndex];
+				// if we're hiding a parent category, then we need to flag this sub category for show
+				if (HideCategories.Contains(SubCategoryPath))
+				{
+					ShowSubCatgories.AddUnique(Value);
+					break;
+				}
+				SubCategoryPath += "|";
+			}
+		}
+	}
+	// Once the categories have been merged, empty the array as we will no longer need it nor should we use it
+	ShowCategories.Empty();
+}
+
+void FClassDeclarationMetaData::MergeClassCategories(FClass* Class)
+{
+	static FName DEBUG_Name(TEXT("StaticMeshComponent"));
+	if (Class->GetFName() == DEBUG_Name)
+	{
+		static volatile int32 xx = 0;
+		xx++;
+	}
+	TArray<FString> ParentHideCategories;
+	TArray<FString> ParentShowSubCatgories;
+	TArray<FString> ParentHideFunctions;
+	TArray<FString> ParentAutoExpandCategories;
+	TArray<FString> ParentAutoCollapseCategories;
+	Class->GetHideCategories(ParentHideCategories);
+	Class->GetShowCategories(ParentShowSubCatgories);
+	Class->GetHideFunctions(ParentHideFunctions);
+	Class->GetAutoExpandCategories(ParentAutoExpandCategories);
+	Class->GetAutoCollapseCategories(ParentAutoCollapseCategories);
+
+	// Add parent categories. We store the opposite of HideCategories and HideFunctions in a separate array anyway.
+	HideCategories.Append(ParentHideCategories);
+	ShowSubCatgories.Append(ParentShowSubCatgories);
+	HideFunctions.Append(ParentHideFunctions);
+
+	MergeShowCategories();
+
+	// Merge ShowFunctions and HideFunctions
+	for (const FString& Value : ShowFunctions)
+	{
+		HideFunctions.Remove(Value);
+	}
+	ShowFunctions.Empty();
+
+	// Merge DontAutoCollapseCategories and AutoCollapseCategories
+	for (const FString& Value : DontAutoCollapseCategories)
+	{
+		AutoCollapseCategories.Remove(Value);
+	}
+	DontAutoCollapseCategories.Empty();
+
+	// Merge ShowFunctions and HideFunctions
+	for (const FString& Value : ShowFunctions)
+	{
+		HideFunctions.Remove(Value);
+	}
+	ShowFunctions.Empty();
+
+	// Merge AutoExpandCategories and AutoCollapseCategories (we still want to keep AutoExpandCategories though!)
+	for (const FString& Value : AutoExpandCategories)
+	{
+		AutoCollapseCategories.Remove(Value);
+		ParentAutoCollapseCategories.Remove(Value);
+	}
+	
+	// Do the same as above but the other way around
+	for (const FString& Value : AutoCollapseCategories)
+	{
+		AutoExpandCategories.Remove(Value);
+		ParentAutoExpandCategories.Remove(Value);
+	}	
+
+	// Once AutoExpandCategories and AutoCollapseCategories for THIS class have been parsed, add the parent inherited categories
+	AutoCollapseCategories.Append(ParentAutoCollapseCategories);
+	AutoExpandCategories.Append(ParentAutoExpandCategories);
+}
+
+void FClassDeclarationMetaData::MergeAndValidateClassFlags(const FString& DeclaredClassName, uint32 PreviousClassFlags, FClass* Class, const FClasses& AllClasses)
+{
+	if (WantsToBePlaceable)
+	{
+		if (!(Class->ClassFlags & CLASS_NotPlaceable))
+		{
+			FError::Throwf(TEXT("The 'placeable' specifier is only allowed on classes which have a base class that's marked as not placeable. Classes are assumed to be placeable by default."));
+		}
+		Class->ClassFlags &= ~CLASS_NotPlaceable;
+		WantsToBePlaceable = false; // Reset this flag after it's been merged
+	}
+	
+	// Now merge all remaining flags/properties
+	Class->ClassFlags |= ClassFlags;
+	Class->ClassConfigName = FName(*ConfigName);
+
+	SetAndValidateWithinClass(Class, AllClasses);
+	SetAndValidateConfigName(Class);
+
+	if (!!(Class->ClassFlags & CLASS_EditInlineNew))
+	{
+		// don't allow actor classes to be declared editinlinenew
+		if (IsActorClass(Class))
+		{
+			FError::Throwf(TEXT("Invalid class attribute: Creating actor instances via the property window is not allowed"));
+		}
+	}
+
+	// Make sure both RequiredAPI and MinimalAPI aren't specified
+	if (Class->HasAllClassFlags(CLASS_MinimalAPI | CLASS_RequiredAPI))
+	{
+		FError::Throwf(TEXT("MinimalAPI cannot be specified when the class is fully exported using a MODULENAME_API macro"));
+	}
+
+	// All classes must start with a valid Unreal prefix
+	const FString ExpectedClassName = Class->GetNameWithPrefix();
+	if (DeclaredClassName != ExpectedClassName)
+	{
+		FError::Throwf(TEXT("Class name '%s' is invalid, should be identified as '%s'"), *DeclaredClassName, *ExpectedClassName);
+	}
+
+	if ((Class->ClassFlags&CLASS_NoExport))
+	{
+		// if the class's class flags didn't contain CLASS_NoExport before it was parsed, it means either:
+		// a) the DECLARE_CLASS macro for this native class doesn't contain the CLASS_NoExport flag (this is an error)
+		// b) this is a new native class, which isn't yet hooked up to static registration (this is OK)
+		if (!(Class->ClassFlags&CLASS_Intrinsic) && (PreviousClassFlags & CLASS_NoExport) == 0 &&
+			(PreviousClassFlags&CLASS_Native) != 0)	// a new native class (one that hasn't been compiled into C++ yet) won't have this set
+		{
+			FError::Throwf(TEXT("'noexport': Must include CLASS_NoExport in native class declaration"));
+		}
+	}
+
+	if (!Class->HasAnyClassFlags(CLASS_Abstract) && ((PreviousClassFlags & CLASS_Abstract) != 0))
+	{
+		if (Class->HasAnyClassFlags(CLASS_NoExport))
+		{
+			FError::Throwf(TEXT("'abstract': NoExport class missing abstract keyword from class declaration (must change C++ version first)"));
+			Class->ClassFlags |= CLASS_Abstract;
+		}
+		else if (Class->HasAnyFlags(RF_Native))
+		{
+			FError::Throwf(TEXT("'abstract': missing abstract keyword from class declaration - class will no longer be exported as abstract"));
+		}
+	}
+}
+
+void FClassDeclarationMetaData::SetAndValidateConfigName(FClass* Class)
+{
+	if (ConfigName.IsEmpty() == false)
+	{
+		// if the user specified "inherit", we're just going to use the parent class's config filename
+		// this is not actually necessary but it can be useful for explicitly communicating config-ness
+		if (ConfigName == TEXT("inherit"))
+		{
+			UClass* SuperClass = Class->GetSuperClass();
+			if (!SuperClass)
+			{
+				FError::Throwf(TEXT("Cannot inherit config filename: %s has no super class"), *Class->GetName());
+			}
+
+			if (SuperClass->ClassConfigName == NAME_None)
+			{
+				FError::Throwf(TEXT("Cannot inherit config filename: parent class %s is not marked config."), *SuperClass->GetPathName());
+			}
+		}
+		else
+		{
+			// otherwise, set the config name to the parsed identifier
+			Class->ClassConfigName = FName(*ConfigName);
+		}
+	}
+	else
+	{
+		// Invalidate config name if not specifically declared.
+		Class->ClassConfigName = NAME_None;
+	}
+}
+
+void FClassDeclarationMetaData::SetAndValidateWithinClass(FClass* Class, const FClasses& AllClasses)
+{
+	// Process all of the class specifiers
+	if (ClassWithin.IsEmpty() == false)
+	{
+		UClass* RequiredWithinClass = AllClasses.FindClass(*ClassWithin);
+		if (!RequiredWithinClass)
+		{
+			FError::Throwf(TEXT("Within class '%s' not found."), *ClassWithin);
+		}
+		if (RequiredWithinClass->IsChildOf(UInterface::StaticClass()))
+		{
+			FError::Throwf(TEXT("Classes cannot be 'within' interfaces"));
+		}
+		else if (Class->ClassWithin == NULL || Class->ClassWithin == UObject::StaticClass() || RequiredWithinClass->IsChildOf(Class->ClassWithin))
+		{
+			Class->ClassWithin = RequiredWithinClass;
+		}
+		else if (Class->ClassWithin != RequiredWithinClass)
+		{
+			FError::Throwf(TEXT("%s must be within %s, not %s"), *Class->GetPathName(), *Class->ClassWithin->GetPathName(), *RequiredWithinClass->GetPathName());
+		}
+	}
+	else
+	{
+		// Make sure there is a valid within
+		Class->ClassWithin = Class->GetSuperClass()
+			? Class->GetSuperClass()->ClassWithin
+			: UObject::StaticClass();
+	}
+
+	UClass* ExpectedWithin = Class->GetSuperClass()
+		? Class->GetSuperClass()->ClassWithin
+		: UObject::StaticClass();
+
+	if (!Class->ClassWithin->IsChildOf(ExpectedWithin))
+	{
+		FError::Throwf(TEXT("Parent class declared within '%s'.  Cannot override within class with '%s' since it isn't a child"), *ExpectedWithin->GetName(), *Class->ClassWithin->GetName());
+	}
+}
\ No newline at end of file
diff --git a/Engine/Source/Programs/UnrealHeaderTool/Private/ClassDeclarationMetaData.h b/Engine/Source/Programs/UnrealHeaderTool/Private/ClassDeclarationMetaData.h
new file mode 100644
--- /dev/null
+++ b/Engine/Source/Programs/UnrealHeaderTool/Private/ClassDeclarationMetaData.h
@@ -0,0 +1,62 @@
+// Copyright 1998-2015 Epic Games, Inc. All Rights Reserved.
+
+#pragma once
+
+/** Structure that holds class meta data generated from its UCLASS declaration */
+class FClassDeclarationMetaData
+{
+public:
+	FClassDeclarationMetaData();
+
+	uint64 ClassFlags;
+	TMap<FName, FString> MetaData;
+	FString ClassWithin;
+	FString ConfigName;
+	
+	TArray<FString> HideCategories;
+	TArray<FString> ShowSubCatgories;
+	TArray<FString> HideFunctions;
+	TArray<FString> AutoExpandCategories;
+	TArray<FString> AutoCollapseCategories;
+	TArray<FString> DependsOn;
+	TArray<FString> ClassGroupNames;
+	
+	/**
+	* Parse Class's properties to generate its declaration data.
+	*
+	* @param	InClassSpecifiers Class properties collected from its UCLASS macro
+	* @param	InRequiredAPIMacroIfPresent *_API macro if present (empty otherwise)
+	* @param	OutClassData Parsed class meta data
+	*/
+	void ParseClassProperties(const TArray<FPropertySpecifier>& InClassSpecifiers, const FString& InRequiredAPIMacroIfPresent);
+
+	/**
+	* Merges all category properties with the class which at this point only has its parent propagated categories
+	*
+	* @param	Class Class to merge categories for
+	*/
+	void MergeClassCategories(FClass* Class);
+
+	/**
+	* Merges all class flags and validates them
+	*
+	* @param	DeclaredClassName Name this class was declared with (for validation)
+	* @param	PreviousClassFlags Class flags before resetting the class (for validation)
+	* @param	Class Class to merge flags for
+	* @param  AllClasses All known classes
+	*/
+	void MergeAndValidateClassFlags(const FString& DeclaredClassName, uint32 PreviousClassFlags, FClass* Class, const FClasses& AllClasses);
+private:
+
+	/** Merges all 'show' categories */
+	void MergeShowCategories();
+	/** Sets and validates 'within' property */
+	void SetAndValidateWithinClass(FClass* Class, const FClasses& AllClasses);
+	/** Sets and validates 'ConfigName' property */
+	void SetAndValidateConfigName(FClass* Class);
+
+	TArray<FString> ShowCategories;
+	TArray<FString> ShowFunctions;
+	TArray<FString> DontAutoCollapseCategories;
+	bool WantsToBePlaceable;
+};
diff --git a/Engine/Source/Programs/UnrealHeaderTool/Private/ClassMaps.h b/Engine/Source/Programs/UnrealHeaderTool/Private/ClassMaps.h
--- a/Engine/Source/Programs/UnrealHeaderTool/Private/ClassMaps.h
+++ b/Engine/Source/Programs/UnrealHeaderTool/Private/ClassMaps.h
@@ -8,26 +8,12 @@
 #include "HeaderProvider.h"
 
 #include "UnrealSourceFile.h"
+#include "ClassDeclarationMetaData.h"
 
 struct FManifestModule;
 class FUnrealSourceFile;
 class FUnrealTypeDefinitionInfo;
 
-/** Structure that holds class meta data generated from its UCLASS declaration */
-struct FClassDeclarationMetaData
-{
-	uint64 ClassFlags;
-	TMap<FName, FString> MetaData;
-	FString ClassWithin;
-	FString ConfigName;
-	TArray<FString> HideCategories;
-	TArray<FString> ShowSubCatgories;
-	TArray<FString> HideFunctions;
-	TArray<FString> AutoExpandCategories;
-	TArray<FString> AutoCollapseCategories;
-	TArray<FString> DependsOn;
-	TArray<FString> ClassGroupNames;
-};
 
 
 extern TMap<FString, TSharedRef<FUnrealSourceFile> > GUnrealSourceFilesMap;
diff --git a/Engine/Source/Programs/UnrealHeaderTool/Private/HeaderParser.cpp b/Engine/Source/Programs/UnrealHeaderTool/Private/HeaderParser.cpp
--- a/Engine/Source/Programs/UnrealHeaderTool/Private/HeaderParser.cpp
+++ b/Engine/Source/Programs/UnrealHeaderTool/Private/HeaderParser.cpp
@@ -65,20 +65,6 @@ extern ECompilationResult::Type GCompilationResult;
 
 namespace
 {
-	bool IsActorClass(UClass *Class)
-	{
-		static FName ActorName = FName(TEXT("Actor"));
-		while (Class)
-		{
-			if (Class->GetFName() == ActorName)
-			{
-				return true;
-			}
-			Class = Class->GetSuperClass();
-		}
-		return false;
-	}
-
 	bool ProbablyAMacro(const TCHAR* Identifier)
 	{
 		// Test for known delegate and event macros.
@@ -4304,255 +4290,6 @@ void PostParsingClassSetup(UClass* Class)
 	}
 }
 
-void FHeaderParser::ParseClassProperties(const TArray<FPropertySpecifier>& InClassSpecifiers, const FString& InRequiredAPIMacroIfPresent, FClassDeclarationMetaData& OutClassData)
-{
-	OutClassData.ClassFlags = 0;
-	// Record that this class is RequiredAPI if the CORE_API style macro was present
-	if (!InRequiredAPIMacroIfPresent.IsEmpty())
-	{
-		OutClassData.ClassFlags |= CLASS_RequiredAPI;
-	}
-	OutClassData.ClassFlags |= CLASS_Native;
-
-	// Process all of the class specifiers
-	
-	bool bWithinSpecified = false;
-	bool bDeclaresConfigFile = false;
-	for (const FPropertySpecifier& PropSpecifier : InClassSpecifiers)
-	{
-		const FString& Specifier = PropSpecifier.Key;
-
-		if (Specifier == TEXT("noexport"))
-		{
-			// Don't export to C++ header.
-			OutClassData.ClassFlags |= CLASS_NoExport;
-		}
-		else if (Specifier == TEXT("intrinsic"))
-		{
-			OutClassData.ClassFlags |= CLASS_Intrinsic;
-		}
-		else if (Specifier == TEXT("ComponentWrapperClass"))
-		{
-			OutClassData.MetaData.Add(TEXT("IgnoreCategoryKeywordsInSubclasses"), TEXT("true"));
-		}
-		else if (Specifier == TEXT("within"))
-		{
-			OutClassData.ClassWithin = RequireExactlyOneSpecifierValue(PropSpecifier);
-		}
-		else if (Specifier == TEXT("editinlinenew"))
-		{
-			// Class can be constructed from the New button in editinline
-			OutClassData.ClassFlags |= CLASS_EditInlineNew;
-		}
-		else if (Specifier == TEXT("noteditinlinenew"))
-		{
-			// Class cannot be constructed from the New button in editinline
-			OutClassData.ClassFlags &= ~CLASS_EditInlineNew;
-		}
-		else if (Specifier == TEXT("placeable"))
-		{
-			if (!(OutClassData.ClassFlags & CLASS_NotPlaceable))
-			{
-				FError::Throwf(TEXT("The placeable specifier is deprecated. Classes are assumed to be placeable by default."));
-			}
-			OutClassData.ClassFlags &= ~CLASS_NotPlaceable;
-		}
-		else if (Specifier == TEXT("defaulttoinstanced"))
-		{
-			// these classed default to instanced.
-			OutClassData.ClassFlags |= CLASS_DefaultToInstanced;
-		}
-		else if (Specifier == TEXT("notplaceable"))
-		{
-			// Don't allow the class to be placed in the editor.
-			OutClassData.ClassFlags |= CLASS_NotPlaceable;
-		}
-		else if (Specifier == TEXT("hidedropdown"))
-		{
-			// Prevents class from appearing in class comboboxes in the property window
-			OutClassData.ClassFlags |= CLASS_HideDropDown;
-		}
-		else if (Specifier == TEXT("dependsOn"))
-		{
-			FError::Throwf(TEXT("The dependsOn specifier is deprecated. Please use #include \"ClassHeaderFilename.h\" instead."));
-		}
-		else if (Specifier == TEXT("MinimalAPI"))
-		{
-			OutClassData.ClassFlags |= CLASS_MinimalAPI;
-		}
-		else if (Specifier == TEXT("const"))
-		{
-			OutClassData.ClassFlags |= CLASS_Const;
-		}
-		else if (Specifier == TEXT("perObjectConfig"))
-		{
-			OutClassData.ClassFlags |= CLASS_PerObjectConfig;
-		}
-		else if (Specifier == TEXT("configdonotcheckdefaults"))
-		{
-			OutClassData.ClassFlags |= CLASS_ConfigDoNotCheckDefaults;
-		}
-		else if (Specifier == TEXT("abstract"))
-		{
-			// Hide all editable properties.
-			OutClassData.ClassFlags |= CLASS_Abstract;
-		}
-		else if (Specifier == TEXT("deprecated"))
-		{
-			OutClassData.ClassFlags |= CLASS_Deprecated;
-
-			// Don't allow the class to be placed in the editor.
-			OutClassData.ClassFlags |= CLASS_NotPlaceable;
-		}
-		else if (Specifier == TEXT("transient"))
-		{
-			// Transient class.
-			OutClassData.ClassFlags |= CLASS_Transient;
-		}
-		else if (Specifier == TEXT("nonTransient"))
-		{
-			// this child of a transient class is not transient - remove the transient flag
-			OutClassData.ClassFlags &= ~CLASS_Transient;
-		}
-		else if (Specifier == TEXT("customConstructor"))
-		{
-			// we will not export a constructor for this class, assuming it is in the CPP block
-			OutClassData.ClassFlags |= CLASS_CustomConstructor;
-		}
-		else if (Specifier == TEXT("config"))
-		{
-			// Class containing config properties - parse the name of the config file to use
-			OutClassData.ConfigName = RequireExactlyOneSpecifierValue(PropSpecifier);
-		}
-		else if (Specifier == TEXT("defaultconfig"))
-		{
-			// Save object config only to Default INIs, never to local INIs.
-			OutClassData.ClassFlags |= CLASS_DefaultConfig;
-		}
-		else if (Specifier == TEXT("globaluserconfig"))
-		{
-			// Save object config only to global user overrides, never to local INIs
-			OutClassData.ClassFlags |= CLASS_GlobalUserConfig;
-		}
-		else if (Specifier == TEXT("showCategories"))
-		{
-			RequireSpecifierValue(PropSpecifier);
-
-			for (const FString& Value : PropSpecifier.Values)
-			{
-				// if we didn't find this specific category path in the HideCategories metadata
-				if (OutClassData.HideCategories.Remove(Value) == 0)
-				{
-					TArray<FString> SubCategoryList;
-					Value.ParseIntoArray(SubCategoryList, TEXT("|"), true);
-
-					FString SubCategoryPath;
-					// look to see if any of the parent paths are excluded in the HideCategories list
-					for (int32 CategoryPathIndex = 0; CategoryPathIndex < SubCategoryList.Num() - 1; ++CategoryPathIndex)
-					{
-						SubCategoryPath += SubCategoryList[CategoryPathIndex];
-						// if we're hiding a parent category, then we need to flag this sub category for show
-						if (OutClassData.HideCategories.Contains(SubCategoryPath))
-						{
-							OutClassData.ShowSubCatgories.AddUnique(Value);
-							break;
-						}
-						SubCategoryPath += "|";
-					}
-				}
-			}
-		}
-		else if (Specifier == TEXT("hideCategories"))
-		{
-			RequireSpecifierValue(PropSpecifier);
-
-			for (const FString& Value : PropSpecifier.Values)
-			{
-				OutClassData.HideCategories.AddUnique(Value);
-			}
-		}
-		else if (Specifier == TEXT("showFunctions"))
-		{
-			RequireSpecifierValue(PropSpecifier);
-
-			for (const FString& Value : PropSpecifier.Values)
-			{
-				OutClassData.HideFunctions.Remove(Value);
-			}
-		}
-		else if (Specifier == TEXT("hideFunctions"))
-		{
-			RequireSpecifierValue(PropSpecifier);
-
-			for (const FString& Value : PropSpecifier.Values)
-			{
-				OutClassData.HideFunctions.AddUnique(Value);
-			}
-		}
-		else if (Specifier == TEXT("classGroup"))
-		{
-			RequireSpecifierValue(PropSpecifier);
-
-			for (const FString& Value : PropSpecifier.Values)
-			{
-				OutClassData.ClassGroupNames.Add(Value);
-			}
-		}
-		else if (Specifier == TEXT("autoExpandCategories"))
-		{
-			RequireSpecifierValue(PropSpecifier);
-
-			for (const FString& Value : PropSpecifier.Values)
-			{
-				OutClassData.AutoCollapseCategories.Remove(Value);
-				OutClassData.AutoExpandCategories.AddUnique(Value);
-			}
-		}
-		else if (Specifier == TEXT("autoCollapseCategories"))
-		{
-			RequireSpecifierValue(PropSpecifier);
-
-			for (const FString& Value : PropSpecifier.Values)
-			{
-				OutClassData.AutoExpandCategories.Remove(Value);
-				OutClassData.AutoCollapseCategories.AddUnique(Value);
-			}
-		}
-		else if (Specifier == TEXT("dontAutoCollapseCategories"))
-		{
-			RequireSpecifierValue(PropSpecifier);
-
-			for (const FString& Value : PropSpecifier.Values)
-			{
-				OutClassData.AutoCollapseCategories.Remove(Value);
-			}
-		}
-		else if (Specifier == TEXT("collapseCategories"))
-		{
-			// Class' properties should not be shown categorized in the editor.
-			OutClassData.ClassFlags |= CLASS_CollapseCategories;
-		}
-		else if (Specifier == TEXT("dontCollapseCategories"))
-		{
-			// Class' properties should be shown categorized in the editor.
-			OutClassData.ClassFlags &= ~CLASS_CollapseCategories;
-		}
-		else if (Specifier == TEXT("AdvancedClassDisplay"))
-		{
-			// By default the class properties are shown in advanced sections in UI
-			OutClassData.ClassFlags |= CLASS_AdvancedDisplay;
-		}
-		else if (Specifier == TEXT("ConversionRoot"))
-		{
-			OutClassData.MetaData.Add(FName(TEXT("IsConversionRoot")), "true");
-		}
-		else
-		{
-			FError::Throwf(TEXT("Unknown class specifier '%s'"), *Specifier);
-		}
-	}
-}
-
 /**
  * Compiles a class declaration.
  */
@@ -4581,7 +4318,7 @@ void FHeaderParser::CompileClassDeclaration(FClasses& AllClasses)
 	
 	FClass* Class = ParseClassNameDeclaration(AllClasses, /*out*/ DeclaredClassName, /*out*/ RequiredAPIMacroIfPresent);
 	check(Class);
-	auto ClassDeclarationData = GClassDeclarations.FindChecked(Class->GetFName());
+	TSharedRef<FClassDeclarationMetaData> ClassDeclarationData = GClassDeclarations.FindChecked(Class->GetFName());
 
 	ClassDefinitionRanges.Add(Class, ClassDefinitionRange(&Input[InputPos], nullptr));
 
@@ -4591,8 +4328,7 @@ void FHeaderParser::CompileClassDeclaration(FClasses& AllClasses)
 
 	PushNest(ENestType::NEST_Class, Class);
 	
-	uint32 PreviousClassFlags = Class->ClassFlags;	
-
+	const uint32 PreviousClassFlags = Class->ClassFlags;	
 	ResetClassData();
 
 	// Verify class variables haven't been filled in
@@ -4609,87 +4345,18 @@ void FHeaderParser::CompileClassDeclaration(FClasses& AllClasses)
 		}
 	}
 
-	// Get categories inherited from the parent.
-	Class->GetHideCategories(ClassDeclarationData->HideCategories);
-	Class->GetShowCategories(ClassDeclarationData->ShowSubCatgories);
-	Class->GetHideFunctions(ClassDeclarationData->HideFunctions);
-	Class->GetAutoExpandCategories(ClassDeclarationData->AutoExpandCategories);
-	Class->GetAutoCollapseCategories(ClassDeclarationData->AutoCollapseCategories);
+	// Merge with categories inherited from the parent.
+	ClassDeclarationData->MergeClassCategories(Class);
 
 	// Class attributes.
 	FClassMetaData* ClassData = GScriptHelper.FindClassData(Class);
 	check(ClassData);
-
 	ClassData->SetPrologLine(PrologFinishLine);
 
-	Class->ClassFlags |= ClassDeclarationData->ClassFlags;
-
-	// All classes that are parsed are expected to be native
-	UClass* Super = Class->GetSuperClass();
-	if (Super && !Super->HasAnyClassFlags(CLASS_Native))
-	{
-		FError::Throwf(TEXT("Native classes cannot extend non-native classes") );
-	}
-
+	ClassDeclarationData->MergeAndValidateClassFlags(DeclaredClassName, PreviousClassFlags, Class, AllClasses);
 	Class->SetFlags(RF_Native);
 
-	// Process all of the class specifiers
-	bool bWithinSpecified = false;
-	bool bDeclaresConfigFile = false;
-	if (ClassDeclarationData->ClassWithin.IsEmpty() == false)
-	{
-		UClass* RequiredWithinClass = AllClasses.FindClass(*ClassDeclarationData->ClassWithin);
-		if (!RequiredWithinClass)
-		{
-			FError::Throwf(TEXT("Within class '%s' not found."), *ClassDeclarationData->ClassWithin);
-		}
-		if (RequiredWithinClass->IsChildOf(UInterface::StaticClass()))
-		{
-			FError::Throwf(TEXT("Classes cannot be 'within' interfaces"));
-		}
-		else if (Class->ClassWithin == NULL || Class->ClassWithin == UObject::StaticClass() || RequiredWithinClass->IsChildOf(Class->ClassWithin))
-		{
-			Class->ClassWithin = RequiredWithinClass;
-		}
-		else if (Class->ClassWithin != RequiredWithinClass)
-		{
-			FError::Throwf(TEXT("%s must be within %s, not %s"), *Class->GetPathName(), *Class->ClassWithin->GetPathName(), *RequiredWithinClass->GetPathName());
-		}
-		bWithinSpecified = true;
-	}
-	if (!!(Class->ClassFlags & CLASS_EditInlineNew))
-	{
-		// don't allow actor classes to be declared editinlinenew
-		if (IsActorClass(Class))
-		{
-			FError::Throwf(TEXT("Invalid class attribute: Creating actor instances via the property window is not allowed"));
-		}
-	}
-	if (ClassDeclarationData->ConfigName.IsEmpty() == false)
-	{
-		// if the user specified "inherit", we're just going to use the parent class's config filename
-		// this is not actually necessary but it can be useful for explicitly communicating config-ness
-		if (ClassDeclarationData->ConfigName == TEXT("inherit"))
-		{
-			UClass* SuperClass = Class->GetSuperClass();
-			if (!SuperClass)
-			{
-				FError::Throwf(TEXT("Cannot inherit config filename: %s has no super class"), *Class->GetName());
-			}
-
-			if (SuperClass->ClassConfigName == NAME_None)
-			{
-				FError::Throwf(TEXT("Cannot inherit config filename: parent class %s is not marked config."), *SuperClass->GetPathName());
-			}
-		}
-		else
-		{
-			// otherwise, set the config name to the parsed identifier
-			Class->ClassConfigName = FName(*ClassDeclarationData->ConfigName);
-		}
-		bDeclaresConfigFile = true;
-	}
-
+	// Class metadata
 	if (ClassDeclarationData->ClassGroupNames.Num()) { MetaData.Add("ClassGroupNames", FString::Join(ClassDeclarationData->ClassGroupNames, TEXT(" "))); }
 	if (ClassDeclarationData->AutoCollapseCategories.Num()) { MetaData.Add("AutoCollapseCategories", FString::Join(ClassDeclarationData->AutoCollapseCategories, TEXT(" "))); }
 	if (ClassDeclarationData->HideCategories.Num()) { MetaData.Add("HideCategories", FString::Join(ClassDeclarationData->HideCategories, TEXT(" "))); }
@@ -4697,70 +4364,6 @@ void FHeaderParser::CompileClassDeclaration(FClasses& AllClasses)
 	if (ClassDeclarationData->HideFunctions.Num()) { MetaData.Add("HideFunctions", FString::Join(ClassDeclarationData->HideFunctions, TEXT(" "))); }
 	if (ClassDeclarationData->AutoExpandCategories.Num()) { MetaData.Add("AutoExpandCategories", FString::Join(ClassDeclarationData->AutoExpandCategories, TEXT(" "))); }
 
-	// Make sure both RequiredAPI and MinimalAPI aren't specified
-	if (Class->HasAllClassFlags(CLASS_MinimalAPI | CLASS_RequiredAPI))
-	{
-		FError::Throwf(TEXT("MinimalAPI cannot be specified when the class is fully exported using a MODULENAME_API macro"));
-	}
-
-	// Make sure there is a valid within
-	if (!bWithinSpecified)
-	{
-		// classes always have a ClassWithin
-		Class->ClassWithin = Class->GetSuperClass()
-			? Class->GetSuperClass()->ClassWithin
-			: UObject::StaticClass();
-	}
-
-	UClass* ExpectedWithin = Class->GetSuperClass()
-		? Class->GetSuperClass()->ClassWithin
-		: UObject::StaticClass();
-
-	if (!Class->ClassWithin->IsChildOf(ExpectedWithin))
-	{
-		FError::Throwf(TEXT("Parent class declared within '%s'.  Cannot override within class with '%s' since it isn't a child"), *ExpectedWithin->GetName(), *Class->ClassWithin->GetName());
-		///Class->ClassWithin = ExpectedWithin;
-	}
-
-	// All classes must start with a valid Unreal prefix
-	const FString ExpectedClassName = Class->GetNameWithPrefix();
-	if( DeclaredClassName != ExpectedClassName )
-	{
-		FError::Throwf(TEXT("Class name '%s' is invalid, should be identified as '%s'"), *DeclaredClassName, *ExpectedClassName );
-	}
-
-	// Validate.
-	if ((Class->ClassFlags&CLASS_NoExport))
-	{
-		// if the class's class flags didn't contain CLASS_NoExport before it was parsed, it means either:
-		// a) the DECLARE_CLASS macro for this native class doesn't contain the CLASS_NoExport flag (this is an error)
-		// b) this is a new native class, which isn't yet hooked up to static registration (this is OK)
-		if (!(Class->ClassFlags&CLASS_Intrinsic) && (PreviousClassFlags & CLASS_NoExport) == 0 &&
-			(PreviousClassFlags&CLASS_Native) != 0 )	// a new native class (one that hasn't been compiled into C++ yet) won't have this set
-		{
-			FError::Throwf(TEXT("'noexport': Must include CLASS_NoExport in native class declaration"));
-		}
-	}
-
-	if (!Class->HasAnyClassFlags(CLASS_Abstract) && ((PreviousClassFlags & CLASS_Abstract) != 0))
-	{
-		if (Class->HasAnyClassFlags(CLASS_NoExport))
-		{
-			FError::Throwf(TEXT("'abstract': NoExport class missing abstract keyword from class declaration (must change C++ version first)"));
-			Class->ClassFlags |= CLASS_Abstract;
-		}
-		else if (Class->HasAnyFlags(RF_Native))
-		{
-			FError::Throwf(TEXT("'abstract': missing abstract keyword from class declaration - class will no longer be exported as abstract"));
-		}
-	}
-
-	// Invalidate config name if not specifically declared.
-	if (!bDeclaresConfigFile)
-	{
-		Class->ClassConfigName = NAME_None;
-	}
-
 	AddIncludePathToMetadata(Class, MetaData);
 	AddModuleRelativePathToMetadata(Class, MetaData);
 
@@ -7307,8 +6910,7 @@ void FHeaderPreParser::ParseClassDeclaration(const TCHAR* InputText, int32 InLin
 		// Add class declaration meta data so that we can access class flags before the class is fully parsed
 		TSharedRef<FClassDeclarationMetaData> DeclarationData = MakeShareable(new FClassDeclarationMetaData());
 		DeclarationData->MetaData = MetaData;
-		DeclarationData->ClassFlags = 0;
-		FHeaderParser::ParseClassProperties(SpecifiersFound, RequiredAPIMacroIfPresent, *DeclarationData);
+		DeclarationData->ParseClassProperties(SpecifiersFound, RequiredAPIMacroIfPresent);
 		GClassDeclarations.Add(ClassNameWithoutPrefix, DeclarationData);
 	}
 
diff --git a/Engine/Source/Programs/UnrealHeaderTool/Private/HeaderParser.h b/Engine/Source/Programs/UnrealHeaderTool/Private/HeaderParser.h
--- a/Engine/Source/Programs/UnrealHeaderTool/Private/HeaderParser.h
+++ b/Engine/Source/Programs/UnrealHeaderTool/Private/HeaderParser.h
@@ -256,15 +256,6 @@ public:
 	 */
 	static ECompilationResult::Type ParseHeaders(FClasses& AllClasses, FHeaderParser& HeaderParser, FUnrealSourceFile& SourceFile, bool bParseSubclasses);
 
-	/**
-	 * Parse Class's properties to generate its declaration data.
-	 *
-	 * @param	InClassSpecifiers Class properties collected from its UCLASS macro
-	 * @param	InRequiredAPIMacroIfPresent *_API macro if present (empty otherwise)
-	 * @param	OutClassData Parsed class meta data
-	 */
-	static void ParseClassProperties(const TArray<FPropertySpecifier>& InClassSpecifiers, const FString& InRequiredAPIMacroIfPresent, FClassDeclarationMetaData& OutClassData);
-
 protected:
 	friend struct FScriptLocation;
 
@@ -459,10 +450,13 @@ protected:
 	// Parse the parameter list of a function or delegate declaration
 	void ParseParameterList(FClasses& AllClasses, UFunction* Function, bool bExpectCommaBeforeName = false, TMap<FName, FString>* MetaData = NULL);
 
+public:
 	// Throws if a specifier value wasn't provided
 	static void RequireSpecifierValue(const FPropertySpecifier& Specifier, bool bRequireExactlyOne = false);
 	static FString RequireExactlyOneSpecifierValue(const FPropertySpecifier& Specifier);
 
+protected:
+
 	/**
 	 * Parse rest of the module's source files.
 	 *
diff --git a/Engine/Source/Runtime/Landscape/Classes/Landscape.h b/Engine/Source/Runtime/Landscape/Classes/Landscape.h
--- a/Engine/Source/Runtime/Landscape/Classes/Landscape.h
+++ b/Engine/Source/Runtime/Landscape/Classes/Landscape.h
@@ -26,7 +26,7 @@ enum ELandscapeSetupErrors
 	LSE_MAX,
 };
 
-UCLASS(hidecategories=LandscapeProxy, showcategories=(Display, Movement, Collision, Lighting, LOD, Input), MinimalAPI)
+UCLASS(Placeable, hidecategories = LandscapeProxy, showcategories = (Display, Movement, Collision, Lighting, LOD, Input), MinimalAPI)
 class ALandscape : public ALandscapeProxy
 {
 	GENERATED_UCLASS_BODY()
