Removing debug code

diff --git a/Engine/Source/Programs/UnrealHeaderTool/Private/ClassDeclarationMetaData.cpp b/Engine/Source/Programs/UnrealHeaderTool/Private/ClassDeclarationMetaData.cpp
--- a/Engine/Source/Programs/UnrealHeaderTool/Private/ClassDeclarationMetaData.cpp
+++ b/Engine/Source/Programs/UnrealHeaderTool/Private/ClassDeclarationMetaData.cpp
@@ -307,12 +307,6 @@ void FClassDeclarationMetaData::MergeShowCategories()
 
 void FClassDeclarationMetaData::MergeClassCategories(FClass* Class)
 {
-	static FName DEBUG_Name(TEXT("StaticMeshComponent"));
-	if (Class->GetFName() == DEBUG_Name)
-	{
-		static volatile int32 xx = 0;
-		xx++;
-	}
 	TArray<FString> ParentHideCategories;
 	TArray<FString> ParentShowSubCatgories;
 	TArray<FString> ParentHideFunctions;
