Warn that virtual BlueprintImplementableEvents are not allowed, will remove support in a future release.

diff --git a/Engine/Source/Programs/UnrealHeaderTool/Private/HeaderParser.cpp b/Engine/Source/Programs/UnrealHeaderTool/Private/HeaderParser.cpp
--- a/Engine/Source/Programs/UnrealHeaderTool/Private/HeaderParser.cpp
+++ b/Engine/Source/Programs/UnrealHeaderTool/Private/HeaderParser.cpp
@@ -5181,22 +5181,24 @@ void FHeaderParser::CompileFunctionDeclaration(FUnrealSourceFile& SourceFile, FC
 		bAutomaticallyFinal = false;
 
 		// if this is a BlueprintNativeEvent or BlueprintImplementableEvent in an interface, make sure it's not "virtual"
-		if ((GetCurrentClass()->HasAnyClassFlags(CLASS_Interface)) && (FuncInfo.FunctionFlags & FUNC_BlueprintEvent))
+		if (FuncInfo.FunctionFlags & FUNC_BlueprintEvent)
 		{
-			FError::Throwf(TEXT("BlueprintImplementableEvents in Interfaces must not be declared 'virtual'"));
-		}
+			if (GetCurrentClass()->HasAnyClassFlags(CLASS_Interface))
+			{
+				FError::Throwf(TEXT("BlueprintImplementableEvents in Interfaces must not be declared 'virtual'"));
+			}
 
-		// if this is a BlueprintNativeEvent, make sure it's not "virtual"
-		if ((FuncInfo.FunctionFlags & FUNC_BlueprintEvent) && (FuncInfo.FunctionFlags & FUNC_Native))
-		{
-			FError::Throwf(TEXT("BlueprintNativeEvent functions must be non-virtual."));
-		}
+			// if this is a BlueprintNativeEvent, make sure it's not "virtual"
+			else if (FuncInfo.FunctionFlags & FUNC_Native)
+			{
+				FError::Throwf(TEXT("BlueprintNativeEvent functions must be non-virtual."));
+			}
 
-		//		@todo: we should consider making BIEs nonvirtual as well, where could simplify these checks to this...
-		// 		if ( (FuncInfo.FunctionFlags & FUNC_BlueprintEvent) )
-		// 		{
-		// 			FError::Throwf(TEXT("Functions marked BlueprintImplementableEvent or BlueprintNativeEvent must not be declared 'virtual'"));
-		// 		}
+			else
+			{
+				UE_LOG(LogCompile, Warning, TEXT("BlueprintImplementableEvents should not be virtual. Use BlueprintNativeEvent instead."));
+			}
+		}
 	}
 	else
 	{
