Temporarily disable warning since UHT turns it in to an error

diff --git a/Engine/Source/Programs/UnrealHeaderTool/Private/HeaderParser.cpp b/Engine/Source/Programs/UnrealHeaderTool/Private/HeaderParser.cpp
--- a/Engine/Source/Programs/UnrealHeaderTool/Private/HeaderParser.cpp
+++ b/Engine/Source/Programs/UnrealHeaderTool/Private/HeaderParser.cpp
@@ -5196,7 +5196,7 @@ void FHeaderParser::CompileFunctionDeclaration(FUnrealSourceFile& SourceFile, FC
 
 			else
 			{
-				UE_LOG(LogCompile, Warning, TEXT("BlueprintImplementableEvents should not be virtual. Use BlueprintNativeEvent instead."));
+				//UE_LOG(LogCompile, Warning, TEXT("BlueprintImplementableEvents should not be virtual. Use BlueprintNativeEvent instead."));
 			}
 		}
 	}
