Function with a TArray<UClass*> return value added to test, to check for bad generated code in exec function.

diff --git a/Engine/Source/Programs/UnrealHeaderTool/Resources/UHTDebugging/TestObject.h b/Engine/Source/Programs/UnrealHeaderTool/Resources/UHTDebugging/TestObject.h
--- a/Engine/Source/Programs/UnrealHeaderTool/Resources/UHTDebugging/TestObject.h
+++ b/Engine/Source/Programs/UnrealHeaderTool/Resources/UHTDebugging/TestObject.h
@@ -27,6 +27,9 @@ public:
 	UFUNCTION()
 	void CodeGenTestForEnumClasses(ECppEnum Val);
 
+	UFUNCTION(Category="Xyz", BlueprintCallable)
+	TArray<UClass*> ReturnArrayOfUClassPtrs();
+
 	UFUNCTION()
 	int32 InlineFunc()
 	{
