Shadowed variables fixed.

diff --git a/Engine/Source/Programs/UnrealHeaderTool/Private/CodeGenerator.cpp b/Engine/Source/Programs/UnrealHeaderTool/Private/CodeGenerator.cpp
--- a/Engine/Source/Programs/UnrealHeaderTool/Private/CodeGenerator.cpp
+++ b/Engine/Source/Programs/UnrealHeaderTool/Private/CodeGenerator.cpp
@@ -1099,10 +1099,10 @@ FString GetPackageSingletonName(UPackage* Package)
 	return FString(TEXT("Z_Construct_")) + ClassString + TEXT("_") + FPackageName::GetShortName(Package) + TEXT("()");
 }
 
-void FNativeClassHeaderGenerator::ExportGeneratedPackageInitCode(UPackage* Package)
+void FNativeClassHeaderGenerator::ExportGeneratedPackageInitCode(UPackage* InPackage)
 {
 	FString ApiString = GetAPIString();
-	FString SingletonName(GetPackageSingletonName(Package));
+	FString SingletonName(GetPackageSingletonName(InPackage));
 
 	FUHTStringBuilder& GeneratedFunctionText = GetGeneratedFunctionTextDevice();
 
@@ -1113,9 +1113,9 @@ void FNativeClassHeaderGenerator::ExportGeneratedPackageInitCode(UPackage* Packa
 	GeneratedFunctionText.Logf(TEXT("        static UPackage* ReturnPackage = NULL;\r\n"));
 	GeneratedFunctionText.Logf(TEXT("        if (!ReturnPackage)\r\n"));
 	GeneratedFunctionText.Logf(TEXT("        {\r\n"));
-	GeneratedFunctionText.Logf(TEXT("            ReturnPackage = CastChecked<UPackage>(StaticFindObjectFast(UPackage::StaticClass(), NULL, FName(TEXT(\"%s\")), false, false));\r\n"), *Package->GetName());
+	GeneratedFunctionText.Logf(TEXT("            ReturnPackage = CastChecked<UPackage>(StaticFindObjectFast(UPackage::StaticClass(), NULL, FName(TEXT(\"%s\")), false, false));\r\n"), *InPackage->GetName());
 
-	FString Meta = GetMetaDataCodeForObject(Package, TEXT("ReturnPackage"), TEXT("            "));
+	FString Meta = GetMetaDataCodeForObject(InPackage, TEXT("ReturnPackage"), TEXT("            "));
 	if (Meta.Len())
 	{
 		GeneratedFunctionText.Logf(TEXT("#if WITH_METADATA\r\n"));
@@ -1124,8 +1124,8 @@ void FNativeClassHeaderGenerator::ExportGeneratedPackageInitCode(UPackage* Packa
 		GeneratedFunctionText.Logf(TEXT("#endif\r\n"));
 	}
 
-	GeneratedFunctionText.Logf(TEXT("            ReturnPackage->PackageFlags |= PKG_CompiledIn | 0x%08X;\r\n"), Package->PackageFlags & (PKG_ClientOptional|PKG_ServerSideOnly));
-	TheFlagAudit.Add(Package, TEXT("PackageFlags"), Package->PackageFlags);
+	GeneratedFunctionText.Logf(TEXT("            ReturnPackage->PackageFlags |= PKG_CompiledIn | 0x%08X;\r\n"), InPackage->PackageFlags & (PKG_ClientOptional|PKG_ServerSideOnly));
+	TheFlagAudit.Add(InPackage, TEXT("PackageFlags"), InPackage->PackageFlags);
 	{
 		FGuid Guid;
 
@@ -1157,7 +1157,7 @@ void FNativeClassHeaderGenerator::ExportGeneratedPackageInitCode(UPackage* Packa
 
 		for (UField* ScriptType : TObjectRange<UField>())
 		{
-			if (ScriptType->GetOutermost() != Package)
+			if (ScriptType->GetOutermost() != InPackage)
 			{
 				continue;
 			}
@@ -1933,107 +1933,100 @@ void FNativeClassHeaderGenerator::ExportClassesFromSourceFileInner(FUnrealSource
 		}
 		else
 		{
-			FClass* SuperClass = Class->GetSuperClass();
+			FUHTStringBuilder ClassBoilerplate;
 
-			const TCHAR* ClassCPPName = NameLookupCPP.GetNameCPP(Class);
-			const TCHAR* SuperClassCPPName = (SuperClass != NULL) ? NameLookupCPP.GetNameCPP(SuperClass) : NULL;
+			// Export the class's native function registration.
+			ClassBoilerplate.Logf(TEXT("    private:\r\n"));
+			ClassBoilerplate.Logf(TEXT("    static void StaticRegisterNatives%s();\r\n"), NameLookupCPP.GetNameCPP(Class));
+			ClassBoilerplate.Log(*FriendText);
+			ClassBoilerplate.Logf(TEXT("    public:\r\n"));
 
+			// Build the DECLARE_CLASS line
+			FString APIArg(API);
+			if (!Class->HasAnyClassFlags(CLASS_MinimalAPI))
 			{
-				FUHTStringBuilder ClassBoilerplate;
-
-				// Export the class's native function registration.
-				ClassBoilerplate.Logf(TEXT("    private:\r\n"));
-				ClassBoilerplate.Logf(TEXT("    static void StaticRegisterNatives%s();\r\n"), NameLookupCPP.GetNameCPP(Class));
-				ClassBoilerplate.Log(*FriendText);
-				ClassBoilerplate.Logf(TEXT("    public:\r\n"));
+				APIArg = TEXT("NO");
+			}
+			const bool bCastedClass = Class->HasAnyCastFlag(CASTCLASS_AllFlags) && SuperClass && Class->ClassCastFlags != SuperClass->ClassCastFlags;
+			ClassBoilerplate.Logf(TEXT("    DECLARE_CLASS(%s, %s, COMPILED_IN_FLAGS(%s%s), %s, %s, %s_API)\r\n"),
+				ClassCPPName,
+				SuperClassCPPName ? SuperClassCPPName : TEXT("None"),
+				Class->HasAnyClassFlags(CLASS_Abstract) ? TEXT("CLASS_Abstract") : TEXT("0"),
+				*GetClassFlagExportText(Class),
+				bCastedClass ? *FString::Printf(TEXT("CASTCLASS_%s"), ClassCPPName) : TEXT("0"),
+				*FPackageName::GetShortName(*Class->GetOuter()->GetName()),
+				*APIArg);
+
+			ClassBoilerplate.Logf(TEXT("    DECLARE_SERIALIZER(%s)\r\n"), ClassCPPName);
+			ClassBoilerplate.Log(TEXT("    /** Indicates whether the class is compiled into the engine */"));
+			ClassBoilerplate.Log(TEXT("    enum {IsIntrinsic=COMPILED_IN_INTRINSIC};\r\n"));
+
+			if (SuperClass && Class->ClassWithin != SuperClass->ClassWithin)
+			{
+				ClassBoilerplate.Logf(TEXT("    DECLARE_WITHIN(%s)\r\n"), NameLookupCPP.GetNameCPP(Class->GetClassWithin()));
+			}
 
-				// Build the DECLARE_CLASS line
-				FString APIArg(API);
-				if (!Class->HasAnyClassFlags(CLASS_MinimalAPI))
-				{
-					APIArg = TEXT("NO");
-				}
-				const bool bCastedClass = Class->HasAnyCastFlag(CASTCLASS_AllFlags) && SuperClass && Class->ClassCastFlags != SuperClass->ClassCastFlags;
-				ClassBoilerplate.Logf(TEXT("    DECLARE_CLASS(%s, %s, COMPILED_IN_FLAGS(%s%s), %s, %s, %s_API)\r\n"),
-					ClassCPPName,
-					SuperClassCPPName ? SuperClassCPPName : TEXT("None"),
-					Class->HasAnyClassFlags(CLASS_Abstract) ? TEXT("CLASS_Abstract") : TEXT("0"),
-					*GetClassFlagExportText(Class),
-					bCastedClass ? *FString::Printf(TEXT("CASTCLASS_%s"), ClassCPPName) : TEXT("0"),
-					*FPackageName::GetShortName(*Class->GetOuter()->GetName()),
-					*APIArg);
+			// export the class's config name
+			if (SuperClass && Class->ClassConfigName != NAME_None && Class->ClassConfigName != SuperClass->ClassConfigName)
+			{
+				ClassBoilerplate.Logf(TEXT("    static const TCHAR* StaticConfigName() {return TEXT(\"%s\");}\r\n\r\n"), *Class->ClassConfigName.ToString());
+			}
 
-				ClassBoilerplate.Logf(TEXT("    DECLARE_SERIALIZER(%s)\r\n"), ClassCPPName);
-				ClassBoilerplate.Log(TEXT("    /** Indicates whether the class is compiled into the engine */"));
-				ClassBoilerplate.Log(TEXT("    enum {IsIntrinsic=COMPILED_IN_INTRINSIC};\r\n"));
+			// @todo srobb: clean up _getUObject() code generation
+			if (((SuperClass == UClass::StaticClass()) || (SuperClass == UObject::StaticClass()) || (SuperClass && (SuperClass->ClassFlags & CLASS_Intrinsic))) && !Class->IsChildOf(UInterface::StaticClass()))
+			{
+				ClassBoilerplate.Logf(TEXT("\tvirtual UObject* _getUObject() const { return const_cast<%s*>(this); }\r\n"), ClassCPPName);
+			}
+			else
+			{
+				ClassBoilerplate.Logf(TEXT("\tvirtual UObject* _getUObject() const override { return const_cast<%s*>(this); }\r\n"), ClassCPPName);
+			}
 
-				if (SuperClass && Class->ClassWithin != SuperClass->ClassWithin)
+			// Replication, add in the declaration for GetLifetimeReplicatedProps() automatically if there are any net flagged properties
+			bool bHasNetFlaggedProperties = false;
+			for (TFieldIterator<UProperty> It(Class, EFieldIteratorFlags::ExcludeSuper); It; ++It)
+			{
+				if ((It->PropertyFlags & CPF_Net) != 0)
 				{
-					ClassBoilerplate.Logf(TEXT("    DECLARE_WITHIN(%s)\r\n"), NameLookupCPP.GetNameCPP(Class->GetClassWithin()));
+					bHasNetFlaggedProperties = true;
+					break;
 				}
+			}
 
-				// export the class's config name
-				if (SuperClass && Class->ClassConfigName != NAME_None && Class->ClassConfigName != SuperClass->ClassConfigName)
-				{
-					ClassBoilerplate.Logf(TEXT("    static const TCHAR* StaticConfigName() {return TEXT(\"%s\");}\r\n\r\n"), *Class->ClassConfigName.ToString());
-				}
+			auto ClassRange = ClassDefinitionRange();
+			if (ClassDefinitionRanges.Contains(Class))
+			{
+				ClassRange = ClassDefinitionRanges[Class];
+			}
 
-				// @todo srobb: clean up _getUObject() code generation
-				if (((SuperClass == UClass::StaticClass()) || (SuperClass == UObject::StaticClass()) || (SuperClass && (SuperClass->ClassFlags & CLASS_Intrinsic))) && !Class->IsChildOf(UInterface::StaticClass()))
-				{
-					ClassBoilerplate.Logf(TEXT("\tvirtual UObject* _getUObject() const { return const_cast<%s*>(this); }\r\n"), ClassCPPName);
-				}
-				else
-				{
-					ClassBoilerplate.Logf(TEXT("\tvirtual UObject* _getUObject() const override { return const_cast<%s*>(this); }\r\n"), ClassCPPName);
-				}
+			auto ClassStart = ClassRange.Start;
+			auto ClassEnd = ClassRange.End;
+			auto ClassDefinition = FString(ClassEnd - ClassStart, ClassStart);
 
-				// Replication, add in the declaration for GetLifetimeReplicatedProps() automatically if there are any net flagged properties
-				bool bHasNetFlaggedProperties = false;
-				for (TFieldIterator<UProperty> It(Class, EFieldIteratorFlags::ExcludeSuper); It; ++It)
-				{
-					if ((It->PropertyFlags & CPF_Net) != 0)
-					{
-						bHasNetFlaggedProperties = true;
-						break;
-					}
-				}
+			bool bHasGetLifetimeReplicatedProps = HasIdentifierExactMatch(ClassDefinition, TEXT("GetLifetimeReplicatedProps"));
 
-				auto ClassRange = ClassDefinitionRange();
-				if (ClassDefinitionRanges.Contains(Class))
+			if (bHasNetFlaggedProperties && !bHasGetLifetimeReplicatedProps)
+			{
+				// Default version autogenerates declarations.
+				if (SourceFile.GetGeneratedCodeVersionForStruct(Class) == EGeneratedCodeVersion::V1)
 				{
-					ClassRange = ClassDefinitionRanges[Class];
+					ClassBoilerplate.Logf(TEXT("\tvoid GetLifetimeReplicatedProps(TArray<FLifetimeProperty>& OutLifetimeProps) const override;\r\n"));
 				}
-
-				auto ClassStart = ClassRange.Start;
-				auto ClassEnd = ClassRange.End;
-				auto ClassDefinition = FString(ClassEnd - ClassStart, ClassStart);
-
-				bool bHasGetLifetimeReplicatedProps = HasIdentifierExactMatch(ClassDefinition, TEXT("GetLifetimeReplicatedProps"));
-
-				if (bHasNetFlaggedProperties && !bHasGetLifetimeReplicatedProps)
+				else
 				{
-					// Default version autogenerates declarations.
-					if (SourceFile.GetGeneratedCodeVersionForStruct(Class) == EGeneratedCodeVersion::V1)
-					{
-						ClassBoilerplate.Logf(TEXT("\tvoid GetLifetimeReplicatedProps(TArray<FLifetimeProperty>& OutLifetimeProps) const override;\r\n"));
-					}
-					else
-					{
-						FError::Throwf(TEXT("Class %s has Net flagged properties and should declare member function: void GetLifetimeReplicatedProps(TArray<FLifetimeProperty>& OutLifetimeProps) const override"), ClassCPPName);
-					}
+					FError::Throwf(TEXT("Class %s has Net flagged properties and should declare member function: void GetLifetimeReplicatedProps(TArray<FLifetimeProperty>& OutLifetimeProps) const override"), ClassCPPName);
 				}
+			}
 
-				auto NoPureDeclsMacroName = SourceFile.GetGeneratedMacroName(ClassData, TEXT("_INCLASS_NO_PURE_DECLS"));
-				WriteMacro(GeneratedHeaderText, NoPureDeclsMacroName, ClassBoilerplate);
-				InClassNoPureDeclsMacroCalls.Logf(TEXT("\t%s\r\n"), *NoPureDeclsMacroName);
+			auto NoPureDeclsMacroName = SourceFile.GetGeneratedMacroName(ClassData, TEXT("_INCLASS_NO_PURE_DECLS"));
+			WriteMacro(GeneratedHeaderText, NoPureDeclsMacroName, ClassBoilerplate);
+			InClassNoPureDeclsMacroCalls.Logf(TEXT("\t%s\r\n"), *NoPureDeclsMacroName);
 
-				FString MacroName = SourceFile.GetGeneratedMacroName(ClassData, TEXT("_INCLASS"));
-				WriteMacro(GeneratedHeaderText, MacroName, ClassBoilerplate);
-				InClassMacroCalls.Logf(TEXT("\t%s\r\n"), *MacroName);
+			FString MacroName = SourceFile.GetGeneratedMacroName(ClassData, TEXT("_INCLASS"));
+			WriteMacro(GeneratedHeaderText, MacroName, ClassBoilerplate);
+			InClassMacroCalls.Logf(TEXT("\t%s\r\n"), *MacroName);
 
-				ExportConstructorsMacros(SourceFile.GetGeneratedMacroName(ClassData), Class);
-			}
+			ExportConstructorsMacros(SourceFile.GetGeneratedMacroName(ClassData), Class);
 		}
 
 		{
@@ -2302,14 +2295,14 @@ FString GetBuildPath(FUnrealSourceFile& SourceFile)
 	return Out;
 }
 
-FString FNativeClassHeaderGenerator::GetListOfPublicHeaderGroupIncludesString(UPackage* Package)
+FString FNativeClassHeaderGenerator::GetListOfPublicHeaderGroupIncludesString(UPackage* InPackage)
 {
 	FUHTStringBuilder Out;
 
 	// Fill with the rest source files from this package.
 	for (auto* SourceFile : GPublicSourceFileSet)
 	{
-		if (SourceFile->GetPackage() == Package && !ListOfPublicHeaderGroupIncludes.Contains(SourceFile))
+		if (SourceFile->GetPackage() == InPackage && !ListOfPublicHeaderGroupIncludes.Contains(SourceFile))
 		{
 			ListOfPublicHeaderGroupIncludes.Add(SourceFile);
 		}
@@ -2485,10 +2478,10 @@ void FNativeClassHeaderGenerator::ExportSourceFileHeaderRecursive(FClasses& AllC
 	TArray<FUnrealSourceFile*> DependOnList;
 	for (auto& Include : SourceFile->GetIncludes())
 	{
-		auto* SourceFile = Include.Resolve();
-		if (SourceFile && !DependOnList.Contains(SourceFile))
+		auto* OtherSourceFile = Include.Resolve();
+		if (OtherSourceFile && !DependOnList.Contains(OtherSourceFile))
 		{
-			DependOnList.Add(SourceFile);
+			DependOnList.Add(OtherSourceFile);
 		}
 	}
 
@@ -2783,15 +2776,15 @@ void FNativeClassHeaderGenerator::ExportGeneratedEnumsInitCode(const TArray<UEnu
 		GeneratedPackageCPP.Logf(TEXT("static FCompiledInDeferEnum Z_CompiledInDeferEnum_UEnum_%s(%s_StaticEnum, TEXT(\"%s\"));\r\n"), *Enum->GetName(), *Enum->GetName(), *Enum->GetOutermost()->GetName());
 
 		{
-			const FString SingletonName(GetSingletonName(Enum));
+			const FString EnumSingletonName = GetSingletonName(Enum);
 
-			FString Extern = FString::Printf(TEXT("    %s_API class UEnum* %s;\r\n"), *API, *SingletonName);
-			SingletonNameToExternDecl.Add(SingletonName, Extern);
+			FString Extern = FString::Printf(TEXT("    %s_API class UEnum* %s;\r\n"), *API, *EnumSingletonName);
+			SingletonNameToExternDecl.Add(EnumSingletonName, Extern);
 			GeneratedFunctionDeclarations.Log(*Extern);
 
 			FUHTStringBuilder GeneratedEnumRegisterFunctionText;
 
-			GeneratedEnumRegisterFunctionText.Logf(TEXT("    UEnum* %s\r\n"), *SingletonName);
+			GeneratedEnumRegisterFunctionText.Logf(TEXT("    UEnum* %s\r\n"), *EnumSingletonName);
 			GeneratedEnumRegisterFunctionText.Logf(TEXT("    {\r\n"));
 			// Enums can either have a UClass or UPackage as outer (if declared in non-UClass header).
 			if (Enum->GetOuter()->IsA(UStruct::StaticClass()))
@@ -2840,7 +2833,7 @@ void FNativeClassHeaderGenerator::ExportGeneratedEnumsInitCode(const TArray<UEnu
 
 			uint32 EnumCrc = GenerateTextCRC(*GeneratedEnumRegisterFunctionText);
 			GGeneratedCodeCRCs.Add(Enum, EnumCrc);
-			// CallSingletons.Logf(TEXT("                OuterClass->LinkChild(%s); // %u\r\n"), *SingletonName, EnumCrc);
+			// CallSingletons.Logf(TEXT("                OuterClass->LinkChild(%s); // %u\r\n"), *EnumSingletonName, EnumCrc);
 		}
 	}
 }
@@ -4538,7 +4531,7 @@ FNativeClassHeaderGenerator::FNativeClassHeaderGenerator(
 	FClasses& AllClasses,
 	bool InAllowSaveExportedHeaders
 #if WITH_HOT_RELOAD_CTORS
-	, bool bExportVTableConstructors
+	, bool bInExportVTableConstructors
 #endif // WITH_HOT_RELOAD_CTORS
 )
 	: API                                   (FPackageName::GetShortName(InPackage).ToUpper())
@@ -4547,7 +4540,7 @@ FNativeClassHeaderGenerator::FNativeClassHeaderGenerator(
 	, bAllowSaveExportedHeaders             (InAllowSaveExportedHeaders)
 	, bFailIfGeneratedCodeChanges           (false)
 #if WITH_HOT_RELOAD_CTORS
-	, bExportVTableConstructors				(bExportVTableConstructors)
+	, bExportVTableConstructors				(bInExportVTableConstructors)
 #endif // WITH_HOT_RELOAD_CTORS
 {
 	const FString PackageName = FPackageName::GetShortName(Package);
diff --git a/Engine/Source/Programs/UnrealHeaderTool/Private/HeaderParser.cpp b/Engine/Source/Programs/UnrealHeaderTool/Private/HeaderParser.cpp
--- a/Engine/Source/Programs/UnrealHeaderTool/Private/HeaderParser.cpp
+++ b/Engine/Source/Programs/UnrealHeaderTool/Private/HeaderParser.cpp
@@ -749,9 +749,9 @@ FString FHeaderParser::GetContext()
 {
 	auto* FileScope = GetCurrentFileScope();
 
-	FString Filename = IFileManager::Get().ConvertToAbsolutePathForExternalAppForRead(*FileScope->GetSourceFile()->GetFilename());
+	FString ScopeFilename = IFileManager::Get().ConvertToAbsolutePathForExternalAppForRead(*FileScope->GetSourceFile()->GetFilename());
 
-	return FString::Printf(TEXT("%s(%i)"), *Filename, InputLine);
+	return FString::Printf(TEXT("%s(%i)"), *ScopeFilename, InputLine);
 }
 
 /*-----------------------------------------------------------------------------
diff --git a/Engine/Source/Programs/UnrealHeaderTool/Private/HeaderParser.h b/Engine/Source/Programs/UnrealHeaderTool/Private/HeaderParser.h
--- a/Engine/Source/Programs/UnrealHeaderTool/Private/HeaderParser.h
+++ b/Engine/Source/Programs/UnrealHeaderTool/Private/HeaderParser.h
@@ -82,9 +82,9 @@ public:
 	/**
 	 * Sets nesting scope.
 	 */
-	void SetScope(FScope* Scope)
+	void SetScope(FScope* InScope)
 	{
-		this->Scope = Scope;
+		this->Scope = InScope;
 	}
 
 	/** Statement that caused the nesting. */
@@ -129,8 +129,11 @@ class FSimplifiedParsingClassInfo
 {
 public:
 	// Constructor.
-	FSimplifiedParsingClassInfo(FString ClassName, FString BaseClassName, int32 ClassDefLine, bool bClassIsAnInterface)
-		: ClassName(MoveTemp(ClassName)), BaseClassName(MoveTemp(BaseClassName)), ClassDefLine(ClassDefLine), bClassIsAnInterface(bClassIsAnInterface)
+	FSimplifiedParsingClassInfo(FString InClassName, FString InBaseClassName, int32 InClassDefLine, bool bInClassIsAnInterface)
+		: ClassName          (MoveTemp(InClassName))
+		, BaseClassName      (MoveTemp(InBaseClassName))
+		, ClassDefLine       (InClassDefLine)
+		, bClassIsAnInterface(bInClassIsAnInterface)
 	{}
 
 	/**
diff --git a/Engine/Source/Programs/UnrealHeaderTool/Private/UnrealSourceFile.h b/Engine/Source/Programs/UnrealHeaderTool/Private/UnrealSourceFile.h
--- a/Engine/Source/Programs/UnrealHeaderTool/Private/UnrealSourceFile.h
+++ b/Engine/Source/Programs/UnrealHeaderTool/Private/UnrealSourceFile.h
@@ -16,9 +16,14 @@ class FUnrealSourceFile
 {
 public:
 	// Constructor.
-	FUnrealSourceFile(UPackage* Package, FString Filename, FString Content)
-		: Scope(MakeShareable(new FFileScope(*(FString(TEXT("__")) + FPaths::GetBaseFilename(Filename) + FString(TEXT("__File"))), this))),
-		Filename(MoveTemp(Filename)), Package(Package), bHasChanged(false), Content(MoveTemp(Content)), bParsed(false), bDependenciesResolved(false)
+	FUnrealSourceFile(UPackage* InPackage, FString InFilename, FString InContent)
+		: Scope                (MakeShareable(new FFileScope(*(FString(TEXT("__")) + FPaths::GetBaseFilename(Filename) + FString(TEXT("__File"))), this)))
+		, Filename             (MoveTemp(InFilename))
+		, Package              (InPackage)
+		, bHasChanged          (false)
+		, Content              (MoveTemp(InContent))
+		, bParsed              (false)
+		, bDependenciesResolved(false)
 	{
 		if (GetStrippedFilename() != "Object")
 		{
diff --git a/Engine/Source/Programs/UnrealHeaderTool/Private/UnrealTypeDefinitionInfo.h b/Engine/Source/Programs/UnrealHeaderTool/Private/UnrealTypeDefinitionInfo.h
--- a/Engine/Source/Programs/UnrealHeaderTool/Private/UnrealTypeDefinitionInfo.h
+++ b/Engine/Source/Programs/UnrealHeaderTool/Private/UnrealTypeDefinitionInfo.h
@@ -12,8 +12,9 @@ class FUnrealTypeDefinitionInfo
 {
 public:
 	// Constructor
-	FUnrealTypeDefinitionInfo(FUnrealSourceFile& SourceFile, int32 LineNumber)
-		: SourceFile(SourceFile), LineNumber(LineNumber)
+	FUnrealTypeDefinitionInfo(FUnrealSourceFile& InSourceFile, int32 InLineNumber)
+		: SourceFile(InSourceFile)
+		, LineNumber(InLineNumber)
 	{}
 
 	/**
