Shadowing fixes.

diff --git a/Engine/Source/Programs/UnrealHeaderTool/Private/Scope.cpp b/Engine/Source/Programs/UnrealHeaderTool/Private/Scope.cpp
--- a/Engine/Source/Programs/UnrealHeaderTool/Private/Scope.cpp
+++ b/Engine/Source/Programs/UnrealHeaderTool/Private/Scope.cpp
@@ -6,8 +6,8 @@
 
 extern FCompilerMetadataManager GScriptHelper;
 
-FScope::FScope(FScope* Parent)
-	: Parent(Parent)
+FScope::FScope(FScope* InParent)
+	: Parent(InParent)
 {
 	check(Parent);
 }
@@ -143,8 +143,8 @@ bool FScope::ContainsTypes() const
 
 TMap<UStruct*, TSharedRef<FScope> > FScope::ScopeMap;
 
-FFileScope::FFileScope(FName Name, FUnrealSourceFile* SourceFile)
-	: SourceFile(SourceFile), Name(Name)
+FFileScope::FFileScope(FName Name, FUnrealSourceFile* InSourceFile)
+	: SourceFile(InSourceFile), Name(Name)
 {
 
 }
@@ -174,8 +174,8 @@ FName FStructScope::GetName() const
 	return Struct->GetFName();
 }
 
-FStructScope::FStructScope(UStruct* Struct, FScope* Parent)
-	: FScope(Parent), Struct(Struct)
+FStructScope::FStructScope(UStruct* InStruct, FScope* InParent)
+	: FScope(InParent), Struct(InStruct)
 {
 
 }
diff --git a/Engine/Source/Programs/UnrealHeaderTool/Private/UnrealSourceFile.cpp b/Engine/Source/Programs/UnrealHeaderTool/Private/UnrealSourceFile.cpp
--- a/Engine/Source/Programs/UnrealHeaderTool/Private/UnrealSourceFile.cpp
+++ b/Engine/Source/Programs/UnrealHeaderTool/Private/UnrealSourceFile.cpp
@@ -75,24 +75,24 @@ FString FUnrealSourceFile::GetGeneratedBodyMacroName(int32 LineNumber, bool bLeg
 	return GetGeneratedMacroName(LineNumber, *FString::Printf(TEXT("%s%s"), TEXT("_GENERATED_BODY"), bLegacy ? TEXT("_LEGACY") : TEXT("")));
 }
 
-void FUnrealSourceFile::SetGeneratedFilename(FString GeneratedFilename)
+void FUnrealSourceFile::SetGeneratedFilename(FString InGeneratedFilename)
 {
-	this->GeneratedFilename = MoveTemp(GeneratedFilename);
+	GeneratedFilename = MoveTemp(InGeneratedFilename);
 }
 
-void FUnrealSourceFile::SetHasChanged(bool bHasChanged)
+void FUnrealSourceFile::SetHasChanged(bool bInHasChanged)
 {
-	this->bHasChanged = bHasChanged;
+	bHasChanged = bInHasChanged;
 }
 
-void FUnrealSourceFile::SetModuleRelativePath(FString ModuleRelativePath)
+void FUnrealSourceFile::SetModuleRelativePath(FString InModuleRelativePath)
 {
-	this->ModuleRelativePath = MoveTemp(ModuleRelativePath);
+	ModuleRelativePath = MoveTemp(InModuleRelativePath);
 }
 
-void FUnrealSourceFile::SetIncludePath(FString IncludePath)
+void FUnrealSourceFile::SetIncludePath(FString InIncludePath)
 {
-	this->IncludePath = MoveTemp(IncludePath);
+	IncludePath = MoveTemp(InIncludePath);
 }
 
 const FString& FUnrealSourceFile::GetContent() const
