Fix UHT shadow variables

diff --git a/Engine/Source/Programs/UnrealHeaderTool/Private/HeaderParser.cpp b/Engine/Source/Programs/UnrealHeaderTool/Private/HeaderParser.cpp
--- a/Engine/Source/Programs/UnrealHeaderTool/Private/HeaderParser.cpp
+++ b/Engine/Source/Programs/UnrealHeaderTool/Private/HeaderParser.cpp
@@ -3115,8 +3115,8 @@ void FHeaderParser::GetVarType
 
 		if (VariableCategory == EVariableCategory::Member)
 		{
-			auto* UnderlyingType = GEnumUnderlyingTypes.Find(Enum);
-			if (!UnderlyingType || *UnderlyingType != CPT_Byte)
+			auto* EnumUnderlyingType = GEnumUnderlyingTypes.Find(Enum);
+			if (!EnumUnderlyingType || *EnumUnderlyingType != CPT_Byte)
 			{
 				FError::Throwf(TEXT("You cannot use the raw enum name as a type for member variables, instead use TEnumAsByte or a C++11 enum class with an explicit underlying type (currently only uint8 supported)."), *Enum->CppType);
 			}
@@ -4118,7 +4118,7 @@ bool FHeaderParser::SkipDeclaration(FToken& Token)
 	// Store the current value of PrevComment so it can be restored after we parsed everything.
 	FString OldPrevComment(PrevComment);
 	// Consume all tokens until the end of declaration/definition has been found.
-	int32 Nest = 0;
+	int32 NestedScopes = 0;
 	// Check if this is a class/struct declaration in which case it can be followed by member variable declaration.	
 	bool bPossiblyClassDeclaration = Token.Matches(TEXT("class")) || Token.Matches(TEXT("struct"));
 	// (known) macros can end without ; or } so use () to find the end of the declaration.
@@ -4133,13 +4133,13 @@ bool FHeaderParser::SkipDeclaration(FToken& Token)
 	{
 		// If we find parentheses at top-level and we think it's a class declaration then it's more likely
 		// to be something like: class UThing* GetThing();
-		if (bPossiblyClassDeclaration && Nest == 0 && Token.Matches(TEXT("(")))
+		if (bPossiblyClassDeclaration && NestedScopes == 0 && Token.Matches(TEXT("(")))
 		{
 			bPossiblyClassDeclaration = false;
 		}
 
 		bRetestCurrentToken = false;
-		if (Token.Matches(TEXT(";")) && Nest == 0)
+		if (Token.Matches(TEXT(";")) && NestedScopes == 0)
 		{
 			bEndOfDeclarationFound = true;
 			break;
@@ -4149,23 +4149,23 @@ bool FHeaderParser::SkipDeclaration(FToken& Token)
 		{
 			// This is a function definition or class declaration.
 			bDefinitionFound = true;
-			Nest++;
+			NestedScopes++;
 		}
 		else if (Token.Matches(ClosingBracket))
 		{
-			Nest--;
-			if (Nest == 0)
+			NestedScopes--;
+			if (NestedScopes == 0)
 			{
 				bEndOfDeclarationFound = true;
 				break;
 			}
 
-			if (Nest < 0)
+			if (NestedScopes < 0)
 			{
 				FError::Throwf(TEXT("Unexpected '}'. Did you miss a semi-colon?"));
 			}
 		}
-		else if (bMacroDeclaration && Nest == 0)
+		else if (bMacroDeclaration && NestedScopes == 0)
 		{
 			bMacroDeclaration = false;
 			OpeningBracket = TEXT("{");
@@ -4205,7 +4205,7 @@ bool FHeaderParser::SkipDeclaration(FToken& Token)
 	//ClearComment();
 
 	// Successfully consumed C++ declaration unless mismatched pair of brackets has been found.
-	return Nest == 0 && bEndOfDeclarationFound;
+	return NestedScopes == 0 && bEndOfDeclarationFound;
 }
 
 bool FHeaderParser::SafeMatchSymbol( const TCHAR* Match )
@@ -4290,8 +4290,8 @@ FClass* FHeaderParser::ParseClassNameDeclaration(FClasses& AllClasses, FString&
 			{
 				InterfaceName += TEXT('<');
 
-				int32 Nest = 1;
-				while (Nest)
+				int32 NestedScopes = 1;
+				while (NestedScopes)
 				{
 					if (!GetToken(Token))
 						FError::Throwf(TEXT("Unexpected end of file"));
@@ -4300,11 +4300,11 @@ FClass* FHeaderParser::ParseClassNameDeclaration(FClasses& AllClasses, FString&
 					{
 						if (!FCString::Strcmp(Token.Identifier, TEXT("<")))
 						{
-							++Nest;
+							++NestedScopes;
 						}
 						else if (!FCString::Strcmp(Token.Identifier, TEXT(">")))
 						{
-							--Nest;
+							--NestedScopes;
 						}
 					}
 
@@ -4423,7 +4423,7 @@ void FHeaderParser::CompileClassDeclaration(FClasses& AllClasses)
 
 	PushNest(ENestType::NEST_Class, Class);
 	
-	const uint32 PreviousClassFlags = Class->ClassFlags;	
+	const uint32 PrevClassFlags = Class->ClassFlags;
 	ResetClassData();
 
 	// Verify class variables haven't been filled in
@@ -4448,7 +4448,7 @@ void FHeaderParser::CompileClassDeclaration(FClasses& AllClasses)
 	check(ClassData);
 	ClassData->SetPrologLine(PrologFinishLine);
 
-	ClassDeclarationData->MergeAndValidateClassFlags(DeclaredClassName, PreviousClassFlags, Class, AllClasses);
+	ClassDeclarationData->MergeAndValidateClassFlags(DeclaredClassName, PrevClassFlags, Class, AllClasses);
 	Class->SetFlags(RF_Native);
 
 	// Class metadata
@@ -6535,8 +6535,8 @@ ECompilationResult::Type FHeaderParser::ParseAllHeadersInside(
 
 			if (SourceFile.GetPackage() == CurrentPackage && (!SourceFile.IsParsed() || SourceFile.GetDefinedClassesCount() == 0))
 			{
-				ECompilationResult::Type Result;
-				if ((Result = ParseHeaders(ModuleClasses, HeaderParser, SourceFile, true)) != ECompilationResult::Succeeded)
+                Result = ParseHeaders(ModuleClasses, HeaderParser, SourceFile, true);
+				if (Result != ECompilationResult::Succeeded)
 				{
 					return Result;
 				}
@@ -6764,7 +6764,7 @@ void FHeaderParser::SimplifiedClassParse(const TCHAR* InBuffer, TArray<FSimplifi
 				}
 
 				CurrentLine++;
-				const TCHAR* Str = *StrLine;
+				Str = *StrLine;
 				bool bIsPrep = false;
 				if( FParse::Command(&Str,TEXT("#endif")) )
 				{
diff --git a/Engine/Source/Programs/UnrealHeaderTool/Private/HeaderProvider.cpp b/Engine/Source/Programs/UnrealHeaderTool/Private/HeaderProvider.cpp
--- a/Engine/Source/Programs/UnrealHeaderTool/Private/HeaderProvider.cpp
+++ b/Engine/Source/Programs/UnrealHeaderTool/Private/HeaderProvider.cpp
@@ -5,8 +5,8 @@
 #include "UnrealSourceFile.h"
 #include "ClassMaps.h"
 
-FHeaderProvider::FHeaderProvider(EHeaderProviderSourceType Type, const FString& Id, bool bAutoInclude/* = false*/)
-	: Type(Type), Id(Id), Cache(nullptr), bAutoInclude(bAutoInclude)
+FHeaderProvider::FHeaderProvider(EHeaderProviderSourceType InType, const FString& InId, bool bInAutoInclude/* = false*/)
+	: Type(InType), Id(InId), Cache(nullptr), bAutoInclude(bInAutoInclude)
 {
 
 }
diff --git a/Engine/Source/Programs/UnrealHeaderTool/Private/Scope.cpp b/Engine/Source/Programs/UnrealHeaderTool/Private/Scope.cpp
--- a/Engine/Source/Programs/UnrealHeaderTool/Private/Scope.cpp
+++ b/Engine/Source/Programs/UnrealHeaderTool/Private/Scope.cpp
@@ -143,8 +143,8 @@ bool FScope::ContainsTypes() const
 
 TMap<UStruct*, TSharedRef<FScope> > FScope::ScopeMap;
 
-FFileScope::FFileScope(FName Name, FUnrealSourceFile* InSourceFile)
-	: SourceFile(InSourceFile), Name(Name)
+FFileScope::FFileScope(FName InName, FUnrealSourceFile* InSourceFile)
+	: SourceFile(InSourceFile), Name(InName)
 {
 
 }
