Shadowed variable fix.

diff --git a/Engine/Source/Programs/UnrealHeaderTool/Private/UnrealSourceFile.h b/Engine/Source/Programs/UnrealHeaderTool/Private/UnrealSourceFile.h
--- a/Engine/Source/Programs/UnrealHeaderTool/Private/UnrealSourceFile.h
+++ b/Engine/Source/Programs/UnrealHeaderTool/Private/UnrealSourceFile.h
@@ -17,7 +17,7 @@ class FUnrealSourceFile
 public:
 	// Constructor.
 	FUnrealSourceFile(UPackage* InPackage, FString InFilename, FString InContent)
-		: Scope                (MakeShareable(new FFileScope(*(FString(TEXT("__")) + FPaths::GetBaseFilename(Filename) + FString(TEXT("__File"))), this)))
+		: Scope                (MakeShareable(new FFileScope(*(FString(TEXT("__")) + FPaths::GetBaseFilename(InFilename) + FString(TEXT("__File"))), this)))
 		, Filename             (MoveTemp(InFilename))
 		, Package              (InPackage)
 		, bHasChanged          (false)
