Fix for inheriting show/hide categories from parent class.

diff --git a/Engine/Source/Programs/UnrealHeaderTool/Private/ClassDeclarationMetaData.cpp b/Engine/Source/Programs/UnrealHeaderTool/Private/ClassDeclarationMetaData.cpp
--- a/Engine/Source/Programs/UnrealHeaderTool/Private/ClassDeclarationMetaData.cpp
+++ b/Engine/Source/Programs/UnrealHeaderTool/Private/ClassDeclarationMetaData.cpp
@@ -163,26 +163,7 @@ void FClassDeclarationMetaData::ParseClassProperties(const TArray<FPropertySpeci
 
 			for (const FString& Value : PropSpecifier.Values)
 			{
-				// if we didn't find this specific category path in the HideCategories metadata
-				if (HideCategories.Remove(Value) == 0)
-				{
-					TArray<FString> SubCategoryList;
-					Value.ParseIntoArray(SubCategoryList, TEXT("|"), true);
-
-					FString SubCategoryPath;
-					// look to see if any of the parent paths are excluded in the HideCategories list
-					for (int32 CategoryPathIndex = 0; CategoryPathIndex < SubCategoryList.Num() - 1; ++CategoryPathIndex)
-					{
-						SubCategoryPath += SubCategoryList[CategoryPathIndex];
-						// if we're hiding a parent category, then we need to flag this sub category for show
-						if (HideCategories.Contains(SubCategoryPath))
-						{
-							ShowSubCatgories.AddUnique(Value);
-							break;
-						}
-						SubCategoryPath += "|";
-					}
-				}
+				ShowCategories.AddUnique(Value);
 			}
 		}
 		else if (Specifier == TEXT("hideCategories"))
