Fix for the skipping of DECLARE_MULTICAST_DELEGATE declarations.
Bunch of extra checks for this kind of failure in future.
Some scope bracing fixes.

diff --git a/Engine/Source/Programs/UnrealHeaderTool/Private/CodeGenerator.cpp b/Engine/Source/Programs/UnrealHeaderTool/Private/CodeGenerator.cpp
--- a/Engine/Source/Programs/UnrealHeaderTool/Private/CodeGenerator.cpp
+++ b/Engine/Source/Programs/UnrealHeaderTool/Private/CodeGenerator.cpp
@@ -632,6 +632,7 @@ void FNativeClassHeaderGenerator::ExportProperties(UStruct* Struct, int32 TextIn
 			{
 				// find the class info for this class
 				FClassMetaData* ClassData = GScriptHelper.FindClassData(Struct);
+				check(ClassData);
 
 				// find the compiler token for this property
 				FTokenData* PropData = ClassData->FindTokenData(Current);
@@ -1197,9 +1198,6 @@ void FNativeClassHeaderGenerator::ExportNativeGeneratedInitCode(FClass* Class, F
 	// Sort the list of functions
 	FunctionsToExport.Sort();
 
-	// find the class info for this class
-	FClassMetaData* ClassData = GScriptHelper.FindClassData(Class);
-
 	// Export the init code for each function
 	for (int32 FuncIndex = 0; FuncIndex < FunctionsToExport.Num(); FuncIndex++)
 	{
@@ -1659,7 +1657,7 @@ FString GetGeneratedMacroDeprecationWarning(const TCHAR* MacroName)
 FString GetPreservedAccessSpecifierString(FClass* Class)
 {
 	FString PreservedAccessSpecifier(FString::Printf(TEXT("static_assert(false, \"Unknown access specifier for GENERATED_BODY() macro in class %s.\");"), *GetNameSafe(Class)));
-	FClassMetaData* Data = GScriptHelper.FindClassData(Class);	
+	FClassMetaData* Data = GScriptHelper.FindClassData(Class);
 	if (Data)
 	{
 		switch (Data->GeneratedBodyMacroAccessSpecifier)
@@ -1795,13 +1793,14 @@ void FNativeClassHeaderGenerator::ExportClassesFromSourceFileInner(FUnrealSource
 		EnhancedUObjectConstructorsMacroCall.Reset();
 
 		auto* Class = (FClass*)RawUClass;
-		auto* ClassData = GScriptHelper.FindClassData(Class);
-
 		if (Class->ClassFlags & CLASS_Intrinsic)
 		{
 			continue;
 		}
 
+		auto* ClassData = GScriptHelper.FindClassData(Class);
+		check(ClassData);
+
 		// C++ -> VM stubs (native function execs)
 		ExportNativeFunctions(SourceFile, Class, ClassData);
 
@@ -2093,6 +2092,7 @@ void ExportCopyConstructorDefinition(FUHTStringBuilder& Out, FClass* Class, cons
 void ExportVTableHelperConstructorDecl(FStringOutputDevice& Out, FClass* Class, const FString& API)
 {
 	auto* ClassData = GScriptHelper.FindClassData(Class);
+	check(ClassData);
 	if (!ClassData->bCustomVTableHelperConstructorDeclared)
 	{
 		Out.Logf(TEXT("%sDECLARE_VTABLE_PTR_HELPER_CTOR(%s_API, %s);" LINE_TERMINATOR), FCString::Spc(4), *API, NameLookupCPP.GetNameCPP(Class));
@@ -2130,6 +2130,7 @@ void ExportVTableHelperCtorCallerDummy(FStringOutputDevice& Out)
 void ExportVTableHelperConstructorBody(FStringOutputDevice& Out, FClass* Class)
 {
 	auto* ClassData = GScriptHelper.FindClassData(Class);
+	check(ClassData);
 	if (!ClassData->bCustomVTableHelperConstructorDeclared)
 	{
 		Out.Logf(TEXT("%sDEFINE_VTABLE_PTR_HELPER_CTOR(%s);" LINE_TERMINATOR), FCString::Spc(4), NameLookupCPP.GetNameCPP(Class));
@@ -2177,6 +2178,7 @@ void ExportStandardConstructorsMacro(FUHTStringBuilder& Out, FClass* Class, cons
 		Out.Logf(TEXT("\t/** Standard constructor, called after all reflected properties have been initialized */\r\n"));
 
 		auto* ClassData = GScriptHelper.FindClassData(Class);
+		check(ClassData);
 
 		Out.Logf(TEXT("\t%s_API %s(const FObjectInitializer& ObjectInitializer%s);\r\n"), *API, NameLookupCPP.GetNameCPP(Class),
 			ClassData->bDefaultConstructorDeclared ? TEXT("") : TEXT(" = FObjectInitializer::Get()"));
@@ -2200,6 +2202,7 @@ void ExportStandardConstructorsMacro(FUHTStringBuilder& Out, FClass* Class, cons
 void ExportConstructorDefinition(FUHTStringBuilder& Out, FClass* Class, const FString& API)
 {
 	auto* ClassData = GScriptHelper.FindClassData(Class);
+	check(ClassData);
 	if (!ClassData->bConstructorDeclared)
 	{
 		Out.Logf(TEXT("\t/** Standard constructor, called after all reflected properties have been initialized */\r\n"));
@@ -2241,6 +2244,7 @@ void ExportConstructorDefinition(FUHTStringBuilder& Out, FClass* Class, const FS
 void ExportDefaultConstructorCallDefinition(FUHTStringBuilder& Out, FClass* Class)
 {
 	auto* ClassData = GScriptHelper.FindClassData(Class);
+	check(ClassData);
 
 	if (ClassData->bObjectInitializerConstructorDeclared)
 	{
diff --git a/Engine/Source/Programs/UnrealHeaderTool/Private/HeaderParser.cpp b/Engine/Source/Programs/UnrealHeaderTool/Private/HeaderParser.cpp
--- a/Engine/Source/Programs/UnrealHeaderTool/Private/HeaderParser.cpp
+++ b/Engine/Source/Programs/UnrealHeaderTool/Private/HeaderParser.cpp
@@ -68,19 +68,31 @@ namespace
 	bool ProbablyAMacro(const TCHAR* Identifier)
 	{
 		// Test for known delegate and event macros.
-		TCHAR DelegateStart[] = TEXT("DECLARE_DELEGATE_");
+		TCHAR MulticastDelegateStart[] = TEXT("DECLARE_MULTICAST_DELEGATE");
+		if (!FCString::Strncmp(Identifier, MulticastDelegateStart, ARRAY_COUNT(MulticastDelegateStart) - 1))
+		{
+			return true;
+		}
+
+		TCHAR DelegateStart[] = TEXT("DECLARE_DELEGATE");
 		if (!FCString::Strncmp(Identifier, DelegateStart, ARRAY_COUNT(DelegateStart) - 1))
+		{
 			return true;
+		}
 
 		TCHAR DelegateEvent[] = TEXT("DECLARE_EVENT");
 		if (!FCString::Strncmp(Identifier, DelegateEvent, ARRAY_COUNT(DelegateEvent) - 1))
+		{
 			return true;
+		}
 
 		// Failing that, we'll guess about it being a macro based on it being a fully-capitalized identifier.
 		while (TCHAR Ch = *Identifier++)
 		{
 			if (Ch != TEXT('_') && (Ch < TEXT('A') || Ch > TEXT('Z')))
+			{
 				return false;
+			}
 		}
 
 		return true;
@@ -2119,7 +2131,9 @@ void FHeaderParser::FixupDelegateProperties( FClasses& AllClasses, UStruct* Stru
 			{
 				// this UDelegateProperty corresponds to an actual delegate variable (i.e. delegate<SomeDelegate> Foo); we need to lookup the token data for
 				// this property and verify that the delegate property's "type" is an actual delegate function
-				FTokenData* DelegatePropertyToken = GScriptHelper.FindClassData(Struct)->FindTokenData(Property);
+				FClassMetaData* StructData = GScriptHelper.FindClassData(Struct);
+				check(StructData);
+				FTokenData* DelegatePropertyToken = StructData->FindTokenData(Property);
 				check(DelegatePropertyToken);
 
 				// attempt to find the delegate function in the map of functions we've already found
@@ -2251,7 +2265,9 @@ void FHeaderParser::VerifyRepNotifyCallbacks( UClass* TargetClass )
 		UProperty* Prop = Cast<UProperty>(Field);
 		if( Prop && (Prop->GetPropertyFlags() & CPF_RepNotify) )
 		{
-			FTokenData* PropertyToken = GScriptHelper.FindClassData(TargetClass)->FindTokenData(Prop);
+			FClassMetaData* TargetClassData = GScriptHelper.FindClassData(TargetClass);
+			check(TargetClassData);
+			FTokenData* PropertyToken = TargetClassData->FindTokenData(Prop);
 			check(PropertyToken);
 
 			// Search through this class and its superclasses looking for the specified callback
@@ -3856,7 +3872,9 @@ UProperty* FHeaderParser::GetVarNameAndDim
 	}
 
 	VarProperty.TokenProperty = NewProperty;
-	GScriptHelper.FindClassData(Scope)->AddProperty(VarProperty);
+	FClassMetaData* ScopeData = GScriptHelper.FindClassData(Scope);
+	check(ScopeData);
+	ScopeData->AddProperty(VarProperty);
 
 	// if we had any metadata, add it to the class
 	AddMetaDataToClassData(VarProperty.TokenProperty, VarProperty.MetaData);
@@ -4340,13 +4358,17 @@ void FHeaderParser::HandleOneInheritedClass(FClasses& AllClasses, UClass* Class,
 		new (Class->Interfaces) FImplementedInterface(Interface, 0, false);
 		if (Interface->HasAnyClassFlags(CLASS_Native))
 		{
-			GScriptHelper.FindClassData(Class)->AddInheritanceParent(Interface);
+			FClassMetaData* ClassData = GScriptHelper.FindClassData(Class);
+			check(ClassData);
+			ClassData->AddInheritanceParent(Interface);
 		}
 	}
 	else
 	{
 		// Non-UObject inheritance
-		GScriptHelper.FindClassData(Class)->AddInheritanceParent(InterfaceName);
+		FClassMetaData* ClassData = GScriptHelper.FindClassData(Class);
+		check(ClassData);
+		ClassData->AddInheritanceParent(InterfaceName);
 	}
 }
 
@@ -7344,6 +7366,7 @@ bool FHeaderParser::TryToMatchConstructorParameterList(FToken Token)
 	}
 
 	auto* ClassData = GScriptHelper.FindClassData(GetCurrentClass());
+	check(ClassData);
 
 	bool bOICtor = false;
 #if WITH_HOT_RELOAD_CTORS
@@ -7646,7 +7669,9 @@ void FHeaderParser::PostPopFunctionDeclaration(FClasses& AllClasses, UFunction*
 
 void FHeaderParser::PostPopNestInterface(FClasses& AllClasses, UClass* CurrentInterface)
 {
-	if (GScriptHelper.FindClassData(CurrentInterface)->ContainsDelegates())
+	FClassMetaData* ClassData = GScriptHelper.FindClassData(CurrentInterface);
+	check(ClassData);
+	if (ClassData->ContainsDelegates())
 	{
 		TMap<FName, UFunction*> DelegateCache;
 		FixupDelegateProperties(AllClasses, CurrentInterface, FScope::GetTypeScope(ExactCast<UClass>(CurrentInterface)).Get(), DelegateCache);
