Fixing hot-reload not updating derived classes if the base class layout changed resulting in the derived class layout changes without changing the overall class size.

diff --git a/Engine/Source/Programs/UnrealHeaderTool/Private/CodeGenerator.cpp b/Engine/Source/Programs/UnrealHeaderTool/Private/CodeGenerator.cpp
--- a/Engine/Source/Programs/UnrealHeaderTool/Private/CodeGenerator.cpp
+++ b/Engine/Source/Programs/UnrealHeaderTool/Private/CodeGenerator.cpp
@@ -1348,12 +1348,22 @@ void FNativeClassHeaderGenerator::ExportNativeGeneratedInitCode(FClass* Class, F
 	SingletonName.ReplaceInline(TEXT("()"), TEXT(""), ESearchCase::CaseSensitive); // function address
 
 	{	
-		auto ClassNameCPP = NameLookupCPP.GetNameCPP(Class);
+		const TCHAR* ClassNameCPP = NameLookupCPP.GetNameCPP(Class);
 		GeneratedFunctionText.Logf(TEXT("    static FCompiledInDefer Z_CompiledInDefer_UClass_%s(%s, TEXT(\"%s\"));\r\n"), 
 			ClassNameCPP, *SingletonName, ClassNameCPP);
 		
+		// Append base class' CRC at the end of the generated code, this will force update derived classes
+		// when base class changes during hot-reload.
+		uint32 BaseClassCRC = 0;
+		if (Class->GetSuperClass() && !Class->GetSuperClass()->HasAnyClassFlags(CLASS_Intrinsic))
+		{
+			BaseClassCRC = GGeneratedCodeCRCs.FindChecked(Class->GetSuperClass());
+		}
+		GeneratedClassRegisterFunctionText.Logf(TEXT("\r\n// %u\r\n"), BaseClassCRC);
+
 		// Calculate generated class initialization code CRC so that we know when it changes after hot-reload
 		uint32 ClassCrc = GenerateTextCRC(*GeneratedClassRegisterFunctionText);
+		GGeneratedCodeCRCs.Add(Class, ClassCrc);
 		// Emit the IMPLEMENT_CLASS macro to go in the generated cpp file.
 		GeneratedPackageCPP.Logf(TEXT("    IMPLEMENT_CLASS(%s, %u);\r\n"), ClassNameCPP, ClassCrc);
 	}
diff --git a/Engine/Source/Programs/UnrealHeaderTool/Private/ParserHelper.h b/Engine/Source/Programs/UnrealHeaderTool/Private/ParserHelper.h
--- a/Engine/Source/Programs/UnrealHeaderTool/Private/ParserHelper.h
+++ b/Engine/Source/Programs/UnrealHeaderTool/Private/ParserHelper.h
@@ -1794,7 +1794,9 @@ struct FNameLookupCPP
 	{
 		TCHAR* NameCPP = StructNameMap.FindRef( Struct );
 		if (NameCPP && !bForceInterface)
+		{
 			return NameCPP;
+		}
 
 		FString DesiredStructName = Struct->GetName();
 		FString	TempName = FString(bForceInterface ? TEXT("I") : Struct->GetPrefixCPP()) + DesiredStructName;
diff --git a/Engine/Source/Runtime/CoreUObject/Private/UObject/UObjectBase.cpp b/Engine/Source/Runtime/CoreUObject/Private/UObject/UObjectBase.cpp
--- a/Engine/Source/Runtime/CoreUObject/Private/UObject/UObjectBase.cpp
+++ b/Engine/Source/Runtime/CoreUObject/Private/UObject/UObjectBase.cpp
@@ -556,7 +556,7 @@ void UClassCompiledInDefer(FFieldCompiledInInfo* ClassInfo, const TCHAR* Name, S
 	const FName CPPClassName(Name);
 #if WITH_HOT_RELOAD
 	// Check for existing classes
-	auto ExistingClassInfo = GetDeferRegisterClassMap().Find(CPPClassName);
+	FFieldCompiledInInfo** ExistingClassInfo = GetDeferRegisterClassMap().Find(CPPClassName);
 	ClassInfo->bHasChanged = !ExistingClassInfo || (*ExistingClassInfo)->Size != ClassInfo->Size || (*ExistingClassInfo)->Crc != ClassInfo->Crc;
 	if (ExistingClassInfo)
 	{
@@ -565,7 +565,7 @@ void UClassCompiledInDefer(FFieldCompiledInInfo* ClassInfo, const TCHAR* Name, S
 
 		// Get the native name
 		FString NameWithoutPrefix(RemoveClassPrefix(Name));
-		auto ExistingClass = FindObjectChecked<UClass>(ANY_PACKAGE, *NameWithoutPrefix);
+		UClass* ExistingClass = FindObjectChecked<UClass>(ANY_PACKAGE, *NameWithoutPrefix);
 
 		if (ClassInfo->bHasChanged)
 		{
diff --git a/Engine/Source/Runtime/CoreUObject/Public/UObject/UObjectBase.h b/Engine/Source/Runtime/CoreUObject/Public/UObject/UObjectBase.h
--- a/Engine/Source/Runtime/CoreUObject/Public/UObject/UObjectBase.h
+++ b/Engine/Source/Runtime/CoreUObject/Public/UObject/UObjectBase.h
@@ -283,7 +283,7 @@ struct FFieldCompiledInInfo
 	FFieldCompiledInInfo(SIZE_T InClassSize, uint32 InCrc)
 		: Size(InClassSize)
 		, Crc(InCrc)
-		, OldClass(NULL)
+		, OldClass(nullptr)
 		, bHasChanged(false)
 	{
 	}
