Improved error message for missing UHT types.

diff --git a/Engine/Source/Programs/UnrealHeaderTool/Private/HeaderParser.cpp b/Engine/Source/Programs/UnrealHeaderTool/Private/HeaderParser.cpp
--- a/Engine/Source/Programs/UnrealHeaderTool/Private/HeaderParser.cpp
+++ b/Engine/Source/Programs/UnrealHeaderTool/Private/HeaderParser.cpp
@@ -3261,7 +3261,7 @@ void FHeaderParser::GetVarType
 					TempClass = AllClasses.FindScriptClass(InnerClass.Identifier);
 					if (TempClass == nullptr)
 					{
-						FError::Throwf(TEXT("Unrecognized type '%s' (in expression %s<%s>)"), InnerClass.Identifier, VarType.Identifier, InnerClass.Identifier);
+						FError::Throwf(TEXT("Unrecognized type '%s' (in expression %s<%s>) - type must be a UCLASS"), InnerClass.Identifier, VarType.Identifier, InnerClass.Identifier);
 					}
 
 					if (bIsAutoweakPtrTemplate)
@@ -3394,7 +3394,7 @@ void FHeaderParser::GetVarType
 
 			if (!bHandledType)
 			{
-				FError::Throwf(TEXT("Unrecognized type '%s'"), VarType.Identifier );
+				FError::Throwf(TEXT("Unrecognized type '%s' - type must be a UCLASS, USTRUCT or UENUM"), VarType.Identifier );
 			}
 		}
 
