Use PRAGMA_ENABLE_DEPRECATION_WARNINGS instead of PRAGMA_POP in generated code

diff --git a/Engine/Source/Programs/UnrealHeaderTool/Private/CodeGenerator.cpp b/Engine/Source/Programs/UnrealHeaderTool/Private/CodeGenerator.cpp
--- a/Engine/Source/Programs/UnrealHeaderTool/Private/CodeGenerator.cpp
+++ b/Engine/Source/Programs/UnrealHeaderTool/Private/CodeGenerator.cpp
@@ -1270,7 +1270,7 @@ void FNativeClassHeaderGenerator::ExportNativeGeneratedInitCode(FClass* Class, F
 			{
 				GeneratedClassRegisterFunctionText.Logf(TEXT("PRAGMA_DISABLE_DEPRECATION_WARNINGS\r\n"));
 				OutputProperties(Meta, GeneratedClassRegisterFunctionText, OuterString, Props, TEXT("                "));
-				GeneratedClassRegisterFunctionText.Logf(TEXT("PRAGMA_POP\r\n"));
+				GeneratedClassRegisterFunctionText.Logf(TEXT("PRAGMA_ENABLE_DEPRECATION_WARNINGS\r\n"));
 			}
 		}
 		// function table
@@ -1700,7 +1700,7 @@ void ExportInterfaceBodyMacros(FUHTStringBuilder& Out, FUnrealSourceFile& Source
 	auto DeprecationWarning = GetGeneratedMacroDeprecationWarning(TEXT("GENERATED_UINTERFACE_BODY"));
 
 	auto DeprecationPushString = TEXT("PRAGMA_DISABLE_DEPRECATION_WARNINGS") LINE_TERMINATOR;
-	auto DeprecationPopString = TEXT("PRAGMA_POP") LINE_TERMINATOR;
+	auto DeprecationPopString = TEXT("PRAGMA_ENABLE_DEPRECATION_WARNINGS") LINE_TERMINATOR;
 	auto Offset = TEXT("\t");
 
 	Out.Log(Macroize(*SourceFile.GetGeneratedBodyMacroName(ClassGeneratedBodyLine, true), *(FString() +
@@ -2051,7 +2051,7 @@ void FNativeClassHeaderGenerator::ExportClassesFromSourceFileInner(FUnrealSource
 			auto DeprecationWarning = bIsIInterface ? FString(TEXT("")) : GetGeneratedMacroDeprecationWarning(*MacroName);
 
 			auto DeprecationPushString = TEXT("PRAGMA_DISABLE_DEPRECATION_WARNINGS") LINE_TERMINATOR;
-			auto DeprecationPopString = TEXT("PRAGMA_POP") LINE_TERMINATOR;
+			auto DeprecationPopString = TEXT("PRAGMA_ENABLE_DEPRECATION_WARNINGS") LINE_TERMINATOR;
 
 			auto Public = TEXT("public:" LINE_TERMINATOR);
 
