CIS fix

diff --git a/Engine/Source/Programs/UnrealHeaderTool/Private/HeaderParser.cpp b/Engine/Source/Programs/UnrealHeaderTool/Private/HeaderParser.cpp
--- a/Engine/Source/Programs/UnrealHeaderTool/Private/HeaderParser.cpp
+++ b/Engine/Source/Programs/UnrealHeaderTool/Private/HeaderParser.cpp
@@ -3769,7 +3769,7 @@ UProperty* FHeaderParser::GetVarNameAndDim
 		// UFunctions with a smart pointer as input parameter wont compile anyway, because of missing P_GET_... macro.
 		// UFunctions with a smart pointer as return type will crash when called via blueprint, because they are not supported in VM.
 		// WeakPointer is supported by VM as return type (see UObject::execLetWeakObjPtr), but there is no P_GET_... macro for WeakPointer.
-		if ((VarProperty.Type == CPT_LazyObjectReference))
+		if (VarProperty.Type == CPT_LazyObjectReference)
 		{
 			FError::Throwf(TEXT("UFunctions cannot take a lazy pointer as a parameter."));
 		}
