Re-enable warning about virtual Blueprint Implementable Events

diff --git a/Engine/Source/Programs/UnrealHeaderTool/Private/HeaderParser.cpp b/Engine/Source/Programs/UnrealHeaderTool/Private/HeaderParser.cpp
--- a/Engine/Source/Programs/UnrealHeaderTool/Private/HeaderParser.cpp
+++ b/Engine/Source/Programs/UnrealHeaderTool/Private/HeaderParser.cpp
@@ -4985,7 +4985,9 @@ void FHeaderParser::CompileDelegateDeclaration(FUnrealSourceFile& SourceFile, FC
 		// Check the expected versus actual number of parameters
 		int32 ParamCount = FoundParamCount - DelegateParameterCountStrings.GetData() + 1;
 		if (DelegateSignatureFunction->NumParms != ParamCount)
+		{
 			FError::Throwf(TEXT("Expected %d parameters but found %d parameters"), ParamCount, DelegateSignatureFunction->NumParms);
+		}
 	}
 	else
 	{
@@ -5219,7 +5221,7 @@ void FHeaderParser::CompileFunctionDeclaration(FUnrealSourceFile& SourceFile, FC
 
 			else
 			{
-				//UE_LOG(LogCompile, Warning, TEXT("BlueprintImplementableEvents should not be virtual. Use BlueprintNativeEvent instead."));
+				UE_LOG(LogCompile, Warning, TEXT("BlueprintImplementableEvents should not be virtual. Use BlueprintNativeEvent instead."));
 			}
 		}
 	}
