UE-12916: Attempting to use the GameplayTags module in a binary code project will result in build failures in Visual Studio.

diff --git a/Engine/Source/Programs/UnrealHeaderTool/Private/CodeGenerator.cpp b/Engine/Source/Programs/UnrealHeaderTool/Private/CodeGenerator.cpp
--- a/Engine/Source/Programs/UnrealHeaderTool/Private/CodeGenerator.cpp
+++ b/Engine/Source/Programs/UnrealHeaderTool/Private/CodeGenerator.cpp
@@ -729,82 +729,225 @@ void FNativeClassHeaderGenerator::ExportProperties(UStruct* Struct, int32 TextIn
 	}
 }
 
-/** This table maps a singleton name to an extern declaration for it. For cross-module access, we add these to GeneratedFunctionDeclarations **/
-static TMap<FString, FString> SingletonNameToExternDecl;
-
-FString FNativeClassHeaderGenerator::GetSingletonName(UField* Item, bool bRequiresValidObject)
+/**
+ * Class that is representing a type singleton.
+ */
+struct FTypeSingleton
 {
-	FString Suffix;
-	if (UClass* ItemClass = Cast<UClass>(Item))
+public:
+	/** Constructor */
+	FTypeSingleton(FString InName, UField* InType)
+		: Name(MoveTemp(InName)), Type(InType) {}
+
+	/**
+	 * Gets this singleton's name.
+	 */
+	const FString& GetName() const
+	{
+		return Name;
+	}
+
+	/**
+	 * Gets this singleton's extern declaration.
+	 */
+	const FString& GetExternDecl() const
+	{
+		if (ExternDecl.IsEmpty())
+		{
+			ExternDecl = GenerateExternDecl(Type, GetName());
+		}
+
+		return ExternDecl;
+	}
+
+private:
+	/**
+	 * Extern declaration generator.
+	 */
+	static FString GenerateExternDecl(UField* InType, const FString& InName)
 	{
-		if (ItemClass->HasAllClassFlags(CLASS_Intrinsic))
+		const TCHAR* FormatStr = nullptr;
+		
+		if (InType->GetClass() == UClass::StaticClass())
+		{
+			FormatStr = TEXT("    %s_API class UClass* %s;\r\n");
+		}
+		else if (InType->GetClass() == UFunction::StaticClass() || InType->GetClass() == UDelegateFunction::StaticClass())
+		{
+			FormatStr = TEXT("    %s_API class UFunction* %s;\r\n");
+		}
+		else if (InType->GetClass() == UScriptStruct::StaticClass())
+		{
+			FormatStr = TEXT("    %s_API class UScriptStruct* %s;\r\n");
+		}
+		else if (InType->GetClass() == UEnum::StaticClass())
+		{
+			FormatStr = TEXT("    %s_API class UEnum* %s;\r\n");
+		}
+		else
 		{
-			return FString::Printf(TEXT("%s::StaticClass()"), NameLookupCPP.GetNameCPP(ItemClass));
+			FError::Throwf(TEXT("Unsupported item type to get extern for."));
 		}
 
-		if (!bRequiresValidObject && ItemClass->HasAllClassFlags(CLASS_Native))
+		return FString::Printf(FormatStr,
+			*FPackageName::GetShortName(InType->GetOutermost()).ToUpper(),
+			*InName);
+	}
+
+	/** Field that stores this singleton name. */
+	FString Name;
+
+	/** Cached field that stores this singleton extern declaration. */
+	mutable FString ExternDecl;
+
+	/** Type of the singleton */
+	UField* Type;
+};
+
+/**
+ * Class that represents type singleton cache.
+ */
+class FTypeSingletonCache
+{
+public:
+	/**
+	 * Gets type singleton from cache.
+	 *
+	 * @param Type Singleton type.
+	 * @param bRequiresValidObject Does it require a valid object?
+	 */
+	static const FTypeSingleton& Get(UField* Type, bool bRequiresValidObject = true)
+	{
+		static TMap<FTypeSingletonCacheKey, FTypeSingleton> CacheData;
+
+		FTypeSingletonCacheKey Key(Type, bRequiresValidObject);
+		if (FTypeSingleton* SingletonPtr = CacheData.Find(Key))
 		{
-			Suffix = TEXT("_NoRegister");
+			return *SingletonPtr;
 		}
+
+		return CacheData.Add(Key,
+			FTypeSingleton(GenerateSingletonName(Type, bRequiresValidObject), Type)
+		);
 	}
 
-	FString Result;
-	UObject* Outer = Item;
-	while (Outer)
+private:
+	/**
+	 * Private type that represents cache map key.
+	 */
+	struct FTypeSingletonCacheKey
 	{
-		if (Cast<UClass>(Outer) || Cast<UScriptStruct>(Outer))
+		/** FTypeSingleton type */
+		UField* Type;
+
+		/** If this type singleton requires valid object. */
+		bool bRequiresValidObject;
+
+		/* Constructor */
+		FTypeSingletonCacheKey(UField* InType, bool bInRequiresValidObject)
+			: Type(InType), bRequiresValidObject(bInRequiresValidObject)
+		{}
+
+		/**
+		 * Equality operator.
+		 *
+		 * @param Other Other key.
+		 *
+		 * @returns True if this is equal to Other. False otherwise.
+		 */
+		bool operator==(const FTypeSingletonCacheKey& Other) const
 		{
-			FString OuterName = NameLookupCPP.GetNameCPP(Cast<UStruct>(Outer));
-			if (Result.Len())
-			{
-				Result = OuterName + TEXT("_") + Result;
-			}
-			else
+			return Type == Other.Type && bRequiresValidObject == Other.bRequiresValidObject;
+		}
+
+		/**
+		 * Gets hash value for this object.
+		 */
+		friend uint32 GetTypeHash(const FTypeSingletonCacheKey& Object)
+		{
+			return HashCombine(
+				GetTypeHash(Object.Type),
+				GetTypeHash(Object.bRequiresValidObject)
+			);
+		}
+	};
+
+	/**
+	 * Generates singleton name.
+	 */
+	static FString GenerateSingletonName(UField* Item, bool bRequiresValidObject)
+	{
+		FString Suffix;
+		if (UClass* ItemClass = Cast<UClass>(Item))
+		{
+			if (ItemClass->HasAllClassFlags(CLASS_Intrinsic))
 			{
-				Result = OuterName;
+				return FString::Printf(TEXT("%s::StaticClass()"), NameLookupCPP.GetNameCPP(ItemClass));
 			}
-			// Structs can also have UPackage outer.
-			if (Cast<UClass>(Outer) || Cast<UPackage>(Outer->GetOuter()))
+
+			if (!bRequiresValidObject && ItemClass->HasAllClassFlags(CLASS_Native))
 			{
-				break;
+				Suffix = TEXT("_NoRegister");
 			}
 		}
-		else if (Result.Len())
+
+		FString Result;
+		for (UObject* Outer = Item; Outer; Outer = Outer->GetOuter())
 		{
-			// Handle UEnums with UPackage outer.
-			Result = Outer->GetName() + TEXT("_") + Result;
+			if (!Result.IsEmpty())
+			{
+				Result = TEXT("_") + Result;
+			}
+
+			if (Cast<UClass>(Outer) || Cast<UScriptStruct>(Outer))
+			{
+				FString OuterName = NameLookupCPP.GetNameCPP(Cast<UStruct>(Outer));
+				Result = OuterName + Result;
+
+				// Structs can also have UPackage outer.
+				if (Cast<UClass>(Outer) || Cast<UPackage>(Outer->GetOuter()))
+				{
+					break;
+				}
+			}
+			else
+			{
+				Result = Outer->GetName() + Result;
+			}
 		}
-		else
+
+		// Can't use long package names in function names.
+		if (Result.StartsWith(TEXT("/Script/"), ESearchCase::CaseSensitive))
 		{
-			Result = Outer->GetName();
+			Result = FPackageName::GetShortName(Result);
 		}
-		Outer = Outer->GetOuter();
+
+		FString ClassString = NameLookupCPP.GetNameCPP(Item->GetClass());
+		return FString(TEXT("Z_Construct_")) + ClassString + TEXT("_") + Result + Suffix + TEXT("()");
 	}
+};
 
-	FString ClassString = NameLookupCPP.GetNameCPP(Item->GetClass());
-	// Can't use long package names in function names.
-	if (Result.StartsWith(TEXT("/Script/"), ESearchCase::CaseSensitive))
+FString FNativeClassHeaderGenerator::GetSingletonName(UField* Item, bool bRequiresValidObject)
+{
+	FString Result = FTypeSingletonCache::Get(Item, bRequiresValidObject).GetName();
+
+	UClass* ItemClass = Cast<UClass>(Item);
+	if (ItemClass != nullptr && ItemClass->HasAllClassFlags(CLASS_Intrinsic))
 	{
-		Result = FPackageName::GetShortName(Result);
+		return Result;
 	}
-	Result = FString(TEXT("Z_Construct_")) + ClassString + TEXT("_") + Result + Suffix + TEXT("()");
+
 	if (CastChecked<UPackage>(Item->GetOutermost()) != Package)
 	{
 		// this is a cross module reference, we need to include the right extern decl
-		FString* Extern = SingletonNameToExternDecl.Find(Result);
-		if (!Extern)
-		{
-			UE_LOG(LogCompile, Warning, TEXT("Could not find extern declaration for cross module reference %s\n"), *Result);
-		}
-		else
-		{
-			bool bAlreadyInSet = false;
-			UniqueCrossModuleReferences.Add(*Extern, &bAlreadyInSet);
+		FString Extern = FTypeSingletonCache::Get(Item, bRequiresValidObject).GetExternDecl();
 
-			if (!bAlreadyInSet)
-			{
-				CrossModuleGeneratedFunctionDeclarations.Log(*Extern);
-			}
+		bool bAlreadyInSet = false;
+		UniqueCrossModuleReferences.Add(*Extern, &bAlreadyInSet);
+
+		if (!bAlreadyInSet)
+		{
+			CrossModuleGeneratedFunctionDeclarations.Log(*Extern);
 		}
 	}
 	return Result;
@@ -1219,9 +1362,7 @@ void FNativeClassHeaderGenerator::ExportNativeGeneratedInitCode(FClass* Class, F
 		{
 			FString SingletonNameNoRegister(GetSingletonName(Class, false));
 
-			FString Extern = FString::Printf(TEXT("    %sclass UClass* %s;\r\n"), *ApiString, *SingletonNameNoRegister);
-			SingletonNameToExternDecl.Add(SingletonNameNoRegister, Extern);
-			GeneratedFunctionDeclarations.Log(*Extern);
+			GeneratedFunctionDeclarations.Log(FTypeSingletonCache::Get(Class, false).GetExternDecl());
 
 			GeneratedClassRegisterFunctionText.Logf(TEXT("    UClass* %s\r\n"), *SingletonNameNoRegister);
 			GeneratedClassRegisterFunctionText.Logf(TEXT("    {\r\n"));
@@ -1231,9 +1372,7 @@ void FNativeClassHeaderGenerator::ExportNativeGeneratedInitCode(FClass* Class, F
 		FString SingletonName(GetSingletonName(Class));
 
 		OutFriendText.Logf(TEXT("    friend %sclass UClass* %s;\r\n"), *ApiString, *SingletonName);
-		FString Extern = FString::Printf(TEXT("    %sclass UClass* %s;\r\n"), *ApiString, *SingletonName);
-		SingletonNameToExternDecl.Add(SingletonName, Extern);
-		GeneratedFunctionDeclarations.Log(*Extern);
+		GeneratedFunctionDeclarations.Log(FTypeSingletonCache::Get(Class).GetExternDecl());
 
 		GeneratedClassRegisterFunctionText.Logf(TEXT("    UClass* %s\r\n"), *SingletonName);
 		GeneratedClassRegisterFunctionText.Logf(TEXT("    {\r\n"));
@@ -1380,10 +1519,7 @@ void FNativeClassHeaderGenerator::ExportFunction(UFunction* Function, FScope* Sc
 	const FFuncInfo& FunctionData = CompilerInfo->GetFunctionData();
 
 	const FString SingletonName(GetSingletonName(Function));
-
-	FString Extern = FString::Printf(TEXT("    %sclass UFunction* %s;\r\n"), *GetAPIString(), *SingletonName);
-	SingletonNameToExternDecl.Add(SingletonName, Extern);
-	GeneratedFunctionDeclarations.Log(*Extern);
+	GeneratedFunctionDeclarations.Log(FTypeSingletonCache::Get(Function).GetExternDecl());
 
 	FUHTStringBuilder CurrentFunctionText;
 
@@ -2661,11 +2797,7 @@ void FNativeClassHeaderGenerator::ExportGeneratedStructBodyMacros(FUnrealSourceF
 			UStruct* BaseStruct = ScriptStruct->GetSuperStruct();
 
 			const FString SingletonName(GetSingletonName(ScriptStruct));
-
-			FString ApiString = GetAPIString();
-			FString Extern = FString::Printf(TEXT("    %sclass UScriptStruct* %s;\r\n"), *ApiString, *SingletonName);
-			SingletonNameToExternDecl.Add(SingletonName, Extern);
-			GeneratedFunctionDeclarations.Log(*Extern);
+			GeneratedFunctionDeclarations.Log(FTypeSingletonCache::Get(ScriptStruct).GetExternDecl());
 
 			FUHTStringBuilder GeneratedStructRegisterFunctionText;
 
@@ -2791,10 +2923,7 @@ void FNativeClassHeaderGenerator::ExportGeneratedEnumsInitCode(const TArray<UEnu
 
 		{
 			const FString EnumSingletonName = GetSingletonName(Enum);
-
-			FString Extern = FString::Printf(TEXT("    %s_API class UEnum* %s;\r\n"), *API, *EnumSingletonName);
-			SingletonNameToExternDecl.Add(EnumSingletonName, Extern);
-			GeneratedFunctionDeclarations.Log(*Extern);
+			GeneratedFunctionDeclarations.Log(FTypeSingletonCache::Get(Enum).GetExternDecl());
 
 			FUHTStringBuilder GeneratedEnumRegisterFunctionText;
 
