Disallow UClass* properties from having the "Instanced" keyword

diff --git a/Engine/Source/Programs/UnrealHeaderTool/Private/HeaderParser.cpp b/Engine/Source/Programs/UnrealHeaderTool/Private/HeaderParser.cpp
--- a/Engine/Source/Programs/UnrealHeaderTool/Private/HeaderParser.cpp
+++ b/Engine/Source/Programs/UnrealHeaderTool/Private/HeaderParser.cpp
@@ -3507,9 +3507,19 @@ void FHeaderParser::GetVarType
 	}
 
 	// Perform some more specific validation on the property flags
-	if ((VarProperty.PropertyFlags & CPF_PersistentInstance) && VarProperty.Type != CPT_ObjectReference)
+	if (VarProperty.PropertyFlags & CPF_PersistentInstance)
 	{
-		FError::Throwf(TEXT("'Instanced' is only allowed on object property (or array of objects)"));
+		if (VarProperty.Type == CPT_ObjectReference)
+		{
+			if (VarProperty.PropertyClass->IsChildOf<UClass>())
+			{
+				FError::Throwf(TEXT("'Instanced' cannot be applied to class properties (UClass* or TSubclassOf<>)"));
+			}
+		}
+		else
+		{
+			FError::Throwf(TEXT("'Instanced' is only allowed on object property (or array of objects)"));
+		}
 	}
 
 	if ( VarProperty.IsObject() && VarProperty.MetaClass == NULL && (VarProperty.PropertyFlags&CPF_Config) != 0 )
diff --git a/Engine/Source/Runtime/AIModule/Classes/Perception/AIPerceptionComponent.h b/Engine/Source/Runtime/AIModule/Classes/Perception/AIPerceptionComponent.h
--- a/Engine/Source/Runtime/AIModule/Classes/Perception/AIPerceptionComponent.h
+++ b/Engine/Source/Runtime/AIModule/Classes/Perception/AIPerceptionComponent.h
@@ -144,7 +144,7 @@ protected:
 
 	/** Indicated sense that takes precedence over other senses when determining sensed actor's location. 
 	 *	Should be set to one of the sences configured in SensesConfig, or None. */
-	UPROPERTY(EditDefaultsOnly, Instanced, Category = "AI Perception")
+	UPROPERTY(EditDefaultsOnly, Category = "AI Perception")
 	TSubclassOf<UAISense> DominantSense;
 	
 	FAISenseID DominantSenseID;
