Fixing UHT not finding ScriptGeneratorPlugin

diff --git a/Engine/Plugins/ScriptGeneratorPlugin/Source/ScriptGeneratorPlugin/Private/ScriptGeneratorPlugin.cpp b/Engine/Plugins/ScriptGeneratorPlugin/Source/ScriptGeneratorPlugin/Private/ScriptGeneratorPlugin.cpp
--- a/Engine/Plugins/ScriptGeneratorPlugin/Source/ScriptGeneratorPlugin/Private/ScriptGeneratorPlugin.cpp
+++ b/Engine/Plugins/ScriptGeneratorPlugin/Source/ScriptGeneratorPlugin/Private/ScriptGeneratorPlugin.cpp
@@ -5,6 +5,7 @@
 #include "GenericScriptCodeGenerator.h"
 #include "LuaScriptCodeGenerator.h"
 #include "IProjectManager.h"
+#include "Runtime/Core/Public/Features/IModularFeatures.h"
 
 DEFINE_LOG_CATEGORY(LogScriptGenerator);
 
@@ -24,6 +25,7 @@ class FScriptGeneratorPlugin : public IScriptGeneratorPlugin
 	virtual void Initialize(const FString& RootLocalPath, const FString& RootBuildPath, const FString& OutputDirectory, const FString& IncludeBase) override;
 	virtual void ExportClass(UClass* Class, const FString& SourceHeaderFilename, const FString& GeneratedHeaderFilename, bool bHasChanged) override;
 	virtual void FinishExport() override;
+	virtual FString GetGeneratorName() const override;
 };
 
 IMPLEMENT_MODULE( FScriptGeneratorPlugin, ScriptGeneratorPlugin )
@@ -31,11 +33,25 @@ IMPLEMENT_MODULE( FScriptGeneratorPlugin, ScriptGeneratorPlugin )
 
 void FScriptGeneratorPlugin::StartupModule()
 {
+	// Register ourselves as an editor feature
+	IModularFeatures::Get().RegisterModularFeature(TEXT("ScriptGenerator"), this);
 }
 
 void FScriptGeneratorPlugin::ShutdownModule()
 {
 	CodeGenerator.Reset();
+
+	// Unregister our feature
+	IModularFeatures::Get().UnregisterModularFeature(TEXT("ScriptGenerator"), this);
+}
+
+FString FScriptGeneratorPlugin::GetGeneratorName() const
+{
+#if WITH_LUA
+	return TEXT("Lua Example Code Generator Plugin");
+#else
+	return TEXT("Example Code Generator Plugin");
+#endif
 }
 
 void FScriptGeneratorPlugin::Initialize(const FString& RootLocalPath, const FString& RootBuildPath, const FString& OutputDirectory, const FString& IncludeBase)
diff --git a/Engine/Source/Programs/UnrealHeaderTool/Private/CodeGenerator.cpp b/Engine/Source/Programs/UnrealHeaderTool/Private/CodeGenerator.cpp
--- a/Engine/Source/Programs/UnrealHeaderTool/Private/CodeGenerator.cpp
+++ b/Engine/Source/Programs/UnrealHeaderTool/Private/CodeGenerator.cpp
@@ -11,6 +11,7 @@
 #include "Manifest.h"
 #include "StringUtils.h"
 #include "IPluginManager.h"
+#include "Runtime/Core/Public/Features/IModularFeatures.h"
 
 /////////////////////////////////////////////////////
 // Globals
@@ -5234,21 +5235,8 @@ void GetScriptPlugins(TArray<IScriptGeneratorPluginInterface*>& ScriptPlugins)
 {
 	FScopedDurationTimer PluginTimeTracker(GPluginOverheadTime);
 
-	const FString CodeGeneratorPluginCategory(TEXT("UnrealHeaderTool.Code Generator"));
-	const TArray<FPluginStatus> Plugins = IPluginManager::Get().QueryStatusForAllPlugins();
-	for (auto PluginIt(Plugins.CreateConstIterator()); PluginIt; ++PluginIt)
-	{
-		const auto& PluginStatus = *PluginIt;
-		if (PluginStatus.bIsEnabled && PluginStatus.Descriptor.Category.StartsWith(CodeGeneratorPluginCategory))
-		{
-			auto GeneratorInterface = FModuleManager::LoadModulePtr<IScriptGeneratorPluginInterface>(*PluginStatus.Name);
-			if (GeneratorInterface)
-			{
-				UE_LOG(LogCompile, Log, TEXT("Detected script generator plugin (%d): %s"), ScriptPlugins.Num(), *PluginStatus.Name);
-				ScriptPlugins.Add(GeneratorInterface);
-			}
-		}
-	}
+	ScriptPlugins = IModularFeatures::Get().GetModularFeatureImplementations<IScriptGeneratorPluginInterface>(TEXT("ScriptGenerator"));
+	UE_LOG(LogCompile, Log, TEXT("Found %d script generator plugins."), ScriptPlugins.Num());
 
 	// Check if we can use these plugins and initialize them
 	for (int32 PluginIndex = ScriptPlugins.Num() - 1; PluginIndex >= 0; --PluginIndex)
@@ -5271,18 +5259,19 @@ void GetScriptPlugins(TArray<IScriptGeneratorPluginInterface*>& ScriptPlugins)
 			}
 			if (GeneratedCodeModule)
 			{
+				UE_LOG(LogCompile, Log, TEXT("Initializing script generator \'%s\'"), *ScriptGenerator->GetGeneratorName());
 				ScriptGenerator->Initialize(GManifest.RootLocalPath, GManifest.RootBuildPath, GeneratedCodeModule->GeneratedIncludeDirectory, GeneratedCodeModule->IncludeBase);
 			}
 			else
 			{
 				// Can't use this plugin
-				UE_LOG(LogCompile, Log, TEXT("Unable to determine output directory for %s. Cannot export script glue."), *GeneratedCodeModuleName);
+				UE_LOG(LogCompile, Log, TEXT("Unable to determine output directory for %s. Cannot export script glue with \'%s\'"), *GeneratedCodeModuleName, *ScriptGenerator->GetGeneratorName());
 				bSupportedPlugin = false;				
 			}
 		}
 		if (!bSupportedPlugin)
 		{
-			UE_LOG(LogCompile, Log, TEXT("Script generator plugin %d not supported for target: %s"), PluginIndex, *GManifest.TargetName);
+			UE_LOG(LogCompile, Log, TEXT("Script generator \'%s\' not supported for target: %s"), *ScriptGenerator->GetGeneratorName(), *GManifest.TargetName);
 			ScriptPlugins.RemoveAt(PluginIndex);
 		}
 	}
diff --git a/Engine/Source/Programs/UnrealHeaderTool/Public/IScriptGeneratorPluginInterface.h b/Engine/Source/Programs/UnrealHeaderTool/Public/IScriptGeneratorPluginInterface.h
--- a/Engine/Source/Programs/UnrealHeaderTool/Public/IScriptGeneratorPluginInterface.h
+++ b/Engine/Source/Programs/UnrealHeaderTool/Public/IScriptGeneratorPluginInterface.h
@@ -69,7 +69,7 @@ struct EBuildModuleType
 /**
  * The public interface to script generator plugins.
  */
-class IScriptGeneratorPluginInterface : public IModuleInterface
+class IScriptGeneratorPluginInterface : public IModuleInterface, public IModularFeature
 {
 public:
 
@@ -85,5 +85,7 @@ public:
 	virtual void ExportClass(class UClass* Class, const FString& SourceHeaderFilename, const FString& GeneratedHeaderFilename, bool bHasChanged) = 0;
 	/** Called once all classes have been exported */
 	virtual void FinishExport() = 0;
+	/** Name of the generator plugin, mostly for debuggind purposes */
+	virtual FString GetGeneratorName() const = 0;
 };
 
