Prevent UFunctions from having  array parameters.

diff --git a/Engine/Source/Programs/UnrealHeaderTool/Private/HeaderParser.cpp b/Engine/Source/Programs/UnrealHeaderTool/Private/HeaderParser.cpp
--- a/Engine/Source/Programs/UnrealHeaderTool/Private/HeaderParser.cpp
+++ b/Engine/Source/Programs/UnrealHeaderTool/Private/HeaderParser.cpp
@@ -3707,9 +3707,18 @@ UProperty* FHeaderParser::GetVarNameAndDim
 	FToken Dimensions;
 	if (MatchSymbol(TEXT("[")))
 	{
-		if (VariableCategory == EVariableCategory::Return)
+		switch (VariableCategory)
 		{
-			FError::Throwf(TEXT("Arrays aren't allowed in this context") );
+			case EVariableCategory::Return:
+			{
+				FError::Throwf(TEXT("Arrays aren't allowed as return types"));
+			}
+
+			case EVariableCategory::RegularParameter:
+			case EVariableCategory::ReplicatedParameter:
+			{
+				FError::Throwf(TEXT("Arrays aren't allowed as function parameters"));
+			}
 		}
 
 		if (VarProperty.ArrayType == EArrayType::Dynamic || VarProperty.MapKeyProp.IsValid())
