UE-16079: Creating Factory in Runtime Module is not possible and has strange error.

Fixed UHT to detect that a parent class hasn't been parsed before it's child and display a better error message.

diff --git a/Engine/Source/Programs/UnrealHeaderTool/Private/CodeGenerator.cpp b/Engine/Source/Programs/UnrealHeaderTool/Private/CodeGenerator.cpp
--- a/Engine/Source/Programs/UnrealHeaderTool/Private/CodeGenerator.cpp
+++ b/Engine/Source/Programs/UnrealHeaderTool/Private/CodeGenerator.cpp
@@ -5277,6 +5277,60 @@ void GetScriptPlugins(TArray<IScriptGeneratorPluginInterface*>& ScriptPlugins)
 	}
 }
 
+/**
+ * Tries to resolve super classes for classes defined in the given
+ * module.
+ *
+ * @param Package Modules package.
+ */
+void ResolveSuperClasses(UPackage* Package)
+{
+	TArray<UObject*> Objects;
+	GetObjectsWithOuter(Package, Objects);
+
+	for (auto* Object : Objects)
+	{
+		if (!Object->IsA<UClass>())
+		{
+			continue;
+		}
+
+		UClass* DefinedClass = Cast<UClass>(Object);
+
+		if (DefinedClass->HasAnyClassFlags(CLASS_Intrinsic | CLASS_NoExport))
+		{
+			continue;
+		}
+
+		const FSimplifiedParsingClassInfo& ParsingInfo = GTypeDefinitionInfoMap[DefinedClass]->GetUnrealSourceFile()
+			.GetDefinedClassParsingInfo(DefinedClass);
+
+		const FString& BaseClassNameStripped = GetClassNameWithPrefixRemoved(ParsingInfo.GetBaseClassName());
+
+		if (!BaseClassNameStripped.IsEmpty() && !DefinedClass->GetSuperClass())
+		{
+			UClass* FoundBaseClass = FindObject<UClass>(Package, *BaseClassNameStripped);
+
+			if (FoundBaseClass == nullptr)
+			{
+				FoundBaseClass = FindObject<UClass>(ANY_PACKAGE, *BaseClassNameStripped);
+			}
+
+			if (FoundBaseClass == nullptr)
+			{
+				// Don't know its parent class. Raise error.
+				UE_LOG(LogCompile, Error, TEXT("Couldn't find parent type for '%s' named '%s' in current module or any other module parsed so far."),
+					*DefinedClass->GetName(), *ParsingInfo.GetBaseClassName());
+			}
+			else
+			{
+				DefinedClass->SetSuperStruct(FoundBaseClass);
+				DefinedClass->ClassCastFlags |= FoundBaseClass->ClassCastFlags;
+			}
+		}
+	}
+}
+
 ECompilationResult::Type UnrealHeaderTool_Main(const FString& ModuleInfoFilename)
 {
 	check(GIsUCCMakeStandaloneHeaderGenerator);
@@ -5460,6 +5514,8 @@ ECompilationResult::Type UnrealHeaderTool_Main(const FString& ModuleInfoFilename
 			}
 		}
 
+		ResolveSuperClasses(Package);
+
 		ThisModuleTimer.Stop();
 		TotalModulePreparseTime += ThisModulePreparseTime;
 		UE_LOG(LogCompile, Log, TEXT("Preparsed module %s containing %i files(s) in %.2f secs."), *Module.LongPackageName, NumHeadersPreparsed, ThisModulePreparseTime);
@@ -5680,30 +5736,6 @@ UClass* ProcessParsedClass(bool bClassIsAnInterface, TArray<FHeaderProvider> &De
 		}
 	}
 
-	if (ResultClass != nullptr)
-	{
-		if (!BaseClassNameStripped.IsEmpty() && !ResultClass->GetSuperClass())
-		{
-			// Find or forward-declare base class.
-			ResultClass->SetSuperStruct(FindObject<UClass>(InParent, *BaseClassNameStripped));
-			if (ResultClass->GetSuperStruct() == nullptr)
-			{
-				ResultClass->SetSuperStruct(FindObject<UClass>(ANY_PACKAGE, *BaseClassNameStripped));
-			}
-
-			if (ResultClass->GetSuperStruct() == nullptr)
-			{
-				// don't know its parent class yet
-				ResultClass->SetSuperStruct(new(EC_InternalUseOnlyConstructor, InParent, *BaseClassNameStripped) UClass(FObjectInitializer(), nullptr));
-			}
-
-			if (ResultClass->GetSuperStruct() != nullptr)
-			{
-				ResultClass->ClassCastFlags |= ResultClass->GetSuperClass()->ClassCastFlags;
-			}
-		}
-	}
-
 	if (bVerboseOutput)
 	{
 		for (const auto& Dependency : DependentOn)
@@ -5739,7 +5771,7 @@ TSharedRef<FUnrealSourceFile> PerformInitialParseOnHeader(UPackage* InParent, co
 		FScope::AddTypeScope(ResultClass, &UnrealSourceFile->GetScope().Get());
 
 		AddTypeDefinition(*UnrealSourceFile, ResultClass, ParsedClassInfo.GetClassDefLine());
-		UnrealSourceFile->AddDefinedClass(ResultClass);
+		UnrealSourceFile->AddDefinedClass(ResultClass, MoveTemp(ParsedClassInfo));
 	}
 
 	for (const auto& DependsOnElement : DependsOn)
diff --git a/Engine/Source/Programs/UnrealHeaderTool/Private/HeaderParser.h b/Engine/Source/Programs/UnrealHeaderTool/Private/HeaderParser.h
--- a/Engine/Source/Programs/UnrealHeaderTool/Private/HeaderParser.h
+++ b/Engine/Source/Programs/UnrealHeaderTool/Private/HeaderParser.h
@@ -122,66 +122,6 @@ struct ClassDefinitionRange
 
 extern TMap<UClass*, ClassDefinitionRange> ClassDefinitionRanges;
 
-/**
- * Simplified parsing information about UClasses.
- */
-class FSimplifiedParsingClassInfo
-{
-public:
-	// Constructor.
-	FSimplifiedParsingClassInfo(FString InClassName, FString InBaseClassName, int32 InClassDefLine, bool bInClassIsAnInterface)
-		: ClassName          (MoveTemp(InClassName))
-		, BaseClassName      (MoveTemp(InBaseClassName))
-		, ClassDefLine       (InClassDefLine)
-		, bClassIsAnInterface(bInClassIsAnInterface)
-	{}
-
-	/**
-	 * Gets class name.
-	 */
-	const FString& GetClassName() const
-	{
-		return ClassName;
-	}
-
-	/**
-	 * Gets base class name.
-	 */
-	const FString& GetBaseClassName() const
-	{
-		return BaseClassName;
-	}
-
-	/**
-	 * Gets class definition line number.
-	 */
-	int32 GetClassDefLine() const
-	{
-		return ClassDefLine;
-	}
-
-	/**
-	 * Tells if this class is an interface.
-	 */
-	bool IsInterface() const
-	{
-		return bClassIsAnInterface;
-	}
-
-private:
-	// Name.
-	FString ClassName;
-
-	// Super-class name.
-	FString BaseClassName;
-
-	// Class definition line.
-	int32 ClassDefLine;
-
-	// Is this an interface class?
-	bool bClassIsAnInterface;
-};
-
 /////////////////////////////////////////////////////
 // FHeaderParser
 
diff --git a/Engine/Source/Programs/UnrealHeaderTool/Private/SimplifiedParsingClassInfo.h b/Engine/Source/Programs/UnrealHeaderTool/Private/SimplifiedParsingClassInfo.h
new file mode 100644
--- /dev/null
+++ b/Engine/Source/Programs/UnrealHeaderTool/Private/SimplifiedParsingClassInfo.h
@@ -0,0 +1,65 @@
+// Copyright 1998-2015 Epic Games, Inc. All Rights Reserved.
+
+#pragma once
+
+#include "UnrealString.h"
+
+/**
+ * Simplified parsing information about UClasses.
+ */
+class FSimplifiedParsingClassInfo
+{
+public:
+	// Constructor.
+	FSimplifiedParsingClassInfo(FString InClassName, FString InBaseClassName, int32 InClassDefLine, bool bInClassIsAnInterface)
+		: ClassName          (MoveTemp(InClassName))
+		, BaseClassName      (MoveTemp(InBaseClassName))
+		, ClassDefLine       (InClassDefLine)
+		, bClassIsAnInterface(bInClassIsAnInterface)
+	{}
+
+	/**
+	 * Gets class name.
+	 */
+	const FString& GetClassName() const
+	{
+		return ClassName;
+	}
+
+	/**
+	 * Gets base class name.
+	 */
+	const FString& GetBaseClassName() const
+	{
+		return BaseClassName;
+	}
+
+	/**
+	 * Gets class definition line number.
+	 */
+	int32 GetClassDefLine() const
+	{
+		return ClassDefLine;
+	}
+
+	/**
+	 * Tells if this class is an interface.
+	 */
+	bool IsInterface() const
+	{
+		return bClassIsAnInterface;
+	}
+
+private:
+	// Name.
+	FString ClassName;
+
+	// Super-class name.
+	FString BaseClassName;
+
+	// Class definition line.
+	int32 ClassDefLine;
+
+	// Is this an interface class?
+	bool bClassIsAnInterface;
+};
\ No newline at end of file
diff --git a/Engine/Source/Programs/UnrealHeaderTool/Private/UnrealSourceFile.cpp b/Engine/Source/Programs/UnrealHeaderTool/Private/UnrealSourceFile.cpp
--- a/Engine/Source/Programs/UnrealHeaderTool/Private/UnrealSourceFile.cpp
+++ b/Engine/Source/Programs/UnrealHeaderTool/Private/UnrealSourceFile.cpp
@@ -6,9 +6,9 @@
 #include "HeaderParser.h"
 #include "GeneratedCodeVersion.h"
 
-void FUnrealSourceFile::AddDefinedClass(UClass* Class)
+void FUnrealSourceFile::AddDefinedClass(UClass* Class, FSimplifiedParsingClassInfo ParsingInfo)
 {
-	DefinedClasses.Add(Class);
+	DefinedClasses.Add(Class, MoveTemp(ParsingInfo));
 }
 
 FString FUnrealSourceFile::GetFileId() const
diff --git a/Engine/Source/Programs/UnrealHeaderTool/Private/UnrealSourceFile.h b/Engine/Source/Programs/UnrealHeaderTool/Private/UnrealSourceFile.h
--- a/Engine/Source/Programs/UnrealHeaderTool/Private/UnrealSourceFile.h
+++ b/Engine/Source/Programs/UnrealHeaderTool/Private/UnrealSourceFile.h
@@ -6,6 +6,7 @@
 #include "HeaderProvider.h"
 #include "UnrealTypeDefinitionInfo.h"
 #include "GeneratedCodeVersion.h"
+#include "SimplifiedParsingClassInfo.h"
 
 class FClassMetaData;
 
@@ -38,9 +39,9 @@ public:
 	 */
 	void AppendDefinedClasses(TArray<UClass*>& OutClasses) const
 	{
-		for (auto* Class : DefinedClasses)
+		for (const auto& KeyValuePair : DefinedClasses)
 		{
-			OutClasses.Add(Class);
+			OutClasses.Add(KeyValuePair.Key);
 		}
 	}
 
@@ -49,7 +50,31 @@ public:
 	 *
 	 * @returns Array with classes defined in this source file.
 	 */
-	const TArray<UClass*>& GetDefinedClasses() const
+	TArray<UClass*> GetDefinedClasses() const
+	{
+		TArray<UClass*> Array;
+
+		AppendDefinedClasses(Array);
+
+		return Array;
+	}
+
+	/**
+	 * Gets parsing info for class defined in this source file.
+	 *
+	 * @returns Parsing info for class defined in this source file.
+	 */
+	const FSimplifiedParsingClassInfo& GetDefinedClassParsingInfo(UClass* DefinedClass) const
+	{
+		return GetDefinedClassesWithParsingInfo()[DefinedClass];
+	}
+
+	/**
+	 * Gets map with classes defined in this source file with parsing info.
+	 *
+	 * @returns Map with classes defined in this source file with parsing info.
+	 */
+	const TMap<UClass*, FSimplifiedParsingClassInfo>& GetDefinedClassesWithParsingInfo() const
 	{
 		return DefinedClasses;
 	}
@@ -126,8 +151,9 @@ public:
 	 * Adds given class to class definition list for this source file.
 	 *
 	 * @param Class Class to add to list.
+	 * @param ParsingInfo Data from simplified parsing step.
 	 */
-	void AddDefinedClass(UClass* Class);
+	void AddDefinedClass(UClass* Class, FSimplifiedParsingClassInfo ParsingInfo);
 
 	/**
 	 * Gets scope for this file.
@@ -290,8 +316,8 @@ private:
 	// This source file includes.
 	TArray<FHeaderProvider> Includes;
 
-	// List of classes defined in this source file.
-	TArray<UClass*> DefinedClasses;
+	// List of classes defined in this source file along with parsing info.
+	TMap<UClass*, FSimplifiedParsingClassInfo> DefinedClasses;
 
 	// Mapping of UStructs to versions, according to which their code should be generated.
 	TMap<UStruct*, EGeneratedCodeVersion> GeneratedCodeVersions;
