Fix missing metadata (ConversionRoot) after UHT parsing.

diff --git a/Engine/Source/Programs/UnrealHeaderTool/Private/HeaderParser.cpp b/Engine/Source/Programs/UnrealHeaderTool/Private/HeaderParser.cpp
--- a/Engine/Source/Programs/UnrealHeaderTool/Private/HeaderParser.cpp
+++ b/Engine/Source/Programs/UnrealHeaderTool/Private/HeaderParser.cpp
@@ -4480,6 +4480,7 @@ void FHeaderParser::CompileClassDeclaration(FClasses& AllClasses)
 	Class->SetFlags(RF_Native);
 
 	// Class metadata
+	MetaData.Append(ClassDeclarationData->MetaData);
 	if (ClassDeclarationData->ClassGroupNames.Num()) { MetaData.Add("ClassGroupNames", FString::Join(ClassDeclarationData->ClassGroupNames, TEXT(" "))); }
 	if (ClassDeclarationData->AutoCollapseCategories.Num()) { MetaData.Add("AutoCollapseCategories", FString::Join(ClassDeclarationData->AutoCollapseCategories, TEXT(" "))); }
 	if (ClassDeclarationData->HideCategories.Num()) { MetaData.Add("HideCategories", FString::Join(ClassDeclarationData->HideCategories, TEXT(" "))); }
