New IsSorted algorithm.
New FindSortedStringCaseInsensitive algorithm, which binary searches strings in an array.
Many UHT string tests replaced with FindSortedStringCaseInsensitive.
FPropertySpecifier moved to UHT.

diff --git a/Engine/Source/Programs/UnrealHeaderTool/Private/BaseParser.cpp b/Engine/Source/Programs/UnrealHeaderTool/Private/BaseParser.cpp
--- a/Engine/Source/Programs/UnrealHeaderTool/Private/BaseParser.cpp
+++ b/Engine/Source/Programs/UnrealHeaderTool/Private/BaseParser.cpp
@@ -117,6 +117,45 @@ namespace
 	}
 }
 
+//////////////////////////////////////////////////////////////////////////
+// FPropertySpecifier
+
+FString FPropertySpecifier::ConvertToString() const
+{
+	FString Result;
+
+	// Emit the specifier key
+	Result += Key;
+
+	// Emit the values if there are any
+	if (Values.Num())
+	{
+		Result += TEXT("=");
+
+		if (Values.Num() == 1)
+		{
+			// One value goes on it's own
+			Result += Values[0];
+		}
+		else
+		{
+			// More than one value goes in parens, separated by commas
+			Result += TEXT("(");
+			for (int32 ValueIndex = 0; ValueIndex < Values.Num(); ++ValueIndex)
+			{
+				if (ValueIndex > 0)
+				{
+					Result += TEXT(", ");
+				}
+				Result += Values[ValueIndex];
+			}
+			Result += TEXT(")");
+		}
+	}
+
+	return Result;
+}
+
 /////////////////////////////////////////////////////
 // FBaseParser
 
@@ -1008,16 +1047,16 @@ void FBaseParser::ReadSpecifierSetInsideMacro(TArray<FPropertySpecifier>& Specif
 		else
 		{
 			// Creating a new specifier
-			FPropertySpecifier* NewPair = new (SpecifiersFound) FPropertySpecifier();
-			NewPair->Key = Specifier.Identifier;
+			SpecifiersFound.Emplace(Specifier.Identifier);
 
 			// Look for a value for this specifier
 			if (MatchSymbol(TEXT("=")) || PeekSymbol(TEXT("(")))
 			{
-				if (!ReadOptionalCommaSeparatedListInParens(NewPair->Values, TypeOfSpecifier))
+				TArray<FString>& NewPairValues = SpecifiersFound.Last().Values;
+				if (!ReadOptionalCommaSeparatedListInParens(NewPairValues, TypeOfSpecifier))
 				{
 					FString Value = ReadNewStyleValue(TypeOfSpecifier);
-					NewPair->Values.Add(Value);
+					NewPairValues.Add(Value);
 				}
 			}
 		}
diff --git a/Engine/Source/Programs/UnrealHeaderTool/Private/BaseParser.h b/Engine/Source/Programs/UnrealHeaderTool/Private/BaseParser.h
--- a/Engine/Source/Programs/UnrealHeaderTool/Private/BaseParser.h
+++ b/Engine/Source/Programs/UnrealHeaderTool/Private/BaseParser.h
@@ -14,6 +14,26 @@ enum class ESymbolParseOption
 	CloseTemplateBracket
 };
 
+// A specifier with optional value
+struct FPropertySpecifier
+{
+public:
+	explicit FPropertySpecifier(FString&& InKey)
+		: Key(MoveTemp(InKey))
+	{
+	}
+
+	explicit FPropertySpecifier(const FString& InKey)
+		: Key(InKey)
+	{
+	}
+
+	FString Key;
+	TArray<FString> Values;
+
+	FString ConvertToString() const;
+};
+
 //
 // Base class of header parsers.
 //
diff --git a/Engine/Source/Programs/UnrealHeaderTool/Private/ClassDeclarationMetaData.h b/Engine/Source/Programs/UnrealHeaderTool/Private/ClassDeclarationMetaData.h
--- a/Engine/Source/Programs/UnrealHeaderTool/Private/ClassDeclarationMetaData.h
+++ b/Engine/Source/Programs/UnrealHeaderTool/Private/ClassDeclarationMetaData.h
@@ -2,6 +2,8 @@
 
 #pragma once
 
+struct FPropertySpecifier;
+
 /** Structure that holds class meta data generated from its UCLASS declaration */
 class FClassDeclarationMetaData
 {
diff --git a/Engine/Source/Programs/UnrealHeaderTool/Private/HeaderParser.cpp b/Engine/Source/Programs/UnrealHeaderTool/Private/HeaderParser.cpp
--- a/Engine/Source/Programs/UnrealHeaderTool/Private/HeaderParser.cpp
+++ b/Engine/Source/Programs/UnrealHeaderTool/Private/HeaderParser.cpp
@@ -14,6 +14,14 @@
 #include "UnitConversion.h"
 #include "GeneratedCodeVersion.h"
 
+#include "Algo/FindSortedStringCaseInsensitive.h"
+
+#include "Specifiers/CheckedMetadataSpecifiers.h"
+#include "Specifiers/FunctionSpecifiers.h"
+#include "Specifiers/InterfaceSpecifiers.h"
+#include "Specifiers/StructSpecifiers.h"
+#include "Specifiers/VariableSpecifiers.h"
+
 double GPluginOverheadTime = 0.0;
 double GHeaderCodeGenTime = 0.0;
 
@@ -176,171 +184,209 @@ namespace
 
 		for (const auto& Specifier : Specifiers)
 		{
-			if (Specifier.Key == TEXT("BlueprintNativeEvent"))
+			switch ((EFunctionSpecifier)Algo::FindSortedStringCaseInsensitive(*Specifier.Key, GFunctionSpecifierStrings, ARRAY_COUNT(GFunctionSpecifierStrings)))
 			{
-				if (FuncInfo.FunctionFlags & FUNC_Net)
+				default:
 				{
-					FError::Throwf(TEXT("BlueprintNativeEvent functions cannot be replicated!") );
+					FError::Throwf(TEXT("Unknown function specifier '%s'"), *Specifier.Key);
 				}
-				else if ( (FuncInfo.FunctionFlags & FUNC_BlueprintEvent) && !(FuncInfo.FunctionFlags & FUNC_Native) )
+				break;
+
+				case EFunctionSpecifier::BlueprintNativeEvent:
 				{
-					// already a BlueprintImplementableEvent
-					FError::Throwf(TEXT("A function cannot be both BlueprintNativeEvent and BlueprintImplementableEvent!") );
+					if (FuncInfo.FunctionFlags & FUNC_Net)
+					{
+						FError::Throwf(TEXT("BlueprintNativeEvent functions cannot be replicated!") );
+					}
+					else if ( (FuncInfo.FunctionFlags & FUNC_BlueprintEvent) && !(FuncInfo.FunctionFlags & FUNC_Native) )
+					{
+						// already a BlueprintImplementableEvent
+						FError::Throwf(TEXT("A function cannot be both BlueprintNativeEvent and BlueprintImplementableEvent!") );
+					}
+					else if ( (FuncInfo.FunctionFlags & FUNC_Private) )
+					{
+						FError::Throwf(TEXT("A Private function cannot be a BlueprintNativeEvent!") );
+					}
+
+					FuncInfo.FunctionFlags |= FUNC_Event;
+					FuncInfo.FunctionFlags |= FUNC_BlueprintEvent;
 				}
-				else if ( (FuncInfo.FunctionFlags & FUNC_Private) )
+				break;
+
+				case EFunctionSpecifier::BlueprintImplementableEvent:
 				{
-					FError::Throwf(TEXT("A Private function cannot be a BlueprintNativeEvent!") );
+					if (FuncInfo.FunctionFlags & FUNC_Net)
+					{
+						FError::Throwf(TEXT("BlueprintImplementableEvent functions cannot be replicated!") );
+					}
+					else if ( (FuncInfo.FunctionFlags & FUNC_BlueprintEvent) && (FuncInfo.FunctionFlags & FUNC_Native) )
+					{
+						// already a BlueprintNativeEvent
+						FError::Throwf(TEXT("A function cannot be both BlueprintNativeEvent and BlueprintImplementableEvent!") );
+					}
+					else if ( (FuncInfo.FunctionFlags & FUNC_Private) )
+					{
+						FError::Throwf(TEXT("A Private function cannot be a BlueprintImplementableEvent!") );
+					}
+
+					FuncInfo.FunctionFlags |= FUNC_Event;
+					FuncInfo.FunctionFlags |= FUNC_BlueprintEvent;
+					FuncInfo.FunctionFlags &= ~FUNC_Native;
 				}
+				break;
 
-				FuncInfo.FunctionFlags |= FUNC_Event;
-				FuncInfo.FunctionFlags |= FUNC_BlueprintEvent;
-			}
-			else if (Specifier.Key == TEXT("BlueprintImplementableEvent"))
-			{
-				if (FuncInfo.FunctionFlags & FUNC_Net)
+				case EFunctionSpecifier::Exec:
 				{
-					FError::Throwf(TEXT("BlueprintImplementableEvent functions cannot be replicated!") );
+					FuncInfo.FunctionFlags |= FUNC_Exec;
+					if( FuncInfo.FunctionFlags & FUNC_Net )
+					{
+						FError::Throwf(TEXT("Exec functions cannot be replicated!") );
+					}
 				}
-				else if ( (FuncInfo.FunctionFlags & FUNC_BlueprintEvent) && (FuncInfo.FunctionFlags & FUNC_Native) )
+				break;
+
+				case EFunctionSpecifier::SealedEvent:
 				{
-					// already a BlueprintNativeEvent
-					FError::Throwf(TEXT("A function cannot be both BlueprintNativeEvent and BlueprintImplementableEvent!") );
+					FuncInfo.bSealedEvent = true;
 				}
-				else if ( (FuncInfo.FunctionFlags & FUNC_Private) )
+				break;
+
+				case EFunctionSpecifier::Server:
 				{
-					FError::Throwf(TEXT("A Private function cannot be a BlueprintImplementableEvent!") );
+					if ((FuncInfo.FunctionFlags & FUNC_BlueprintEvent) != 0)
+					{
+						FError::Throwf(TEXT("BlueprintImplementableEvent or BlueprintNativeEvent functions cannot be declared as Client or Server"));
+					}
+
+					FuncInfo.FunctionFlags |= FUNC_Net;
+					FuncInfo.FunctionFlags |= FUNC_NetServer;
+
+					if (Specifier.Values.Num())
+					{
+						FuncInfo.CppImplName = Specifier.Values[0];
+					}
+
+					if( FuncInfo.FunctionFlags & FUNC_Exec )
+					{
+						FError::Throwf(TEXT("Exec functions cannot be replicated!") );
+					}
 				}
+				break;
 
-				FuncInfo.FunctionFlags |= FUNC_Event;
-				FuncInfo.FunctionFlags |= FUNC_BlueprintEvent;
-				FuncInfo.FunctionFlags &= ~FUNC_Native;
-			}
-			else if (Specifier.Key == TEXT("Exec"))
-			{
-				FuncInfo.FunctionFlags |= FUNC_Exec;
-				if( FuncInfo.FunctionFlags & FUNC_Net )
+				case EFunctionSpecifier::Client:
 				{
-					FError::Throwf(TEXT("Exec functions cannot be replicated!") );
+					if ((FuncInfo.FunctionFlags & FUNC_BlueprintEvent) != 0)
+					{
+						FError::Throwf(TEXT("BlueprintImplementableEvent or BlueprintNativeEvent functions cannot be declared as Client or Server"));
+					}
+
+					FuncInfo.FunctionFlags |= FUNC_Net;
+					FuncInfo.FunctionFlags |= FUNC_NetClient;
+
+					if (Specifier.Values.Num())
+					{
+						FuncInfo.CppImplName = Specifier.Values[0];
+					}
 				}
-			}
-			else if (Specifier.Key == TEXT("SealedEvent"))
-			{
-				FuncInfo.bSealedEvent = true;
-			}
-			else if (Specifier.Key == TEXT("Server"))
-			{
-				if ((FuncInfo.FunctionFlags & FUNC_BlueprintEvent) != 0)
+				break;
+
+				case EFunctionSpecifier::NetMulticast:
 				{
-					FError::Throwf(TEXT("BlueprintImplementableEvent or BlueprintNativeEvent functions cannot be declared as Client or Server"));
-				}
+					if ((FuncInfo.FunctionFlags & FUNC_BlueprintEvent) != 0)
+					{
+						FError::Throwf(TEXT("BlueprintImplementableEvent or BlueprintNativeEvent functions cannot be declared as Multicast"));
+					}
 
-				FuncInfo.FunctionFlags |= FUNC_Net;
-				FuncInfo.FunctionFlags |= FUNC_NetServer;
+					FuncInfo.FunctionFlags |= FUNC_Net;
+					FuncInfo.FunctionFlags |= FUNC_NetMulticast;
+				}
+				break;
 
-				if (Specifier.Values.Num())
+				case EFunctionSpecifier::ServiceRequest:
 				{
-					FuncInfo.CppImplName = Specifier.Values[0];
+					if ((FuncInfo.FunctionFlags & FUNC_BlueprintEvent) != 0)
+					{
+						FError::Throwf(TEXT("BlueprintImplementableEvent or BlueprintNativeEvent functions cannot be declared as a ServiceRequest"));
+					}
+
+					FuncInfo.FunctionFlags |= FUNC_Net;
+					FuncInfo.FunctionFlags |= FUNC_NetReliable;
+					FuncInfo.FunctionFlags |= FUNC_NetRequest;
+					FuncInfo.FunctionExportFlags |= FUNCEXPORT_CustomThunk;
+
+					ParseNetServiceIdentifiers(FuncInfo, Specifier.Values);
 				}
+				break;
 
-				if( FuncInfo.FunctionFlags & FUNC_Exec )
+				case EFunctionSpecifier::ServiceResponse:
 				{
-					FError::Throwf(TEXT("Exec functions cannot be replicated!") );
+					if ((FuncInfo.FunctionFlags & FUNC_BlueprintEvent) != 0)
+					{
+						FError::Throwf(TEXT("BlueprintImplementableEvent or BlueprintNativeEvent functions cannot be declared as a ServiceResponse"));
+					}
+
+					FuncInfo.FunctionFlags |= FUNC_Net;
+					FuncInfo.FunctionFlags |= FUNC_NetReliable;
+					FuncInfo.FunctionFlags |= FUNC_NetResponse;
+
+					ParseNetServiceIdentifiers(FuncInfo, Specifier.Values);
 				}
-			}
-			else if (Specifier.Key == TEXT("Client"))
-			{
-				if ((FuncInfo.FunctionFlags & FUNC_BlueprintEvent) != 0)
+				break;
+
+				case EFunctionSpecifier::Reliable:
 				{
-					FError::Throwf(TEXT("BlueprintImplementableEvent or BlueprintNativeEvent functions cannot be declared as Client or Server"));
+					FuncInfo.FunctionFlags |= FUNC_NetReliable;
 				}
+				break;
 
-				FuncInfo.FunctionFlags |= FUNC_Net;
-				FuncInfo.FunctionFlags |= FUNC_NetClient;
-
-				if (Specifier.Values.Num())
+				case EFunctionSpecifier::Unreliable:
 				{
-					FuncInfo.CppImplName = Specifier.Values[0];
+					bSpecifiedUnreliable = true;
 				}
-			}
-			else if (Specifier.Key == TEXT("NetMulticast"))
-			{
-				if ((FuncInfo.FunctionFlags & FUNC_BlueprintEvent) != 0)
+				break;
+
+				case EFunctionSpecifier::CustomThunk:
 				{
-					FError::Throwf(TEXT("BlueprintImplementableEvent or BlueprintNativeEvent functions cannot be declared as Multicast"));
+					FuncInfo.FunctionExportFlags |= FUNCEXPORT_CustomThunk;
 				}
+				break;
 
-				FuncInfo.FunctionFlags |= FUNC_Net;
-				FuncInfo.FunctionFlags |= FUNC_NetMulticast;
-			}
-			else if (Specifier.Key == TEXT("ServiceRequest"))
-			{
-				if ((FuncInfo.FunctionFlags & FUNC_BlueprintEvent) != 0)
+				case EFunctionSpecifier::BlueprintCallable:
 				{
-					FError::Throwf(TEXT("BlueprintImplementableEvent or BlueprintNativeEvent functions cannot be declared as a ServiceRequest"));
+					FuncInfo.FunctionFlags |= FUNC_BlueprintCallable;
 				}
+				break;
 
-				FuncInfo.FunctionFlags |= FUNC_Net;
-				FuncInfo.FunctionFlags |= FUNC_NetReliable;
-				FuncInfo.FunctionFlags |= FUNC_NetRequest;
-				FuncInfo.FunctionExportFlags |= FUNCEXPORT_CustomThunk;
-
-				ParseNetServiceIdentifiers(FuncInfo, Specifier.Values);
-			}
-			else if (Specifier.Key == TEXT("ServiceResponse"))
-			{
-				if ((FuncInfo.FunctionFlags & FUNC_BlueprintEvent) != 0)
+				case EFunctionSpecifier::BlueprintPure:
 				{
-					FError::Throwf(TEXT("BlueprintImplementableEvent or BlueprintNativeEvent functions cannot be declared as a ServiceResponse"));
+					// This function can be called, and is also pure.
+					FuncInfo.FunctionFlags |= FUNC_BlueprintCallable;
+					FuncInfo.FunctionFlags |= FUNC_BlueprintPure;
 				}
+				break;
 
-				FuncInfo.FunctionFlags |= FUNC_Net;
-				FuncInfo.FunctionFlags |= FUNC_NetReliable;
-				FuncInfo.FunctionFlags |= FUNC_NetResponse;
+				case EFunctionSpecifier::BlueprintAuthorityOnly:
+				{
+					FuncInfo.FunctionFlags |= FUNC_BlueprintAuthorityOnly;
+				}
+				break;
 
-				ParseNetServiceIdentifiers(FuncInfo, Specifier.Values);
-			}
-			else if (Specifier.Key == TEXT("Reliable"))
-			{
-				FuncInfo.FunctionFlags |= FUNC_NetReliable;
-			}
-			else if (Specifier.Key == TEXT("Unreliable"))
-			{
-				bSpecifiedUnreliable = true;
-			}
-			else if (Specifier.Key == TEXT("CustomThunk"))
-			{
-				FuncInfo.FunctionExportFlags |= FUNCEXPORT_CustomThunk;
-			}
-			else if (Specifier.Key == TEXT("BlueprintCallable"))
-			{
-				FuncInfo.FunctionFlags |= FUNC_BlueprintCallable;
-			}
-			else if (Specifier.Key == TEXT("BlueprintPure"))
-			{
-				// This function can be called, and is also pure.
-				FuncInfo.FunctionFlags |= FUNC_BlueprintCallable;
-				FuncInfo.FunctionFlags |= FUNC_BlueprintPure;
-			}
-			else if (Specifier.Key == TEXT("BlueprintAuthorityOnly"))
-			{
-				FuncInfo.FunctionFlags |= FUNC_BlueprintAuthorityOnly;
-			}
-			else if (Specifier.Key == TEXT("BlueprintCosmetic"))
-			{
-				FuncInfo.FunctionFlags |= FUNC_BlueprintCosmetic;
-			}
-			else if (Specifier.Key == TEXT("WithValidation"))
-			{
-				FuncInfo.FunctionFlags |= FUNC_NetValidate;
+				case EFunctionSpecifier::BlueprintCosmetic:
+				{
+					FuncInfo.FunctionFlags |= FUNC_BlueprintCosmetic;
+				}
+				break;
 
-				if (Specifier.Values.Num())
+				case EFunctionSpecifier::WithValidation:
 				{
-					FuncInfo.CppValidationImplName = Specifier.Values[0];
+					FuncInfo.FunctionFlags |= FUNC_NetValidate;
+
+					if (Specifier.Values.Num())
+					{
+						FuncInfo.CppValidationImplName = Specifier.Values[0];
+					}
 				}
-			}
-			else
-			{
-				FError::Throwf(TEXT("Unknown function specifier '%s'"), *Specifier.Key);
+				break;
 			}
 		}
 
@@ -624,53 +670,73 @@ namespace
 	 */
 	void ValidateMetaDataFormat(UField* Field, const FString& InKey, const FString& InValue)
 	{
-		if ((InKey == TEXT("UIMin")) || (InKey == TEXT("UIMax")) || (InKey == TEXT("ClampMin")) || (InKey == TEXT("ClampMax")))
+		switch ((ECheckedMetadataSpecifier)Algo::FindSortedStringCaseInsensitive(*InKey, GCheckedMetadataSpecifierStrings, ARRAY_COUNT(GCheckedMetadataSpecifierStrings)))
 		{
-			if (!InValue.IsNumeric())
+			default:
 			{
-				FError::Throwf(TEXT("Metadata value for '%s' is non-numeric : '%s'"), *InKey, *InValue);
+				// Don't need to validate this specifier
 			}
-		}
-		else if (InKey == /*FBlueprintMetadata::MD_Protected*/ TEXT("BlueprintProtected"))
-		{
-			if (UFunction* Function = Cast<UFunction>(Field))
+			break;
+
+			case ECheckedMetadataSpecifier::UIMin:
+			case ECheckedMetadataSpecifier::UIMax:
+			case ECheckedMetadataSpecifier::ClampMin:
+			case ECheckedMetadataSpecifier::ClampMax:
 			{
-				if (Function->HasAnyFunctionFlags(FUNC_Static))
+				if (!InValue.IsNumeric())
 				{
-					// Determine if it's a function library
-					UClass* Class = Cast<UClass>(Function->GetOuterUClass());
-					while (Class != nullptr && Class->GetSuperClass() != UObject::StaticClass())
-					{
-						Class = Class->GetSuperClass();
-					}
+					FError::Throwf(TEXT("Metadata value for '%s' is non-numeric : '%s'"), *InKey, *InValue);
+				}
+			}
+			break;
 
-					if (Class != nullptr && Class->GetName() == TEXT("BlueprintFunctionLibrary"))
+			case ECheckedMetadataSpecifier::BlueprintProtected:
+			{
+				if (UFunction* Function = Cast<UFunction>(Field))
+				{
+					if (Function->HasAnyFunctionFlags(FUNC_Static))
 					{
-						FError::Throwf(TEXT("%s doesn't make sense on static method '%s' in a blueprint function library"), *InKey, *Function->GetName());
+						// Determine if it's a function library
+						UClass* Class = Cast<UClass>(Function->GetOuterUClass());
+						while (Class != nullptr && Class->GetSuperClass() != UObject::StaticClass())
+						{
+							Class = Class->GetSuperClass();
+						}
+
+						if (Class != nullptr && Class->GetName() == TEXT("BlueprintFunctionLibrary"))
+						{
+							FError::Throwf(TEXT("%s doesn't make sense on static method '%s' in a blueprint function library"), *InKey, *Function->GetName());
+						}
 					}
 				}
 			}
-		}
-		else if (InKey == TEXT("DevelopmentStatus"))
-		{
-			const FString EarlyAccessValue(TEXT("EarlyAccess"));
-			const FString ExperimentalValue(TEXT("Experimental"));
-			if ((InValue != EarlyAccessValue) && (InValue != ExperimentalValue))
-			{
-				FError::Throwf(TEXT("'%s' metadata was '%s' but it must be %s or %s"), *InKey, *InValue, *ExperimentalValue, *EarlyAccessValue);
-			}
-		}
-		else if (InKey == TEXT("Units"))
-		{
-			// Check for numeric property
-			if (!Cast<UNumericProperty>(Field))
+			break;
+
+			case ECheckedMetadataSpecifier::DevelopmentStatus:
 			{
-				FError::Throwf(TEXT("'Units' meta data can only be applied to numeric properties"));
+				const FString EarlyAccessValue(TEXT("EarlyAccess"));
+				const FString ExperimentalValue(TEXT("Experimental"));
+				if ((InValue != EarlyAccessValue) && (InValue != ExperimentalValue))
+				{
+					FError::Throwf(TEXT("'%s' metadata was '%s' but it must be %s or %s"), *InKey, *InValue, *ExperimentalValue, *EarlyAccessValue);
+				}
 			}
-			else if (!FUnitConversion::UnitFromString(*InValue))
+			break;
+
+			case ECheckedMetadataSpecifier::Units:
 			{
-				FError::Throwf(TEXT("Unrecognized units (%s) specified for numeric property '%s'"), *InValue, *Field->GetDisplayNameText().ToString());
+				// Check for numeric property
+				if (!Cast<UNumericProperty>(Field))
+				{
+					FError::Throwf(TEXT("'Units' meta data can only be applied to numeric properties"));
+				}
+
+				if (!FUnitConversion::UnitFromString(*InValue))
+				{
+					FError::Throwf(TEXT("Unrecognized units (%s) specified for numeric property '%s'"), *InValue, *Field->GetDisplayNameText().ToString());
+				}
 			}
+			break;
 		}
 	}
 
@@ -1532,32 +1598,40 @@ UScriptStruct* FHeaderParser::CompileStructDeclaration(FClasses& AllClasses, FUn
 	const FString EffectiveStructName = *StructNameStripped;
 
 	// Process the list of specifiers
-	for (TArray<FPropertySpecifier>::TIterator SpecifierIt(SpecifiersFound); SpecifierIt; ++SpecifierIt)
+	for (const FPropertySpecifier& Specifier : SpecifiersFound)
 	{
-		const FString& Specifier = SpecifierIt->Key;
-
-		if (Specifier == TEXT("NoExport"))
-		{
-			//UE_LOG(LogCompile, Warning, TEXT("Struct named %s in %s is still marked noexport"), *EffectiveStructName, *(Class->GetName()));//@TODO: UCREMOVAL: Debug printing
-			StructFlags &= ~STRUCT_Native;
-			StructFlags |= STRUCT_NoExport;
-		}
-		else if (Specifier == TEXT("Atomic"))
+		switch ((EStructSpecifier)Algo::FindSortedStringCaseInsensitive(*Specifier.Key, GStructSpecifierStrings, ARRAY_COUNT(GStructSpecifierStrings)))
 		{
-			StructFlags |= STRUCT_Atomic;
-		}
-		else if (Specifier == TEXT("Immutable"))
-		{
-			StructFlags |= STRUCT_Immutable | STRUCT_Atomic;
+			default:
+			{
+				FError::Throwf(TEXT("Unknown struct specifier '%s'"), *Specifier.Key);
+			}
+			break;
 
-			if (!FPaths::IsSamePath(Filename, GTypeDefinitionInfoMap[UObject::StaticClass()]->GetUnrealSourceFile().GetFilename()))
+			case EStructSpecifier::NoExport:
 			{
-				FError::Throwf(TEXT("Immutable is being phased out in favor of SerializeNative, and is only legal on the mirror structs declared in UObject"));
+				//UE_LOG(LogCompile, Warning, TEXT("Struct named %s in %s is still marked noexport"), *EffectiveStructName, *(Class->GetName()));//@TODO: UCREMOVAL: Debug printing
+				StructFlags &= ~STRUCT_Native;
+				StructFlags |= STRUCT_NoExport;
 			}
-		}
-		else
-		{
-			FError::Throwf(TEXT("Unknown struct specifier '%s'"), *Specifier);
+			break;
+
+			case EStructSpecifier::Atomic:
+			{
+				StructFlags |= STRUCT_Atomic;
+			}
+			break;
+
+			case EStructSpecifier::Immutable:
+			{
+				StructFlags |= STRUCT_Immutable | STRUCT_Atomic;
+
+				if (!FPaths::IsSamePath(Filename, GTypeDefinitionInfoMap[UObject::StaticClass()]->GetUnrealSourceFile().GetFilename()))
+				{
+					FError::Throwf(TEXT("Immutable is being phased out in favor of SerializeNative, and is only legal on the mirror structs declared in UObject"));
+				}
+			}
+			break;
 		}
 	}
 
@@ -2528,253 +2602,329 @@ void FHeaderParser::GetVarType
 	// Process the list of specifiers
 	bool bSeenEditSpecifier = false;
 	bool bSeenBlueprintEditSpecifier = false;
-	for (TArray<FPropertySpecifier>::TIterator SpecifierIt(SpecifiersFound); SpecifierIt; ++SpecifierIt)
+	for (const FPropertySpecifier& Specifier : SpecifiersFound)
 	{
-		const FString& Specifier = SpecifierIt->Key;
-
+		EVariableSpecifier SpecID = (EVariableSpecifier)Algo::FindSortedStringCaseInsensitive(*Specifier.Key, GVariableSpecifierStrings, ARRAY_COUNT(GVariableSpecifierStrings));
 		if (VariableCategory == EVariableCategory::Member)
 		{
-			if (Specifier == TEXT("EditAnywhere"))
+			switch (SpecID)
 			{
-				if (bSeenEditSpecifier)
+				case EVariableSpecifier::EditAnywhere:
 				{
-					FError::Throwf(TEXT("Found more than one edit/visibility specifier (%s), only one is allowed"), *Specifier);
+					if (bSeenEditSpecifier)
+					{
+						FError::Throwf(TEXT("Found more than one edit/visibility specifier (%s), only one is allowed"), *Specifier.Key);
+					}
+					Flags |= CPF_Edit;
+					bSeenEditSpecifier = true;
 				}
-				Flags |= CPF_Edit;
-				bSeenEditSpecifier = true;
-			}
-			else if (Specifier == TEXT("EditInstanceOnly"))
-			{
-				if (bSeenEditSpecifier)
+				break;
+
+				case EVariableSpecifier::EditInstanceOnly:
 				{
-					FError::Throwf(TEXT("Found more than one edit/visibility specifier (%s), only one is allowed"), *Specifier);
+					if (bSeenEditSpecifier)
+					{
+						FError::Throwf(TEXT("Found more than one edit/visibility specifier (%s), only one is allowed"), *Specifier.Key);
+					}
+					Flags |= CPF_Edit|CPF_DisableEditOnTemplate;
+					bSeenEditSpecifier = true;
 				}
-				Flags |= CPF_Edit|CPF_DisableEditOnTemplate;
-				bSeenEditSpecifier = true;
-			}
-			else if (Specifier == TEXT("EditDefaultsOnly")) 
-			{
-				if (bSeenEditSpecifier)
+				break;
+
+				case EVariableSpecifier::EditDefaultsOnly: 
 				{
-					FError::Throwf(TEXT("Found more than one edit/visibility specifier (%s), only one is allowed"), *Specifier);
+					if (bSeenEditSpecifier)
+					{
+						FError::Throwf(TEXT("Found more than one edit/visibility specifier (%s), only one is allowed"), *Specifier.Key);
+					}
+					Flags |= CPF_Edit|CPF_DisableEditOnInstance;
+					bSeenEditSpecifier = true;
 				}
-				Flags |= CPF_Edit|CPF_DisableEditOnInstance;
-				bSeenEditSpecifier = true;
-			}
-			else if (Specifier == TEXT("VisibleAnywhere")) 
-			{
-				if (bSeenEditSpecifier)
+				break;
+
+				case EVariableSpecifier::VisibleAnywhere: 
 				{
-					FError::Throwf(TEXT("Found more than one edit/visibility specifier (%s), only one is allowed"), *Specifier);
+					if (bSeenEditSpecifier)
+					{
+						FError::Throwf(TEXT("Found more than one edit/visibility specifier (%s), only one is allowed"), *Specifier.Key);
+					}
+					Flags |= CPF_Edit|CPF_EditConst;
+					bSeenEditSpecifier = true;
 				}
-				Flags |= CPF_Edit|CPF_EditConst;
-				bSeenEditSpecifier = true;
-			}
-			else if (Specifier == TEXT("VisibleInstanceOnly"))
-			{
-				if (bSeenEditSpecifier)
+				break;
+
+				case EVariableSpecifier::VisibleInstanceOnly:
 				{
-					FError::Throwf(TEXT("Found more than one edit/visibility specifier (%s), only one is allowed"), *Specifier);
+					if (bSeenEditSpecifier)
+					{
+						FError::Throwf(TEXT("Found more than one edit/visibility specifier (%s), only one is allowed"), *Specifier.Key);
+					}
+					Flags |= CPF_Edit|CPF_EditConst|CPF_DisableEditOnTemplate;
+					bSeenEditSpecifier = true;
 				}
-				Flags |= CPF_Edit|CPF_EditConst|CPF_DisableEditOnTemplate;
-				bSeenEditSpecifier = true;
-			}
-			else if (Specifier == TEXT("VisibleDefaultsOnly"))
-			{
-				if (bSeenEditSpecifier)
+				break;
+
+				case EVariableSpecifier::VisibleDefaultsOnly:
 				{
-					FError::Throwf(TEXT("Found more than one edit/visibility specifier (%s), only one is allowed"), *Specifier);
+					if (bSeenEditSpecifier)
+					{
+						FError::Throwf(TEXT("Found more than one edit/visibility specifier (%s), only one is allowed"), *Specifier.Key);
+					}
+					Flags |= CPF_Edit|CPF_EditConst|CPF_DisableEditOnInstance;
+					bSeenEditSpecifier = true;
 				}
-				Flags |= CPF_Edit|CPF_EditConst|CPF_DisableEditOnInstance;
-				bSeenEditSpecifier = true;
-			}
-			else if (Specifier == TEXT("BlueprintReadWrite"))
-			{
-				if (bSeenBlueprintEditSpecifier)
+				break;
+
+				case EVariableSpecifier::BlueprintReadWrite:
+				{
+					if (bSeenBlueprintEditSpecifier)
+					{
+						FError::Throwf(TEXT("Found more than one Blueprint read/write specifier (%s), only one is allowed"), *Specifier.Key);
+					}
+
+					const FString* PrivateAccessMD = MetaDataFromNewStyle.Find(TEXT("AllowPrivateAccess"));  // FBlueprintMetadata::MD_AllowPrivateAccess
+					const bool bAllowPrivateAccess = PrivateAccessMD ? (*PrivateAccessMD == TEXT("true")) : false;
+					if (CurrentAccessSpecifier == ACCESS_Private && !bAllowPrivateAccess)
+					{
+						FError::Throwf(TEXT("BlueprintReadWrite should not be used on private members"));
+					}
+					Flags |= CPF_BlueprintVisible;
+					bSeenBlueprintEditSpecifier = true;
+				}
+				break;
+
+				case EVariableSpecifier::BlueprintReadOnly:
 				{
-					FError::Throwf(TEXT("Found more than one Blueprint read/write specifier (%s), only one is allowed"), *Specifier);
+					if (bSeenBlueprintEditSpecifier)
+					{
+						FError::Throwf(TEXT("Found more than one Blueprint read/write specifier (%s), only one is allowed"), *Specifier.Key);
+					}
+
+					const FString* PrivateAccessMD = MetaDataFromNewStyle.Find(TEXT("AllowPrivateAccess"));  // FBlueprintMetadata::MD_AllowPrivateAccess
+					const bool bAllowPrivateAccess = PrivateAccessMD ? (*PrivateAccessMD == TEXT("true")) : false;
+					if (CurrentAccessSpecifier == ACCESS_Private && !bAllowPrivateAccess)
+					{
+						FError::Throwf(TEXT("BlueprintReadOnly should not be used on private members"));
+					}
+					Flags |= CPF_BlueprintVisible|CPF_BlueprintReadOnly;
+					bSeenBlueprintEditSpecifier = true;
 				}
+				break;
 
-				const FString* PrivateAccessMD = MetaDataFromNewStyle.Find(TEXT("AllowPrivateAccess"));  // FBlueprintMetadata::MD_AllowPrivateAccess
-				const bool bAllowPrivateAccess = PrivateAccessMD ? (*PrivateAccessMD == TEXT("true")) : false;
-				if ((CurrentAccessSpecifier == ACCESS_Private) && !bAllowPrivateAccess)
+				case EVariableSpecifier::Config:
 				{
-					FError::Throwf(TEXT("BlueprintReadWrite should not be used on private members"));
+					Flags |= CPF_Config;
 				}
-				Flags |= CPF_BlueprintVisible;
-				bSeenBlueprintEditSpecifier = true;
-			}
-			else if (Specifier == TEXT("BlueprintReadOnly"))
-			{
-				if (bSeenBlueprintEditSpecifier)
+				break;
+
+				case EVariableSpecifier::GlobalConfig:
 				{
-					FError::Throwf(TEXT("Found more than one Blueprint read/write specifier (%s), only one is allowed"), *Specifier);
+					Flags |= CPF_GlobalConfig | CPF_Config;
 				}
+				break;
 
-				const FString* PrivateAccessMD = MetaDataFromNewStyle.Find(TEXT("AllowPrivateAccess"));  // FBlueprintMetadata::MD_AllowPrivateAccess
-				const bool bAllowPrivateAccess = PrivateAccessMD ? (*PrivateAccessMD == TEXT("true")) : false;
-				if ((CurrentAccessSpecifier == ACCESS_Private) && !bAllowPrivateAccess)
+				case EVariableSpecifier::Localized:
 				{
-					FError::Throwf(TEXT("BlueprintReadOnly should not be used on private members"));
+					FError::Throwf(TEXT("The Localized specifier is deprecated"));
 				}
-				Flags |= CPF_BlueprintVisible|CPF_BlueprintReadOnly;
-				bSeenBlueprintEditSpecifier = true;
-			}
-			else if (Specifier == TEXT("Config"))
-			{
-				Flags |= CPF_Config;
-			}
-			else if (Specifier == TEXT("GlobalConfig"))
-			{
-				Flags |= CPF_GlobalConfig | CPF_Config;
-			}
-			else if (Specifier == TEXT("Localized"))
-			{
-				FError::Throwf(TEXT("The Localized specifier is deprecated"));
-			}
-			else if (Specifier == TEXT("Transient"))
-			{
-				Flags |= CPF_Transient;
-			}
-			else if (Specifier == TEXT("DuplicateTransient"))
-			{
-				Flags |= CPF_DuplicateTransient;
-			}
-			else if (Specifier == TEXT("TextExportTransient"))
-			{
-				Flags |= CPF_TextExportTransient;
-			}
-			else if (Specifier == TEXT("NonPIETransient"))
-			{
-				UE_LOG(LogCompile, Warning, TEXT("NonPIETransient is deprecated - NonPIEDuplicateTransient should be used instead"));
-				Flags |= CPF_NonPIEDuplicateTransient;
-			}
-			else if (Specifier == TEXT("NonPIEDuplicateTransient"))
-			{
-				Flags |= CPF_NonPIEDuplicateTransient;
-			}
-			else if (Specifier == TEXT("Export"))
-			{
-				Flags |= CPF_ExportObject;
-			}
-			else if (Specifier == TEXT("EditInline"))
-			{
-				FError::Throwf(TEXT("EditInline is deprecated. Remove it, or use Instanced instead."));
-			}
-			else if (Specifier == TEXT("NoClear"))
-			{
-				Flags |= CPF_NoClear;
-			}
-			else if (Specifier == TEXT("EditFixedSize"))
-			{
-				Flags |= CPF_EditFixedSize;
-			}
-			else if (Specifier == TEXT("Replicated") || Specifier == TEXT("ReplicatedUsing"))
-			{
-				if (!OwnerStruct->IsA<UScriptStruct>())
+				break;
+
+				case EVariableSpecifier::Transient:
+				{
+					Flags |= CPF_Transient;
+				}
+				break;
+
+				case EVariableSpecifier::DuplicateTransient:
 				{
+					Flags |= CPF_DuplicateTransient;
+				}
+				break;
+
+				case EVariableSpecifier::TextExportTransient:
+				{
+					Flags |= CPF_TextExportTransient;
+				}
+				break;
+
+				case EVariableSpecifier::NonPIETransient:
+				{
+					UE_LOG(LogCompile, Warning, TEXT("NonPIETransient is deprecated - NonPIEDuplicateTransient should be used instead"));
+					Flags |= CPF_NonPIEDuplicateTransient;
+				}
+				break;
+
+				case EVariableSpecifier::NonPIEDuplicateTransient:
+				{
+					Flags |= CPF_NonPIEDuplicateTransient;
+				}
+				break;
+
+				case EVariableSpecifier::Export:
+				{
+					Flags |= CPF_ExportObject;
+				}
+				break;
+
+				case EVariableSpecifier::EditInline:
+				{
+					FError::Throwf(TEXT("EditInline is deprecated. Remove it, or use Instanced instead."));
+				}
+				break;
+
+				case EVariableSpecifier::NoClear:
+				{
+					Flags |= CPF_NoClear;
+				}
+				break;
+
+				case EVariableSpecifier::EditFixedSize:
+				{
+					Flags |= CPF_EditFixedSize;
+				}
+				break;
+
+				case EVariableSpecifier::Replicated:
+				case EVariableSpecifier::ReplicatedUsing:
+				{
+					if (OwnerStruct->IsA<UScriptStruct>())
+					{
+						FError::Throwf(TEXT("Struct members cannot be replicated"));
+					}
+
 					Flags |= CPF_Net;
 
 					// See if we've specified a rep notification function
-					if (Specifier == TEXT("ReplicatedUsing"))
+					if (SpecID == EVariableSpecifier::ReplicatedUsing)
 					{
-						RepCallbackName = FName(*RequireExactlyOneSpecifierValue(*SpecifierIt));
+						RepCallbackName = FName(*RequireExactlyOneSpecifierValue(Specifier));
 						Flags |= CPF_RepNotify;
 					}
 				}
-				else
+				break;
+
+				case EVariableSpecifier::NotReplicated:
 				{
-					FError::Throwf(TEXT("Struct members cannot be replicated"));
+					if (!OwnerStruct->IsA<UScriptStruct>())
+					{
+						FError::Throwf(TEXT("Only Struct members can be marked NotReplicated"));
+					}
+
+					Flags |= CPF_RepSkip;
 				}
-			}
-			else if (Specifier == TEXT("NotReplicated"))
-			{
-				if (OwnerStruct->IsA<UScriptStruct>())
+				break;
+
+				case EVariableSpecifier::RepRetry:
 				{
-					Flags |= CPF_RepSkip;
+					FError::Throwf(TEXT("'RepRetry' is deprecated."));
 				}
-				else
+				break;
+
+				case EVariableSpecifier::Interp:
 				{
-					FError::Throwf(TEXT("Only Struct members can be marked NotReplicated"));
+					Flags |= CPF_Edit;
+					Flags |= CPF_BlueprintVisible;
+					Flags |= CPF_Interp;
 				}
-			}
-			else if (Specifier == TEXT("RepRetry"))
-			{
-				FError::Throwf(TEXT("'RepRetry' is deprecated."));
-			}
-			else if (Specifier == TEXT("Interp"))
-			{
-				Flags |= CPF_Edit;
-				Flags |= CPF_BlueprintVisible;
-				Flags |= CPF_Interp;
-			}
-			else if (Specifier == TEXT("NonTransactional"))
-			{
-				Flags |= CPF_NonTransactional;
-			}
-			else if (Specifier == TEXT("Instanced"))
-			{
-				Flags |= CPF_PersistentInstance | CPF_ExportObject | CPF_InstancedReference;
-				AddEditInlineMetaData(MetaDataFromNewStyle);
-			}
-			else if (Specifier == TEXT("BlueprintAssignable"))
-			{
-				Flags |= CPF_BlueprintAssignable;
-			}
-			else if (Specifier == TEXT("BlueprintCallable"))
-			{
-				Flags |= CPF_BlueprintCallable;
-			}
-			else if (Specifier == TEXT("BlueprintAuthorityOnly"))
-			{
-				Flags |= CPF_BlueprintAuthorityOnly;
-			}
-			else if (Specifier == TEXT("AssetRegistrySearchable"))
-			{
-				Flags |= CPF_AssetRegistrySearchable;
-			}
-			else if (Specifier == TEXT("SimpleDisplay"))
-			{
-				Flags |= CPF_SimpleDisplay;
-			}
-			else if (Specifier == TEXT("AdvancedDisplay"))
-			{
-				Flags |= CPF_AdvancedDisplay;
-			}
-			else if (Specifier == TEXT("SaveGame"))
-			{
-				Flags |= CPF_SaveGame;
-			}
-			else
-			{
-				FError::Throwf(TEXT("Unknown variable specifier '%s'"), *Specifier);
+				break;
+
+				case EVariableSpecifier::NonTransactional:
+				{
+					Flags |= CPF_NonTransactional;
+				}
+				break;
+
+				case EVariableSpecifier::Instanced:
+				{
+					Flags |= CPF_PersistentInstance | CPF_ExportObject | CPF_InstancedReference;
+					AddEditInlineMetaData(MetaDataFromNewStyle);
+				}
+				break;
+
+				case EVariableSpecifier::BlueprintAssignable:
+				{
+					Flags |= CPF_BlueprintAssignable;
+				}
+				break;
+
+				case EVariableSpecifier::BlueprintCallable:
+				{
+					Flags |= CPF_BlueprintCallable;
+				}
+				break;
+
+				case EVariableSpecifier::BlueprintAuthorityOnly:
+				{
+					Flags |= CPF_BlueprintAuthorityOnly;
+				}
+				break;
+
+				case EVariableSpecifier::AssetRegistrySearchable:
+				{
+					Flags |= CPF_AssetRegistrySearchable;
+				}
+				break;
+
+				case EVariableSpecifier::SimpleDisplay:
+				{
+					Flags |= CPF_SimpleDisplay;
+				}
+				break;
+
+				case EVariableSpecifier::AdvancedDisplay:
+				{
+					Flags |= CPF_AdvancedDisplay;
+				}
+				break;
+
+				case EVariableSpecifier::SaveGame:
+				{
+					Flags |= CPF_SaveGame;
+				}
+				break;
+
+				default:
+				{
+					FError::Throwf(TEXT("Unknown variable specifier '%s'"), *Specifier.Key);
+				}
+				break;
 			}
 		}
 		else
 		{
-			if (Specifier == TEXT("Const"))
-			{
-				Flags |= CPF_ConstParm;
-			}
-			else if (Specifier == TEXT("Ref"))
-			{
-				Flags |= CPF_OutParm | CPF_ReferenceParm;
-			}
-			else if (Specifier == TEXT("NotReplicated"))
+			switch (SpecID)
 			{
-				if (VariableCategory == EVariableCategory::ReplicatedParameter)
+				case EVariableSpecifier::Const:
 				{
-					VariableCategory = EVariableCategory::RegularParameter;
-					Flags |= CPF_RepSkip;
+					Flags |= CPF_ConstParm;
 				}
-				else
+				break;
+
+				case EVariableSpecifier::Ref:
 				{
-					FError::Throwf(TEXT("Only parameters in service request functions can be marked NotReplicated"));
+					Flags |= CPF_OutParm | CPF_ReferenceParm;
 				}
-			}
-			else
-			{
-				FError::Throwf(TEXT("Unknown variable specifier '%s'"), *Specifier);
+				break;
+
+				case EVariableSpecifier::NotReplicated:
+				{
+					if (VariableCategory == EVariableCategory::ReplicatedParameter)
+					{
+						VariableCategory = EVariableCategory::RegularParameter;
+						Flags |= CPF_RepSkip;
+					}
+					else
+					{
+						FError::Throwf(TEXT("Only parameters in service request functions can be marked NotReplicated"));
+					}
+				}
+				break;
+
+				default:
+				{
+					FError::Throwf(TEXT("Unknown variable specifier '%s'"), *Specifier.Key);
+				}
+				break;
 			}
 		}
 	}
@@ -4661,25 +4811,33 @@ void FHeaderParser::CompileInterfaceDeclaration(FClasses& AllClasses)
 	InterfaceClass->ClassFlags |= CLASS_Native;
 
 	// Process all of the interface specifiers
-	for (TArray<FPropertySpecifier>::TIterator SpecifierIt(SpecifiersFound); SpecifierIt; ++SpecifierIt)
+	for (const FPropertySpecifier& Specifier : SpecifiersFound)
 	{
-		const FString& Specifier = SpecifierIt->Key;
-
-		if (Specifier == TEXT("DependsOn"))
-		{
-			FError::Throwf(TEXT("The dependsOn specifier is deprecated. Please use #include \"ClassHeaderFilename.h\" instead."));
-		}
-		else if (Specifier == TEXT("MinimalAPI"))
-		{
-			InterfaceClass->ClassFlags |= CLASS_MinimalAPI;
-		}
-		else if (Specifier == TEXT("ConversionRoot"))
+		switch ((EInterfaceSpecifier)Algo::FindSortedStringCaseInsensitive(*Specifier.Key, GInterfaceSpecifierStrings, ARRAY_COUNT(GInterfaceSpecifierStrings)))
 		{
-			MetaData.Add(FName(TEXT("IsConversionRoot")), "true");
-		}
-		else
-		{
-			FError::Throwf(TEXT("Unknown interface specifier '%s'"), *Specifier);
+			default:
+			{
+				FError::Throwf(TEXT("Unknown interface specifier '%s'"), *Specifier.Key);
+			}
+			break;
+
+			case EInterfaceSpecifier::DependsOn:
+			{
+				FError::Throwf(TEXT("The dependsOn specifier is deprecated. Please use #include \"ClassHeaderFilename.h\" instead."));
+			}
+			break;
+
+			case EInterfaceSpecifier::MinimalAPI:
+			{
+				InterfaceClass->ClassFlags |= CLASS_MinimalAPI;
+			}
+			break;
+
+			case EInterfaceSpecifier::ConversionRoot:
+			{
+				MetaData.Add(FName(TEXT("IsConversionRoot")), "true");
+			}
+			break;
 		}
 	}
 
diff --git a/Engine/Source/Programs/UnrealHeaderTool/Private/Specifiers/CheckedMetadataSpecifiers.def b/Engine/Source/Programs/UnrealHeaderTool/Private/Specifiers/CheckedMetadataSpecifiers.def
new file mode 100644
--- /dev/null
+++ b/Engine/Source/Programs/UnrealHeaderTool/Private/Specifiers/CheckedMetadataSpecifiers.def
@@ -0,0 +1,11 @@
+// CHECKED_METADATA_SPECIFIER(<specifier name>)
+//
+// Must be listed such that the strings are in lexicographical order!
+
+CHECKED_METADATA_SPECIFIER(BlueprintProtected)
+CHECKED_METADATA_SPECIFIER(ClampMax)
+CHECKED_METADATA_SPECIFIER(ClampMin)
+CHECKED_METADATA_SPECIFIER(DevelopmentStatus)
+CHECKED_METADATA_SPECIFIER(UIMax)
+CHECKED_METADATA_SPECIFIER(UIMin)
+CHECKED_METADATA_SPECIFIER(Units)
diff --git a/Engine/Source/Programs/UnrealHeaderTool/Private/Specifiers/CheckedMetadataSpecifiers.h b/Engine/Source/Programs/UnrealHeaderTool/Private/Specifiers/CheckedMetadataSpecifiers.h
new file mode 100644
--- /dev/null
+++ b/Engine/Source/Programs/UnrealHeaderTool/Private/Specifiers/CheckedMetadataSpecifiers.h
@@ -0,0 +1,16 @@
+// Copyright 1998-2015 Epic Games, Inc. All Rights Reserved.
+
+#pragma once
+
+enum class ECheckedMetadataSpecifier
+{
+	None = -1
+
+	#define CHECKED_METADATA_SPECIFIER(SpecifierName) , SpecifierName
+		#include "CheckedMetadataSpecifiers.def"
+	#undef CHECKED_METADATA_SPECIFIER
+
+	, Max
+};
+
+extern const TCHAR* GCheckedMetadataSpecifierStrings[(int32)ECheckedMetadataSpecifier::Max];
diff --git a/Engine/Source/Programs/UnrealHeaderTool/Private/Specifiers/FunctionSpecifiers.def b/Engine/Source/Programs/UnrealHeaderTool/Private/Specifiers/FunctionSpecifiers.def
new file mode 100644
--- /dev/null
+++ b/Engine/Source/Programs/UnrealHeaderTool/Private/Specifiers/FunctionSpecifiers.def
@@ -0,0 +1,21 @@
+// FUNCTION_SPECIFIER(<specifier name>)
+//
+// Must be listed such that the strings are in lexicographical order!
+
+FUNCTION_SPECIFIER(BlueprintAuthorityOnly)
+FUNCTION_SPECIFIER(BlueprintCallable)
+FUNCTION_SPECIFIER(BlueprintCosmetic)
+FUNCTION_SPECIFIER(BlueprintImplementableEvent)
+FUNCTION_SPECIFIER(BlueprintNativeEvent)
+FUNCTION_SPECIFIER(BlueprintPure)
+FUNCTION_SPECIFIER(Client)
+FUNCTION_SPECIFIER(CustomThunk)
+FUNCTION_SPECIFIER(Exec)
+FUNCTION_SPECIFIER(NetMulticast)
+FUNCTION_SPECIFIER(Reliable)
+FUNCTION_SPECIFIER(SealedEvent)
+FUNCTION_SPECIFIER(Server)
+FUNCTION_SPECIFIER(ServiceRequest)
+FUNCTION_SPECIFIER(ServiceResponse)
+FUNCTION_SPECIFIER(Unreliable)
+FUNCTION_SPECIFIER(WithValidation)
diff --git a/Engine/Source/Programs/UnrealHeaderTool/Private/Specifiers/FunctionSpecifiers.h b/Engine/Source/Programs/UnrealHeaderTool/Private/Specifiers/FunctionSpecifiers.h
new file mode 100644
--- /dev/null
+++ b/Engine/Source/Programs/UnrealHeaderTool/Private/Specifiers/FunctionSpecifiers.h
@@ -0,0 +1,16 @@
+// Copyright 1998-2015 Epic Games, Inc. All Rights Reserved.
+
+#pragma once
+
+enum class EFunctionSpecifier
+{
+	None = -1
+
+	#define FUNCTION_SPECIFIER(SpecifierName) , SpecifierName
+		#include "FunctionSpecifiers.def"
+	#undef FUNCTION_SPECIFIER
+
+	, Max
+};
+
+extern const TCHAR* GFunctionSpecifierStrings[(int32)EFunctionSpecifier::Max];
diff --git a/Engine/Source/Programs/UnrealHeaderTool/Private/Specifiers/InterfaceSpecifiers.def b/Engine/Source/Programs/UnrealHeaderTool/Private/Specifiers/InterfaceSpecifiers.def
new file mode 100644
--- /dev/null
+++ b/Engine/Source/Programs/UnrealHeaderTool/Private/Specifiers/InterfaceSpecifiers.def
@@ -0,0 +1,7 @@
+// INTERFACE_SPECIFIER(<specifier name>)
+//
+// Must be listed such that the strings are in lexicographical order!
+
+INTERFACE_SPECIFIER(ConversionRoot)
+INTERFACE_SPECIFIER(DependsOn)
+INTERFACE_SPECIFIER(MinimalAPI)
diff --git a/Engine/Source/Programs/UnrealHeaderTool/Private/Specifiers/InterfaceSpecifiers.h b/Engine/Source/Programs/UnrealHeaderTool/Private/Specifiers/InterfaceSpecifiers.h
new file mode 100644
--- /dev/null
+++ b/Engine/Source/Programs/UnrealHeaderTool/Private/Specifiers/InterfaceSpecifiers.h
@@ -0,0 +1,16 @@
+// Copyright 1998-2015 Epic Games, Inc. All Rights Reserved.
+
+#pragma once
+
+enum class EInterfaceSpecifier
+{
+	None = -1
+
+	#define INTERFACE_SPECIFIER(SpecifierName) , SpecifierName
+		#include "InterfaceSpecifiers.def"
+	#undef INTERFACE_SPECIFIER
+
+	, Max
+};
+
+extern const TCHAR* GInterfaceSpecifierStrings[(int32)EInterfaceSpecifier::Max];
diff --git a/Engine/Source/Programs/UnrealHeaderTool/Private/Specifiers/SpecifierArrays.cpp b/Engine/Source/Programs/UnrealHeaderTool/Private/Specifiers/SpecifierArrays.cpp
new file mode 100644
--- /dev/null
+++ b/Engine/Source/Programs/UnrealHeaderTool/Private/Specifiers/SpecifierArrays.cpp
@@ -0,0 +1,59 @@
+// Copyright 1998-2015 Epic Games, Inc. All Rights Reserved.
+
+#include "UnrealHeaderTool.h"
+
+#include "CheckedMetadataSpecifiers.h"
+#include "FunctionSpecifiers.h"
+#include "InterfaceSpecifiers.h"
+#include "StructSpecifiers.h"
+#include "VariableSpecifiers.h"
+#include "Algo/IsSorted.h"
+
+const TCHAR* GFunctionSpecifierStrings[(int32)EFunctionSpecifier::Max] =
+{
+	#define FUNCTION_SPECIFIER(SpecifierName) TEXT(#SpecifierName),
+		#include "FunctionSpecifiers.def"
+	#undef FUNCTION_SPECIFIER
+};
+
+const TCHAR* GStructSpecifierStrings[(int32)EStructSpecifier::Max] =
+{
+	#define STRUCT_SPECIFIER(SpecifierName) TEXT(#SpecifierName),
+		#include "StructSpecifiers.def"
+	#undef STRUCT_SPECIFIER
+};
+
+const TCHAR* GInterfaceSpecifierStrings[(int32)EInterfaceSpecifier::Max] =
+{
+	#define INTERFACE_SPECIFIER(SpecifierName) TEXT(#SpecifierName),
+		#include "InterfaceSpecifiers.def"
+	#undef INTERFACE_SPECIFIER
+};
+
+const TCHAR* GVariableSpecifierStrings[(int32)EVariableSpecifier::Max] =
+{
+	#define VARIABLE_SPECIFIER(SpecifierName) TEXT(#SpecifierName),
+		#include "VariableSpecifiers.def"
+	#undef VARIABLE_SPECIFIER
+};
+
+const TCHAR* GCheckedMetadataSpecifierStrings[(int32)ECheckedMetadataSpecifier::Max] =
+{
+	#define CHECKED_METADATA_SPECIFIER(SpecifierName) TEXT(#SpecifierName),
+		#include "CheckedMetadataSpecifiers.def"
+	#undef CHECKED_METADATA_SPECIFIER
+};
+
+struct FCStringsLessThanCaseInsensitive
+{
+	FORCEINLINE bool operator()(const TCHAR* Lhs, const TCHAR* Rhs)
+	{
+		return FCString::Stricmp(Lhs, Rhs) < 0;
+	}
+};
+
+const bool GIsGFunctionSpecifierStringsSorted        = ensure(Algo::IsSorted(GFunctionSpecifierStrings,        FCStringsLessThanCaseInsensitive()));
+const bool GIsGStructSpecifierStringsSorted          = ensure(Algo::IsSorted(GStructSpecifierStrings,          FCStringsLessThanCaseInsensitive()));
+const bool GIsGInterfaceSpecifierStringsSorted       = ensure(Algo::IsSorted(GInterfaceSpecifierStrings,       FCStringsLessThanCaseInsensitive()));
+const bool GIsGVariableSpecifierStringsSorted        = ensure(Algo::IsSorted(GVariableSpecifierStrings,        FCStringsLessThanCaseInsensitive()));
+const bool GIsGCheckedMetadataSpecifierStringsSorted = ensure(Algo::IsSorted(GCheckedMetadataSpecifierStrings, FCStringsLessThanCaseInsensitive()));
diff --git a/Engine/Source/Programs/UnrealHeaderTool/Private/Specifiers/StructSpecifiers.def b/Engine/Source/Programs/UnrealHeaderTool/Private/Specifiers/StructSpecifiers.def
new file mode 100644
--- /dev/null
+++ b/Engine/Source/Programs/UnrealHeaderTool/Private/Specifiers/StructSpecifiers.def
@@ -0,0 +1,7 @@
+// STRUCT_SPECIFIER(<specifier name>)
+//
+// Must be listed such that the strings are in lexicographical order!
+
+STRUCT_SPECIFIER(Atomic)
+STRUCT_SPECIFIER(Immutable)
+STRUCT_SPECIFIER(NoExport)
diff --git a/Engine/Source/Programs/UnrealHeaderTool/Private/Specifiers/StructSpecifiers.h b/Engine/Source/Programs/UnrealHeaderTool/Private/Specifiers/StructSpecifiers.h
new file mode 100644
--- /dev/null
+++ b/Engine/Source/Programs/UnrealHeaderTool/Private/Specifiers/StructSpecifiers.h
@@ -0,0 +1,16 @@
+// Copyright 1998-2015 Epic Games, Inc. All Rights Reserved.
+
+#pragma once
+
+enum class EStructSpecifier
+{
+	None = -1
+
+	#define STRUCT_SPECIFIER(SpecifierName) , SpecifierName
+		#include "StructSpecifiers.def"
+	#undef STRUCT_SPECIFIER
+
+	, Max
+};
+
+extern const TCHAR* GStructSpecifierStrings[(int32)EStructSpecifier::Max];
diff --git a/Engine/Source/Programs/UnrealHeaderTool/Private/Specifiers/VariableSpecifiers.def b/Engine/Source/Programs/UnrealHeaderTool/Private/Specifiers/VariableSpecifiers.def
new file mode 100644
--- /dev/null
+++ b/Engine/Source/Programs/UnrealHeaderTool/Private/Specifiers/VariableSpecifiers.def
@@ -0,0 +1,40 @@
+// VARIABLE_SPECIFIER(<specifier name>)
+//
+// Must be listed such that the strings are in lexicographical order!
+
+VARIABLE_SPECIFIER(AdvancedDisplay)
+VARIABLE_SPECIFIER(AssetRegistrySearchable)
+VARIABLE_SPECIFIER(BlueprintAssignable)
+VARIABLE_SPECIFIER(BlueprintAuthorityOnly)
+VARIABLE_SPECIFIER(BlueprintCallable)
+VARIABLE_SPECIFIER(BlueprintReadOnly)
+VARIABLE_SPECIFIER(BlueprintReadWrite)
+VARIABLE_SPECIFIER(Config)
+VARIABLE_SPECIFIER(Const)
+VARIABLE_SPECIFIER(DuplicateTransient)
+VARIABLE_SPECIFIER(EditAnywhere)
+VARIABLE_SPECIFIER(EditDefaultsOnly) 
+VARIABLE_SPECIFIER(EditFixedSize)
+VARIABLE_SPECIFIER(EditInline)
+VARIABLE_SPECIFIER(EditInstanceOnly)
+VARIABLE_SPECIFIER(Export)
+VARIABLE_SPECIFIER(GlobalConfig)
+VARIABLE_SPECIFIER(Instanced)
+VARIABLE_SPECIFIER(Interp)
+VARIABLE_SPECIFIER(Localized)
+VARIABLE_SPECIFIER(NoClear)
+VARIABLE_SPECIFIER(NonPIEDuplicateTransient)
+VARIABLE_SPECIFIER(NonPIETransient)
+VARIABLE_SPECIFIER(NonTransactional)
+VARIABLE_SPECIFIER(NotReplicated)
+VARIABLE_SPECIFIER(Ref)
+VARIABLE_SPECIFIER(Replicated)
+VARIABLE_SPECIFIER(ReplicatedUsing)
+VARIABLE_SPECIFIER(RepRetry)
+VARIABLE_SPECIFIER(SaveGame)
+VARIABLE_SPECIFIER(SimpleDisplay)
+VARIABLE_SPECIFIER(TextExportTransient)
+VARIABLE_SPECIFIER(Transient)
+VARIABLE_SPECIFIER(VisibleAnywhere) 
+VARIABLE_SPECIFIER(VisibleDefaultsOnly)
+VARIABLE_SPECIFIER(VisibleInstanceOnly)
diff --git a/Engine/Source/Programs/UnrealHeaderTool/Private/Specifiers/VariableSpecifiers.h b/Engine/Source/Programs/UnrealHeaderTool/Private/Specifiers/VariableSpecifiers.h
new file mode 100644
--- /dev/null
+++ b/Engine/Source/Programs/UnrealHeaderTool/Private/Specifiers/VariableSpecifiers.h
@@ -0,0 +1,16 @@
+// Copyright 1998-2015 Epic Games, Inc. All Rights Reserved.
+
+#pragma once
+
+enum class EVariableSpecifier
+{
+	None = -1
+
+	#define VARIABLE_SPECIFIER(SpecifierName) , SpecifierName
+		#include "VariableSpecifiers.def"
+	#undef VARIABLE_SPECIFIER
+
+	, Max
+};
+
+extern const TCHAR* GVariableSpecifierStrings[(int32)EVariableSpecifier::Max];
diff --git a/Engine/Source/Runtime/Core/Private/Containers/Algo/FindSortedStringCaseInsensitive.cpp b/Engine/Source/Runtime/Core/Private/Containers/Algo/FindSortedStringCaseInsensitive.cpp
new file mode 100644
--- /dev/null
+++ b/Engine/Source/Runtime/Core/Private/Containers/Algo/FindSortedStringCaseInsensitive.cpp
@@ -0,0 +1,118 @@
+// Copyright 1998-2015 Epic Games, Inc. All Rights Reserved.
+
+#include "CorePrivatePCH.h"
+
+#include "Algo/FindSortedStringCaseInsensitive.h"
+
+namespace Algo
+{
+	/**
+	 * Finds a string in an array of sorted strings, by case-insensitive search, by using binary subdivision of the array.
+	 *
+	 * @param  Str          The string to look for.
+	 * @param  SortedArray  The array of strings to search.  The strings must be sorted lexicographically, case-insensitively.
+	 * @param  ArrayCount   The number of strings in the array.
+	 *
+	 * @return The index of the found string in the array, or -1 if the string was not found.
+	 */
+	int32 FindSortedStringCaseInsensitive(const TCHAR* Str, const TCHAR* const* SortedArray, int32 ArrayCount)
+	{
+		// The current index of the character of Str that we're using for our binary search
+		int32 CharIndex = 0;
+
+		// Begin/end iterators to the range of strings that within SortedArray have matched our string up to CharIndex
+		auto SubArray    = SortedArray;
+		auto SubArrayEnd = SubArray + ArrayCount;
+
+		for (;;)
+		{
+			// If the subarray is empty, then we've failed to find a match
+			if (SubArray == SubArrayEnd)
+			{
+				return -1;
+			}
+
+			// If the subarray is only 1 element, we can use Stricmp to finish the job
+			if (SubArray + 1 == SubArrayEnd)
+			{
+				return (FCString::Stricmp(*SubArray + CharIndex, Str + CharIndex) == 0) ? SubArray - SortedArray : -1;
+			}
+
+			// Get the next char to binary search for
+			TCHAR Ch = FChar::ToLower(Str[CharIndex]);
+
+			// If there are no more characters to search for then we're done
+			if (Ch == TEXT('\0'))
+			{
+				// If the subarray's string is also at the end then we've found a match,
+				// otherwise we've only found a prefix, which is not a match
+				return ((*SubArray)[CharIndex] == TEXT('\0')) ? SubArray - SortedArray : -1;
+			}
+
+			// Now we try to find the equal_range of the character
+			for (;;)
+			{
+				// If the subarray is empty, then we've failed to find a match
+				if (SubArray == SubArrayEnd)
+				{
+					return -1;
+				}
+
+				// Find midpoint of the subarray and its char
+				auto Midpoint = SubArray + (SubArrayEnd - SubArray) / 2;
+				TCHAR ArrCh = FChar::ToLower((*Midpoint)[CharIndex]);
+
+				if (ArrCh < Ch)
+				{
+					// The midpoint char is before the string char, so narrow the subrange to after the midpoint
+					SubArray = Midpoint + 1;
+				}
+				else if (ArrCh > Ch)
+				{
+					// The midpoint char is after the string char, so narrow the subrange to before the midpoint
+					SubArrayEnd = Midpoint;
+				}
+				else
+				{
+					// The midpoint char is the string char, so binary search [SubArray, Midpoint) to find the lower bound
+					auto LowerBoundEnd = Midpoint;
+					while (SubArray != LowerBoundEnd)
+					{
+						auto LowerBoundMid = SubArray + (LowerBoundEnd - SubArray) / 2;
+
+						TCHAR LowerBoundCh = FChar::ToLower((*LowerBoundMid)[CharIndex]);
+						if (LowerBoundCh < Ch)
+						{
+							SubArray = LowerBoundMid + 1;
+						}
+						else
+						{
+							LowerBoundEnd = LowerBoundMid;
+						}
+					}
+
+					// Then binary search [Midpoint, SubArrayEnd) to find the upper bound
+					auto UpperBound = Midpoint + 1;
+					while (UpperBound != SubArrayEnd)
+					{
+						auto UpperBoundMid = UpperBound + (SubArrayEnd - UpperBound) / 2;
+
+						TCHAR UpperBoundCh = FChar::ToLower((*UpperBoundMid)[CharIndex]);
+						if (Ch >= UpperBoundCh)
+						{
+							UpperBound = UpperBoundMid + 1;
+						}
+						else
+						{
+							SubArrayEnd = UpperBoundMid;
+						}
+					}
+
+					break;
+				}
+			}
+
+			++CharIndex;
+		}
+	}
+}
diff --git a/Engine/Source/Runtime/Core/Public/Containers/Algo/FindSortedStringCaseInsensitive.h b/Engine/Source/Runtime/Core/Public/Containers/Algo/FindSortedStringCaseInsensitive.h
new file mode 100644
--- /dev/null
+++ b/Engine/Source/Runtime/Core/Public/Containers/Algo/FindSortedStringCaseInsensitive.h
@@ -0,0 +1,17 @@
+// Copyright 1998-2015 Epic Games, Inc. All Rights Reserved.
+
+#pragma once
+
+namespace Algo
+{
+	/**
+	 * Finds a string in an array of sorted strings, by case-insensitive search, by using binary subdivision of the array.
+	 *
+	 * @param  Str          The string to look for.
+	 * @param  SortedArray  The array of strings to search.  The strings must be sorted lexicographically, case-insensitively.
+	 * @param  ArrayCount   The number of strings in the array.
+	 *
+	 * @return The index of the found string in the array, or -1 if the string was not found.
+	 */
+	CORE_API int32 FindSortedStringCaseInsensitive(const TCHAR* Str, const TCHAR* const* SortedArray, int32 ArrayCount);
+}
diff --git a/Engine/Source/Runtime/Core/Public/Containers/Algo/IsSorted.h b/Engine/Source/Runtime/Core/Public/Containers/Algo/IsSorted.h
new file mode 100644
--- /dev/null
+++ b/Engine/Source/Runtime/Core/Public/Containers/Algo/IsSorted.h
@@ -0,0 +1,131 @@
+// Copyright 1998-2015 Epic Games, Inc. All Rights Reserved.
+
+#pragma once
+
+namespace AlgoImpl
+{
+	template <typename T, typename PredType>
+	bool IsSorted(const T* Range, int32 RangeSize, PredType Pred)
+	{
+		if (RangeSize == 0)
+		{
+			return true;
+		}
+
+		// When comparing N elements, we do N-1 comparisons
+		--RangeSize;
+
+		const T* Next = Range + 1;
+		for (;;)
+		{
+			if (RangeSize == 0)
+			{
+				return true;
+			}
+
+			if (Pred(*Next, *Range))
+			{
+				return false;
+			}
+
+			++Range;
+			++Next;
+			--RangeSize;
+		}
+	}
+
+	template <typename T>
+	struct TLess
+	{
+		FORCEINLINE bool operator()(const T& Lhs, const T& Rhs) const
+		{
+			return Lhs < Rhs;
+		}
+	};
+}
+
+namespace Algo
+{
+	/**
+	 * Tests is a range is sorted.
+	 *
+	 * @param  Array  The array to test for being sorted.
+	 *
+	 * @return true if the range is sorted, false otherwise.
+	 */
+	template <typename T, int32 ArraySize>
+	FORCEINLINE bool IsSorted(const T (&Array)[ArraySize])
+	{
+		return AlgoImpl::IsSorted((const T*)Array, ArraySize, AlgoImpl::TLess<T>());
+	}
+
+	/**
+	 * Tests is a range is sorted.
+	 *
+	 * @param  Array  The array to test for being sorted.
+	 * @param  Pred   A binary sorting predicate which describes the ordering of the elements in the array.
+	 *
+	 * @return true if the range is sorted, false otherwise.
+	 */
+	template <typename T, int32 ArraySize, typename PredType>
+	FORCEINLINE bool IsSorted(const T (&Array)[ArraySize], PredType Pred)
+	{
+		return AlgoImpl::IsSorted((const T*)Array, ArraySize, Pred);
+	}
+
+	/**
+	 * Tests is a range is sorted.
+	 *
+	 * @param  Array      A pointer to the array to test for being sorted.
+	 * @param  ArraySize  The number of elements in the array.
+	 *
+	 * @return true if the range is sorted, false otherwise.
+	 */
+	template <typename T>
+	FORCEINLINE bool IsSorted(const T* Array, int32 ArraySize)
+	{
+		return AlgoImpl::IsSorted(Array, ArraySize, AlgoImpl::TLess<T>());
+	}
+
+	/**
+	 * Tests is a range is sorted.
+	 *
+	 * @param  Array      A pointer to the array to test for being sorted.
+	 * @param  ArraySize  The number of elements in the array.
+	 * @param  Pred       A binary sorting predicate which describes the ordering of the elements in the array.
+	 *
+	 * @return true if the range is sorted, false otherwise.
+	 */
+	template <typename T, typename PredType>
+	FORCEINLINE bool IsSorted(const T* Array, int32 ArraySize, PredType Pred)
+	{
+		return AlgoImpl::IsSorted(Array, ArraySize, Pred);
+	}
+
+	/**
+	 * Tests is a range is sorted.
+	 *
+	 * @param  Container  The container to test for being sorted.
+	 *
+	 * @return true if the range is sorted, false otherwise.
+	 */
+	template <typename ContainerType>
+	FORCEINLINE bool IsSorted(const ContainerType& Container)
+	{
+		return AlgoImpl::IsSorted(Container.GetData(), Container.Num(), AlgoImpl::TLess<typename ContainerType::ElementType>());
+	}
+
+	/**
+	 * Tests is a range is sorted.
+	 *
+	 * @param  Container  The container to test for being sorted.
+	 * @param  Pred       A binary sorting predicate which describes the ordering of the elements in the array.
+	 *
+	 * @return true if the range is sorted, false otherwise.
+	 */
+	template <typename ContainerType, typename PredType>
+	FORCEINLINE bool IsSorted(const ContainerType& Container, PredType Pred)
+	{
+		return AlgoImpl::IsSorted(Container.GetData(), Container.Num(), Pred);
+	}
+}
diff --git a/Engine/Source/Runtime/CoreUObject/Private/UObject/Class.cpp b/Engine/Source/Runtime/CoreUObject/Private/UObject/Class.cpp
--- a/Engine/Source/Runtime/CoreUObject/Private/UObject/Class.cpp
+++ b/Engine/Source/Runtime/CoreUObject/Private/UObject/Class.cpp
@@ -20,45 +20,6 @@ DECLARE_LOG_CATEGORY_EXTERN(LogScriptSerialization, Log, All);
 DEFINE_LOG_CATEGORY(LogScriptSerialization);
 DEFINE_LOG_CATEGORY(LogClass);
 
-//////////////////////////////////////////////////////////////////////////
-// FPropertySpecifier
-
-FString FPropertySpecifier::ConvertToString() const
-{
-	FString Result;
-
-	// Emit the specifier key
-	Result += Key;
-
-	// Emit the values if there are any
-	if (Values.Num())
-	{
-		Result += TEXT("=");
-
-		if (Values.Num() == 1)
-		{
-			// One value goes on it's own
-			Result += Values[0];
-		}
-		else
-		{
-			// More than one value goes in parens, separated by commas
-			Result += TEXT("(");
-			for (int32 ValueIndex = 0; ValueIndex < Values.Num(); ++ValueIndex)
-			{
-				if (ValueIndex > 0)
-				{
-					Result += TEXT(", ");
-				}
-				Result += Values[ValueIndex];
-			}
-			Result += TEXT(")");
-		}
-	}
-
-	return Result;
-}
-
 //////////////////////////////////////////////////////////////////////////
 
 /**
diff --git a/Engine/Source/Runtime/CoreUObject/Public/UObject/Class.h b/Engine/Source/Runtime/CoreUObject/Public/UObject/Class.h
--- a/Engine/Source/Runtime/CoreUObject/Public/UObject/Class.h
+++ b/Engine/Source/Runtime/CoreUObject/Public/UObject/Class.h
@@ -18,16 +18,6 @@ COREUOBJECT_API DECLARE_LOG_CATEGORY_EXTERN(LogClass, Log, All);
 
 struct FPropertyTag;
 
-// A specifier with optional value
-struct FPropertySpecifier
-{
-public:
-	FString Key;
-	TArray<FString> Values;
-
-	COREUOBJECT_API FString ConvertToString() const;
-};
-
 /*-----------------------------------------------------------------------------
 	FRepRecord.
 -----------------------------------------------------------------------------*/
