Forward declarations of DelegateWrapper functions removed from generated headers.

diff --git a/Engine/Source/Programs/UnrealHeaderTool/Private/CodeGenerator.cpp b/Engine/Source/Programs/UnrealHeaderTool/Private/CodeGenerator.cpp
--- a/Engine/Source/Programs/UnrealHeaderTool/Private/CodeGenerator.cpp
+++ b/Engine/Source/Programs/UnrealHeaderTool/Private/CodeGenerator.cpp
@@ -3144,11 +3144,7 @@ void FNativeClassHeaderGenerator::ExportDelegateDefinitions(FUnrealSourceFile& S
 			WriteEventFunctionEpilogue(DelegateOutput, 0, Parameters, *DelegateName);
 		}
 
-		if (!bWrapperImplementationsOnly)
-		{
-			HeaderOutput.Log(DelegateOutput);
-		}
-		else
+		if (bWrapperImplementationsOnly)
 		{
 			FString MacroName = SourceFile.GetGeneratedMacroName(FunctionData.MacroLine, TEXT("_DELEGATE"));
 			WriteMacro(HeaderOutput, MacroName, DelegateOutput);
