Subobject error removed.

diff --git a/Engine/Source/Programs/UnrealHeaderTool/Private/HeaderParser.cpp b/Engine/Source/Programs/UnrealHeaderTool/Private/HeaderParser.cpp
--- a/Engine/Source/Programs/UnrealHeaderTool/Private/HeaderParser.cpp
+++ b/Engine/Source/Programs/UnrealHeaderTool/Private/HeaderParser.cpp
@@ -3534,15 +3534,6 @@ void FHeaderParser::GetVarType
 				FError::Throwf(TEXT("Unrecognized type '%s' - type must be a UCLASS, USTRUCT or UENUM"), VarType.Identifier );
 			}
 		}
-
-		if ((Flags & CPF_InstancedReference) && CurrentAccessSpecifier == ACCESS_Private && VariableCategory == EVariableCategory::Member)
-		{
-			if (((Flags & CPF_Edit) && (Flags & CPF_EditConst) == 0) ||
-				((Flags & CPF_BlueprintVisible) && (Flags & CPF_BlueprintReadOnly) == 0))
-			{
-				FError::Throwf(TEXT("%s: Subobject (instanced) properties can't be editable (use VisibleAnywhere or BlueprintReadOnly instead)."), VarType.Identifier);
-			}
-		}
 	}
 
 	if (VariableCategory != EVariableCategory::Member)
