Fix issues with comment parsing in UHT.

diff --git a/Engine/Source/Programs/UnrealHeaderTool/Private/BaseParser.cpp b/Engine/Source/Programs/UnrealHeaderTool/Private/BaseParser.cpp
--- a/Engine/Source/Programs/UnrealHeaderTool/Private/BaseParser.cpp
+++ b/Engine/Source/Programs/UnrealHeaderTool/Private/BaseParser.cpp
@@ -184,14 +184,14 @@ void FBaseParser::ResetParser(const TCHAR* SourceBuffer, int32 StartingLineNumbe
 //
 TCHAR FBaseParser::GetChar(bool bLiteral)
 {
-	int32 CommentCount = 0;
+	bool bInsideComment = false;
 
 	PrevPos = InputPos;
 	PrevLine = InputLine;
 
 Loop:
 	const TCHAR c = Input[InputPos++];
-	if ( CommentCount > 0 )
+	if (bInsideComment)
 	{
 		// Record the character as a comment.
 		PrevComment += c;
@@ -206,24 +206,29 @@ Loop:
 		const TCHAR NextChar = PeekChar();
 		if ( c==TEXT('/') && NextChar==TEXT('*') )
 		{
-			if ( CommentCount == 0 )
+			if (!bInsideComment)
 			{
 				ClearComment();
 				// Record the slash and star.
 				PrevComment += c;
 				PrevComment += NextChar;
+				bInsideComment = true;
 			}
-			CommentCount++;
+
 			InputPos++;
 			goto Loop;
 		}
 		else if( c==TEXT('*') && NextChar==TEXT('/') )
 		{
-			if (--CommentCount < 0)
+			if (!bInsideComment)
 			{
 				ClearComment();
 				FError::Throwf(TEXT("Unexpected '*/' outside of comment") );
 			}
+
+			/** Asterisk and slash always end comment. */
+			bInsideComment = false;
+
 			// Star already recorded; record the slash.
 			PrevComment += Input[InputPos];
 
@@ -232,7 +237,7 @@ Loop:
 		}
 	}
 
-	if (CommentCount > 0)
+	if (bInsideComment)
 	{
 		if (c == 0)
 		{
