Make UHT parse TSubclassOf<USomeClass>& parameters to RPC without compilation error.

diff --git a/Engine/Source/Programs/UnrealHeaderTool/Private/CodeGenerator.cpp b/Engine/Source/Programs/UnrealHeaderTool/Private/CodeGenerator.cpp
--- a/Engine/Source/Programs/UnrealHeaderTool/Private/CodeGenerator.cpp
+++ b/Engine/Source/Programs/UnrealHeaderTool/Private/CodeGenerator.cpp
@@ -4065,6 +4065,8 @@ void FNativeClassHeaderGenerator::ExportFunctionThunk(FUHTStringBuilder& RPCWrap
 		FString EvalParameterText;				// e.g. (UObject*,NULL)
 
 		FString TypeText;
+		bool bPassAsNoPtr = false;
+
 		if (Param->ArrayDim > 1)
 		{
 			EvalBaseText += TEXT("ARRAY");
@@ -4074,8 +4076,17 @@ void FNativeClassHeaderGenerator::ExportFunctionThunk(FUHTStringBuilder& RPCWrap
 		{
 			EvalBaseText += Param->GetCPPMacroType(TypeText);
 		}
+
+		if (Param->HasAllPropertyFlags(CPF_UObjectWrapper | CPF_OutParm) 
+			&& Param->IsA(UClassProperty::StaticClass()))
+		{
+			TypeText = Param->GetCPPType();
+			bPassAsNoPtr = true;
+		}
+
 		FUHTStringBuilder ReplacementText;
 		ReplacementText += TypeText;
+		
 		ApplyAlternatePropertyExportText(Param, ReplacementText);
 		TypeText = ReplacementText;
 
@@ -4085,7 +4096,16 @@ void FNativeClassHeaderGenerator::ExportFunctionThunk(FUHTStringBuilder& RPCWrap
 		// if this property is an out parm, add the REF tag
 		if (Param->PropertyFlags & CPF_OutParm)
 		{
+			if (!bPassAsNoPtr)
+			{
 			EvalModifierText += TEXT("_REF");
+			}
+			else
+			{
+				// Parameters passed as TSubclassOf<Class>& shouldn't have asterisk added.
+				EvalModifierText += TEXT("_REF_NO_PTR");	
+			}
+
 			ParamPrefix += TEXT("Out_");
 		}
 
diff --git a/Engine/Source/Runtime/CoreUObject/Public/UObject/ScriptMacros.h b/Engine/Source/Runtime/CoreUObject/Public/UObject/ScriptMacros.h
--- a/Engine/Source/Runtime/CoreUObject/Public/UObject/ScriptMacros.h
+++ b/Engine/Source/Runtime/CoreUObject/Public/UObject/ScriptMacros.h
@@ -57,6 +57,9 @@ enum {MAX_VARIABLE_SIZE = 0x0FFF };
 #define P_GET_OBJECT(ObjectType,ParamName)			PARAM_PASSED_BY_VAL_ZEROED(ParamName, UObjectPropertyBase, ObjectType*)
 #define P_GET_OBJECT_REF(ObjectType,ParamName)		PARAM_PASSED_BY_REF_ZEROED(ParamName, UObjectPropertyBase, ObjectType*)
 
+#define P_GET_OBJECT_NO_PTR(ObjectType,ParamName)			PARAM_PASSED_BY_VAL_ZEROED(ParamName, UObjectPropertyBase, ObjectType)
+#define P_GET_OBJECT_REF_NO_PTR(ObjectType,ParamName)		PARAM_PASSED_BY_REF_ZEROED(ParamName, UObjectPropertyBase, ObjectType)
+
 #define P_GET_TARRAY(ElementType,ParamName)			PARAM_PASSED_BY_VAL(ParamName, UArrayProperty, TArray<ElementType>)
 #define P_GET_TARRAY_REF(ElementType,ParamName)		PARAM_PASSED_BY_REF(ParamName, UArrayProperty, TArray<ElementType>)
 
