Bad UENUM specifiers now throw an error.

diff --git a/Engine/Source/Programs/UnrealHeaderTool/Private/HeaderParser.cpp b/Engine/Source/Programs/UnrealHeaderTool/Private/HeaderParser.cpp
--- a/Engine/Source/Programs/UnrealHeaderTool/Private/HeaderParser.cpp
+++ b/Engine/Source/Programs/UnrealHeaderTool/Private/HeaderParser.cpp
@@ -1023,6 +1023,12 @@ UEnum* FHeaderParser::CompileEnum(FUnrealSourceFile& SourceFile)
 	TArray<FPropertySpecifier> SpecifiersFound;
 	ReadSpecifierSetInsideMacro(SpecifiersFound, TEXT("Enum"), EnumToken.MetaData);
 
+	// We don't handle any non-metadata enum specifiers at the moment
+	if (SpecifiersFound.Num() != 0)
+	{
+		FError::Throwf(TEXT("Unknown enum specifier '%s'"), *SpecifiersFound[0].Key);
+	}
+
 	FScriptLocation DeclarationPosition;
 
 	// Check enum type. This can be global 'enum', 'namespace' or 'enum class' enums.
