Prevent UHT crash when looking up parent types during code generation.

diff --git a/Engine/Source/Programs/UnrealHeaderTool/Private/CodeGenerator.cpp b/Engine/Source/Programs/UnrealHeaderTool/Private/CodeGenerator.cpp
--- a/Engine/Source/Programs/UnrealHeaderTool/Private/CodeGenerator.cpp
+++ b/Engine/Source/Programs/UnrealHeaderTool/Private/CodeGenerator.cpp
@@ -5335,14 +5335,12 @@ void ResolveSuperClasses(UPackage* Package)
 			if (FoundBaseClass == nullptr)
 			{
 				// Don't know its parent class. Raise error.
-				UE_LOG(LogCompile, Error, TEXT("Couldn't find parent type for '%s' named '%s' in current module or any other module parsed so far."),
+				FError::Throwf(TEXT("Couldn't find parent type for '%s' named '%s' in current module or any other module parsed so far."),
 					*DefinedClass->GetName(), *ParsingInfo.GetBaseClassName());
 			}
-			else
-			{
-				DefinedClass->SetSuperStruct(FoundBaseClass);
-				DefinedClass->ClassCastFlags |= FoundBaseClass->ClassCastFlags;
-			}
+
+			DefinedClass->SetSuperStruct(FoundBaseClass);
+			DefinedClass->ClassCastFlags |= FoundBaseClass->ClassCastFlags;
 		}
 	}
 }
@@ -5509,11 +5507,8 @@ ECompilationResult::Type UnrealHeaderTool_Main(const FString& ModuleInfoFilename
 				{
 					TGuardValue<ELogTimes::Type> DisableLogTimes(GPrintLogTimes, ELogTimes::None);
 
-					FString Prefix;
-
-					const FString AbsFilename = IFileManager::Get().ConvertToAbsolutePathForExternalAppForRead(*RawFilename);
-					Prefix = FString::Printf(TEXT("%s(1): "), *AbsFilename);
-
+					FString AbsFilename           = IFileManager::Get().ConvertToAbsolutePathForExternalAppForRead(*RawFilename);
+					FString Prefix                = FString::Printf(TEXT("%s(1): "), *AbsFilename);
 					FString FormattedErrorMessage = FString::Printf(TEXT("%sError: %s\r\n"), *Prefix, ErrorMsg);
 					Result = GCompilationResult;
 
@@ -5530,7 +5525,27 @@ ECompilationResult::Type UnrealHeaderTool_Main(const FString& ModuleInfoFilename
 			}
 		}
 
-		ResolveSuperClasses(Package);
+	#if !PLATFORM_EXCEPTIONS_DISABLED
+		try
+	#endif
+		{
+			ResolveSuperClasses(Package);
+		}
+	#if !PLATFORM_EXCEPTIONS_DISABLED
+		catch (TCHAR* ErrorMsg)
+		{
+			TGuardValue<ELogTimes::Type> DisableLogTimes(GPrintLogTimes, ELogTimes::None);
+
+			FString FormattedErrorMessage = FString::Printf(TEXT("Error: %s\r\n"), ErrorMsg);
+
+			Result = GCompilationResult;
+
+			UE_LOG(LogCompile, Log, TEXT("%s"), *FormattedErrorMessage);
+			GWarn->Log(ELogVerbosity::Error, FormattedErrorMessage);
+
+			++NumFailures;
+		}
+	#endif
 
 		ThisModuleTimer.Stop();
 		TotalModulePreparseTime += ThisModulePreparseTime;
