UHT: Added class range validation in exporting code to prevent crashes when they occur during bad parsing.

diff --git a/Engine/Source/Programs/UnrealHeaderTool/Private/CodeGenerator.cpp b/Engine/Source/Programs/UnrealHeaderTool/Private/CodeGenerator.cpp
--- a/Engine/Source/Programs/UnrealHeaderTool/Private/CodeGenerator.cpp
+++ b/Engine/Source/Programs/UnrealHeaderTool/Private/CodeGenerator.cpp
@@ -2049,6 +2049,7 @@ void FNativeClassHeaderGenerator::ExportClassesFromSourceFileInner(FUnrealSource
 			if (ClassDefinitionRanges.Contains(Class))
 			{
 				ClassRange = ClassDefinitionRanges[Class];
+				ClassRange.Validate();
 			}
 
 			auto ClassStart = ClassRange.Start;
@@ -2143,6 +2144,7 @@ void FNativeClassHeaderGenerator::ExportClassesFromSourceFileInner(FUnrealSource
 			if (ClassDefinitionRanges.Contains(Class))
 			{
 				ClassRange = ClassDefinitionRanges[Class];
+				ClassRange.Validate();
 			}
 
 			auto ClassStart = ClassRange.Start;
@@ -4180,6 +4182,7 @@ void FNativeClassHeaderGenerator::ExportFunctionThunk(FUHTStringBuilder& RPCWrap
 	if (ClassDefinitionRanges.Contains(Function->GetOwnerClass()))
 	{
 		ClassRange = ClassDefinitionRanges[Function->GetOwnerClass()];
+		ClassRange.Validate();
 	}
 
 	auto ClassStart = ClassRange.Start;
@@ -4285,6 +4288,7 @@ void FNativeClassHeaderGenerator::ExportNativeFunctions(FUnrealSourceFile& Sourc
 	if (ClassDefinitionRanges.Contains(Class))
 	{
 		ClassRange = ClassDefinitionRanges[Class];
+		ClassRange.Validate();
 	}
 
 	// export the C++ stubs
diff --git a/Engine/Source/Programs/UnrealHeaderTool/Private/HeaderParser.h b/Engine/Source/Programs/UnrealHeaderTool/Private/HeaderParser.h
--- a/Engine/Source/Programs/UnrealHeaderTool/Private/HeaderParser.h
+++ b/Engine/Source/Programs/UnrealHeaderTool/Private/HeaderParser.h
@@ -115,6 +115,14 @@ struct ClassDefinitionRange
 		, bHasGeneratedBody(false)
 	{ }
 
+	void Validate()
+	{
+		if (End <= Start)
+		{
+			FError::Throwf(TEXT("The class definition range is invalid. Most probably caused by previous parsing error."));
+		}
+	}
+
 	const TCHAR* Start;
 	const TCHAR* End;
 	bool bHasGeneratedBody;
