better %f seconds printout

diff --git a/Engine/Source/Editor/UnrealEd/Private/NormalMapIdentification.cpp b/Engine/Source/Editor/UnrealEd/Private/NormalMapIdentification.cpp
--- a/Engine/Source/Editor/UnrealEd/Private/NormalMapIdentification.cpp
+++ b/Engine/Source/Editor/UnrealEd/Private/NormalMapIdentification.cpp
@@ -394,7 +394,7 @@ static bool IsTextureANormalMap( UTexture* Texture )
 #if NORMALMAP_IDENTIFICATION_TIMING
 	double EndSeconds = FPlatformTime::Seconds();
 
-	FString Msg = FString::Printf( TEXT("NormalMapIdentification took %f seconds to analyze %s"), (EndSeconds-StartSeconds), *Texture->GetFullName() ); 
+	FString Msg = FString::Printf( TEXT("NormalMapIdentification took %.2f seconds to analyze %s"), (EndSeconds-StartSeconds), *Texture->GetFullName() ); 
 
 	GLog->Log(Msg);
 #endif
diff --git a/Engine/Source/Programs/UnrealHeaderTool/Private/CodeGenerator.cpp b/Engine/Source/Programs/UnrealHeaderTool/Private/CodeGenerator.cpp
--- a/Engine/Source/Programs/UnrealHeaderTool/Private/CodeGenerator.cpp
+++ b/Engine/Source/Programs/UnrealHeaderTool/Private/CodeGenerator.cpp
@@ -5653,12 +5653,12 @@ ECompilationResult::Type UnrealHeaderTool_Main(const FString& ModuleInfoFilename
 	// Avoid TArray slack for meta data.
 	GScriptHelper.Shrink();
 
-	UE_LOG(LogCompile, Log, TEXT("Preparsing %i modules took %f seconds"), GManifest.Modules.Num(), TotalModulePreparseTime);
-	UE_LOG(LogCompile, Log, TEXT("Parsing took %f seconds"), TotalParseAndCodegenTime - GHeaderCodeGenTime);
-	UE_LOG(LogCompile, Log, TEXT("Code generation took %f seconds"), GHeaderCodeGenTime);
-	UE_LOG(LogCompile, Log, TEXT("ScriptPlugin overhead was %f seconds"), GPluginOverheadTime);
-	UE_LOG(LogCompile, Log, TEXT("Macroize time was %f seconds"), GMacroizeTime);
-	UE_LOG(LogCompile, Log, TEXT("Tabify time was was %f seconds"), GTabifyTime);
+	UE_LOG(LogCompile, Log, TEXT("Preparsing %i modules took %.2f seconds"), GManifest.Modules.Num(), TotalModulePreparseTime);
+	UE_LOG(LogCompile, Log, TEXT("Parsing took %.2f seconds"), TotalParseAndCodegenTime - GHeaderCodeGenTime);
+	UE_LOG(LogCompile, Log, TEXT("Code generation took %.2f seconds"), GHeaderCodeGenTime);
+	UE_LOG(LogCompile, Log, TEXT("ScriptPlugin overhead was %.2f seconds"), GPluginOverheadTime);
+	UE_LOG(LogCompile, Log, TEXT("Macroize time was %.2f seconds"), GMacroizeTime);
+	UE_LOG(LogCompile, Log, TEXT("Tabify time was was %.2f seconds"), GTabifyTime);
 
 	if (bWriteContents)
 	{
diff --git a/Engine/Source/Runtime/CoreUObject/Private/UObject/SavePackage.cpp b/Engine/Source/Runtime/CoreUObject/Private/UObject/SavePackage.cpp
--- a/Engine/Source/Runtime/CoreUObject/Private/UObject/SavePackage.cpp
+++ b/Engine/Source/Runtime/CoreUObject/Private/UObject/SavePackage.cpp
@@ -4283,7 +4283,7 @@ bool UPackage::SavePackage( UPackage* InOuter, UObject* Base, EObjectFlags TopLe
 					Linker->Detach();
 				}
 				UNCLOCK_CYCLES(Time);
-				UE_LOG(LogSavePackage, Log,  TEXT("Save=%fms"), FPlatformTime::ToMilliseconds(Time) );
+				UE_LOG(LogSavePackage, Log,  TEXT("Save=%.2fms"), FPlatformTime::ToMilliseconds(Time) );
 		
 				if ( EndSavingIfCancelled( Linker, TempFilename ) ) { return false; }
 				SlowTask.EnterProgressFrame();
