Added "//~" and "/*~" support to Unreal header parsing, so that comments can be ignored instead of being made into tooltips, mirroring the recent Doxygen filter feature. Also added support to keep [REDACTED_HANDLE] blocks visible in UFUNCTIONS.

diff --git a/Engine/Source/Editor/BlueprintGraph/Private/K2Node_CallFunction.cpp b/Engine/Source/Editor/BlueprintGraph/Private/K2Node_CallFunction.cpp
--- a/Engine/Source/Editor/BlueprintGraph/Private/K2Node_CallFunction.cpp
+++ b/Engine/Source/Editor/BlueprintGraph/Private/K2Node_CallFunction.cpp
@@ -1367,10 +1367,13 @@ FString UK2Node_CallFunction::GetDefaultTooltipForFunction(const UFunction* Func
 		static const FString DoxygenReturn(TEXT("@return"));
 		static const FString DoxygenSee(TEXT("@see"));
 		static const FString TooltipSee(TEXT("See:"));
+		static const FString DoxygenNote(TEXT("@note"));
+		static const FString TooltipNote(TEXT("Note:"));
 
 		Tooltip.Split(DoxygenParam, &Tooltip, nullptr, ESearchCase::IgnoreCase, ESearchDir::FromStart);
 		Tooltip.Split(DoxygenReturn, &Tooltip, nullptr, ESearchCase::IgnoreCase, ESearchDir::FromStart);
 		Tooltip.ReplaceInline(*DoxygenSee, *TooltipSee);
+		Tooltip.ReplaceInline(*DoxygenNote, *TooltipNote);
 
 		Tooltip.Trim();
 		Tooltip.TrimTrailing();
diff --git a/Engine/Source/Programs/UnrealHeaderTool/Private/HeaderParser.cpp b/Engine/Source/Programs/UnrealHeaderTool/Private/HeaderParser.cpp
--- a/Engine/Source/Programs/UnrealHeaderTool/Private/HeaderParser.cpp
+++ b/Engine/Source/Programs/UnrealHeaderTool/Private/HeaderParser.cpp
@@ -1337,11 +1337,47 @@ FString FHeaderParser::FormatCommentForToolTip(const FString& Input)
 		return FString( TEXT("") );
 	}
 
+	FString Result(Input);
+
+	// Sweep out comments marked to be ignored.
+	{
+		int32 CommentStart, CommentEnd;
+		// Block comments go first
+		for (CommentStart = Result.Find(TEXT("/*~")); CommentStart != INDEX_NONE; CommentStart = Result.Find(TEXT("/*~")))
+		{
+			CommentEnd = Result.Find(TEXT("*/"), ESearchCase::CaseSensitive, ESearchDir::FromStart, CommentStart);
+			if (CommentEnd != INDEX_NONE)
+			{
+				Result.RemoveAt(CommentStart, (CommentEnd + 2) - CommentStart, false);
+			}
+			else
+			{
+				// This looks like an error - an unclosed block comment.
+				break;
+			}
+		}
+		// Leftover line comments go next
+		for (CommentStart = Result.Find(TEXT("//~")); CommentStart != INDEX_NONE; CommentStart = Result.Find(TEXT("//~")))
+		{
+			CommentEnd = Result.Find(TEXT("\n"), ESearchCase::CaseSensitive, ESearchDir::FromStart, CommentStart);
+			if (CommentEnd != INDEX_NONE)
+			{
+				Result.RemoveAt(CommentStart, (CommentEnd + 1) - CommentStart, false);
+			}
+			else
+			{
+				Result.RemoveAt(CommentStart, Result.Len() - CommentStart, false);
+				break;
+			}
+		}
+		// Finish by shrinking if anything was removed, since we deferred this during the search.
+		Result.Shrink();
+	}
+
 	// Check for known commenting styles.
-	FString Result( Input );
-	const bool bJavaDocStyle = Input.Contains(TEXT("/**"));
-	const bool bCStyle = Input.Contains(TEXT("/*"));
-	const bool bCPPStyle = Input.StartsWith(TEXT("//"));
+	const bool bJavaDocStyle = Result.Contains(TEXT("/**"));
+	const bool bCStyle = Result.Contains(TEXT("/*"));
+	const bool bCPPStyle = Result.StartsWith(TEXT("//"));
 
 	if ( bJavaDocStyle || bCStyle)
 	{
