Cope with explicit when looking for user constructors.

diff --git a/Engine/Source/Programs/UnrealHeaderTool/Private/HeaderParser.cpp b/Engine/Source/Programs/UnrealHeaderTool/Private/HeaderParser.cpp
--- a/Engine/Source/Programs/UnrealHeaderTool/Private/HeaderParser.cpp
+++ b/Engine/Source/Programs/UnrealHeaderTool/Private/HeaderParser.cpp
@@ -4129,8 +4129,10 @@ bool FHeaderParser::CompileDeclaration(FClasses& AllClasses, FUnrealSourceFile&
 		}
 		check(TopNest->NestType == NEST_Class || TopNest->NestType == NEST_Interface || TopNest->NestType == NEST_NativeInterface);
 		CurrentAccessSpecifier = AccessSpecifier;
+		return true;
 	}
-	else if (Token.Matches(TEXT("class")) && (TopNest->NestType == NEST_GlobalScope))
+
+	if (Token.Matches(TEXT("class")) && (TopNest->NestType == NEST_GlobalScope))
 	{
 		// Make sure the previous class ended with valid nesting.
 		if (bEncounteredNewStyleClass_UnmatchedBrackets)
@@ -4148,8 +4150,10 @@ bool FHeaderParser::CompileDeclaration(FClasses& AllClasses, FUnrealSourceFile&
 			UngetToken(Token);
 			return SkipDeclaration(Token);
 		}
+		return true;
 	}
-	else if (Token.Matches(TEXT("GENERATED_IINTERFACE_BODY")) || (Token.Matches(TEXT("GENERATED_BODY")) && TopNest->NestType == NEST_NativeInterface))
+
+	if (Token.Matches(TEXT("GENERATED_IINTERFACE_BODY")) || (Token.Matches(TEXT("GENERATED_BODY")) && TopNest->NestType == NEST_NativeInterface))
 	{
 		if (TopNest->NestType != NEST_NativeInterface)
 		{
@@ -4175,8 +4179,10 @@ bool FHeaderParser::CompileDeclaration(FClasses& AllClasses, FUnrealSourceFile&
 		{
 			ClassDefinitionRanges[GetCurrentClass()].bHasGeneratedBody = true;
 		}
+		return true;
 	}
-	else if (Token.Matches(TEXT("GENERATED_UINTERFACE_BODY")) || (Token.Matches(TEXT("GENERATED_BODY")) && TopNest->NestType == NEST_Interface))
+
+	if (Token.Matches(TEXT("GENERATED_UINTERFACE_BODY")) || (Token.Matches(TEXT("GENERATED_BODY")) && TopNest->NestType == NEST_Interface))
 	{
 		if (TopNest->NestType != NEST_Interface)
 		{
@@ -4197,8 +4203,10 @@ bool FHeaderParser::CompileDeclaration(FClasses& AllClasses, FUnrealSourceFile&
 		{
 			CurrentAccessSpecifier = ACCESS_Public;
 		}
+		return true;
 	}
-	else if (Token.Matches(TEXT("GENERATED_UCLASS_BODY")) || (Token.Matches(TEXT("GENERATED_BODY")) && TopNest->NestType == NEST_Class))
+
+	if (Token.Matches(TEXT("GENERATED_UCLASS_BODY")) || (Token.Matches(TEXT("GENERATED_BODY")) && TopNest->NestType == NEST_Class))
 	{
 		if (TopNest->NestType != NEST_Class)
 		{
@@ -4230,54 +4238,74 @@ bool FHeaderParser::CompileDeclaration(FClasses& AllClasses, FUnrealSourceFile&
 		ClassData->SetGeneratedBodyLine(InputLine);
 
 		bClassHasGeneratedBody = true;
+		return true;
 	}
-	else if (Token.Matches(TEXT("UCLASS"), ESearchCase::CaseSensitive) && (TopNest->Allow & ALLOW_Class))
+
+	if (Token.Matches(TEXT("UCLASS"), ESearchCase::CaseSensitive) && (TopNest->Allow & ALLOW_Class))
 	{
 		bHaveSeenUClass = true;
 		bEncounteredNewStyleClass_UnmatchedBrackets = true;
 		CompileClassDeclaration(AllClasses);
+		return true;
 	}
-	else if (Token.Matches(TEXT("UINTERFACE")) && (TopNest->Allow & ALLOW_Class))
+
+	if (Token.Matches(TEXT("UINTERFACE")) && (TopNest->Allow & ALLOW_Class))
 	{
 		bHaveSeenUClass = true;
 		bEncounteredNewStyleClass_UnmatchedBrackets = true;
 		CompileInterfaceDeclaration(AllClasses);
+		return true;
 	}
-	else if (Token.Matches(TEXT("UFUNCTION"), ESearchCase::CaseSensitive))
+
+	if (Token.Matches(TEXT("UFUNCTION"), ESearchCase::CaseSensitive))
 	{
 		CompileFunctionDeclaration(SourceFile, AllClasses);
+		return true;
 	}
-	else if (Token.Matches(TEXT("UDELEGATE")))
+
+	if (Token.Matches(TEXT("UDELEGATE")))
 	{
 		CompileDelegateDeclaration(SourceFile, AllClasses, Token.Identifier, EDelegateSpecifierAction::Parse);
+		return true;
 	}
-	else if (IsValidDelegateDeclaration(Token)) // Legacy delegate parsing - it didn't need a UDELEGATE
+
+	if (IsValidDelegateDeclaration(Token)) // Legacy delegate parsing - it didn't need a UDELEGATE
 	{
 		CompileDelegateDeclaration(SourceFile, AllClasses, Token.Identifier);
+		return true;
 	}
-	else if (Token.Matches(TEXT("UPROPERTY"), ESearchCase::CaseSensitive))
+
+	if (Token.Matches(TEXT("UPROPERTY"), ESearchCase::CaseSensitive))
 	{
 		CheckAllow(TEXT("'Member variable declaration'"), ALLOW_VarDecl);
 		check(TopNest->NestType == NEST_Class);
 
 		CompileVariableDeclaration(AllClasses, GetCurrentClass());
+		return true;
 	}
-	else if (Token.Matches(TEXT("UENUM")))
+
+	if (Token.Matches(TEXT("UENUM")))
 	{
 		// Enumeration definition.
 		CompileEnum(SourceFile);
+		return true;
 	}
-	else if (Token.Matches(TEXT("USTRUCT")))
+
+	if (Token.Matches(TEXT("USTRUCT")))
 	{
 		// Struct definition.
 		CompileStructDeclaration(AllClasses, SourceFile);
+		return true;
 	}
-	else if (Token.Matches(TEXT("#")))
+
+	if (Token.Matches(TEXT("#")))
 	{
 		// Compiler directive.
 		CompileDirective(AllClasses, SourceFile);
+		return true;
 	}
-	else if (bEncounteredNewStyleClass_UnmatchedBrackets && Token.Matches(TEXT("}")))
+
+	if (bEncounteredNewStyleClass_UnmatchedBrackets && Token.Matches(TEXT("}")))
 	{
 		if (ClassDefinitionRanges.Contains(GetCurrentClass()))
 		{
@@ -4329,8 +4357,10 @@ bool FHeaderParser::CompileDeclaration(FClasses& AllClasses, FUnrealSourceFile&
 		bClassHasGeneratedIInterfaceBody = false;
 
 		GetCurrentScope()->AddType(CurrentClass);
+		return true;
 	}
-	else if (Token.Matches(TEXT(";")))
+
+	if (Token.Matches(TEXT(";")))
 	{
 		if (GetToken(Token))
 		{
@@ -4341,23 +4371,40 @@ bool FHeaderParser::CompileDeclaration(FClasses& AllClasses, FUnrealSourceFile&
 			FError::Throwf(TEXT("Extra ';' before end of file"));
 		}
 	}
-	else if (bEncounteredNewStyleClass_UnmatchedBrackets && IsInAClass() && GetCurrentClass() &&
-		(
-			Token.Matches(NameLookupCPP.GetNameCPP(GetCurrentClass())) || 
-			(FString(Token.Identifier).EndsWith("_API") && GetToken(Token) && Token.Matches(NameLookupCPP.GetNameCPP(GetCurrentClass())))
-		))
+
+	if (bEncounteredNewStyleClass_UnmatchedBrackets && IsInAClass())
 	{
-		if (!TryToMatchConstructorParameterList(Token))
+		if (UClass* Class = GetCurrentClass())
 		{
-			return SkipDeclaration(Token);
+			FToken ConstructorToken = Token;
+
+			// Allow explicit constructors
+			bool bFoundExplicit = ConstructorToken.Matches(TEXT("explicit"));
+			if (bFoundExplicit)
+			{
+				GetToken(ConstructorToken);
+			}
+
+			if (FString(ConstructorToken.Identifier).EndsWith("_API"))
+			{
+				if (!bFoundExplicit)
+				{
+					// Explicit can come before or after an _API
+					MatchIdentifier(TEXT("explicit"));
+				}
+
+				GetToken(ConstructorToken);
+			}
+
+			if (ConstructorToken.Matches(NameLookupCPP.GetNameCPP(Class)) && TryToMatchConstructorParameterList(ConstructorToken))
+			{
+				return true;
+			}
 		}
 	}
-	else
-	{
-		// Ignore C++ declaration / function definition. 
-		return SkipDeclaration(Token);
-	}
-	return true;
+
+	// Ignore C++ declaration / function definition. 
+	return SkipDeclaration(Token);
 }
 
 bool FHeaderParser::SkipDeclaration(FToken& Token)
