Allow UHT to parse 'int', 'signed int' and 'unsigned int' types.

diff --git a/Engine/Source/Programs/UnrealHeaderTool/Private/HeaderParser.cpp b/Engine/Source/Programs/UnrealHeaderTool/Private/HeaderParser.cpp
--- a/Engine/Source/Programs/UnrealHeaderTool/Private/HeaderParser.cpp
+++ b/Engine/Source/Programs/UnrealHeaderTool/Private/HeaderParser.cpp
@@ -805,6 +805,32 @@ namespace
 
 		return bSupportedType || (bIsSupportedMemberVariable && bMemberVariable);
 	}
+
+	/**
+	 * Gets property based on compiler/architecture specific int.
+	 * @param bIsSigned Whether the result should be signed or unsigned version of int.
+	 * @return Property corresponding to 'int' type, based on its sign.
+	 */
+	FPropertyBase HandleNativeIntType(bool bIsSigned)
+	{
+		switch (sizeof(int))
+		{
+		case 2:
+			return bIsSigned ? FPropertyBase(CPT_Int16) : FPropertyBase(CPT_UInt16);
+
+		case 4:
+			return bIsSigned ? FPropertyBase(CPT_Int) : FPropertyBase(CPT_UInt32);
+
+		case 8:
+			return bIsSigned ? FPropertyBase(CPT_Int64) : FPropertyBase(CPT_UInt64);
+
+		default:
+			FError::Throwf(TEXT("Found int property of size %d bytes, which is not 2, 4 or 8. Aborting."), sizeof(int));
+		}
+
+		// Should never get here, as we throw in default, but compiler doesn't know it's a throw, so return dummy value.
+		return FPropertyBase(CPT_Int);
+}
 }
 	
 /////////////////////////////////////////////////////
@@ -3084,6 +3110,20 @@ void FHeaderParser::GetVarType
 		// 8-bit bitfield (bool) type
 		VarProperty = FPropertyBase(CPT_Bool8);
 	}
+	else if ( VarType.Matches(TEXT("int")) )
+	{
+		VarProperty = HandleNativeIntType(true);
+	}
+	else if ( VarType.Matches(TEXT("signed")) )
+	{
+		MatchIdentifier(TEXT("int"));
+		VarProperty = HandleNativeIntType(true);
+	}
+	else if (VarType.Matches(TEXT("unsigned")))
+	{
+		MatchIdentifier(TEXT("int"));
+		VarProperty = HandleNativeIntType(false);
+	}
 	else if ( VarType.Matches(TEXT("bool")) )
 	{
 		if (IsBitfieldProperty())
