UEBP-40 C++ code generation (WIP):
- Native Access Specifiers
- FEmitDefaultValueHelper fix/improvements

diff --git a/Engine/Source/Developer/BlueprintCompilerCppBackend/Private/BlueprintCompilerCppBackendUtils.cpp b/Engine/Source/Developer/BlueprintCompilerCppBackend/Private/BlueprintCompilerCppBackendUtils.cpp
--- a/Engine/Source/Developer/BlueprintCompilerCppBackend/Private/BlueprintCompilerCppBackendUtils.cpp
+++ b/Engine/Source/Developer/BlueprintCompilerCppBackend/Private/BlueprintCompilerCppBackendUtils.cpp
@@ -531,7 +531,13 @@ FString FEmitHelper::LiteralTerm(const FEdGraphPinType& Type, const FString& Cus
 		else
 		{
 			//@todo:  This needs to be more robust, since import text isn't really proper for struct construction.
-			return FString(TEXT("F")) + StructType->GetName() + (CustomValue.IsEmpty() ? TEXT("{}") : *CustomValue);
+			ensure(CustomValue.IsEmpty());
+			
+			const FString StructName = FString(TEXT("F")) + StructType->GetName();
+			const FString ConstructBody = CustomValue.IsEmpty()
+				? (StructType->IsA<UUserDefinedStruct>() ? TEXT("::GetDefaultValue()") : TEXT("{}"))
+				: CustomValue;
+			return StructName + ConstructBody;
 		}
 	}
 	else if (Type.PinSubCategory == UEdGraphSchema_K2::PSC_Self)
@@ -670,16 +676,54 @@ FString FEmitDefaultValueHelper::GenerateGetDefaultValue(const UUserDefinedStruc
 	return Result;
 }
 
+bool FEmitDefaultValueHelper::OneLineConstruction(const UProperty* Property, const uint8* ValuePtr, FString& OutResult)
+{
+	FString ValueStr;
+	bool bComplete = true;
+	if (!HandleSpecialTypes(Property, ValuePtr, ValueStr))
+	{
+		auto StructProperty = Cast<const UStructProperty>(Property);
+		Property->ExportTextItem(ValueStr, ValuePtr, ValuePtr, nullptr, EPropertyPortFlags::PPF_ExportCpp);
+		if (ValueStr.IsEmpty() && StructProperty)
+		{
+			check(StructProperty->Struct);
+			ValueStr = FString::Printf(TEXT("F%s()"), *StructProperty->Struct->GetName());
+			bComplete = false;
+		}
+		if (ValueStr.IsEmpty())
+		{
+			UE_LOG(LogK2Compiler, Warning, TEXT("FEmitDefaultValueHelper Cannot generate initilization: %s"), *Property->GetPathName());
+		}
+	}
+	OutResult += ValueStr;
+	return bComplete;
+}
+
 void FEmitDefaultValueHelper::InnerGenerate(const UProperty* Property, const uint8* ValuePtr, const FString& PathToMember, FString& OutResult)
 {
+	if (Property->HasAnyPropertyFlags(CPF_NativeAccessSpecifierPrivate | CPF_NativeAccessSpecifierProtected))
+	{
+		UE_LOG(LogK2Compiler, Warning, TEXT("FEmitDefaultValueHelper Cannot access property: %s"), *Property->GetPathName());
+		return;
+	}
+
+	auto StructProperty = Cast<const UStructProperty>(Property);
+	check(!StructProperty || StructProperty->Struct);
+	auto ArrayProperty = Cast<const UArrayProperty>(Property);
+	check(!ArrayProperty || ArrayProperty->Inner);
+
 	if (!Property->GetOuter()->IsA<UArrayProperty>())
 	{
 		FString ValueStr;
-		Property->ExportTextItem(ValueStr, ValuePtr, ValuePtr, nullptr, EPropertyPortFlags::PPF_ExportCpp);
+		const bool bComplete = OneLineConstruction(Property, ValuePtr, ValueStr);
 		OutResult += FString::Printf(TEXT("%s = %s;\n"), *PathToMember, *ValueStr);
+		if (bComplete && !ArrayProperty)
+		{
+			return;
+		}
 	}
 
-	if (auto StructProperty = Cast<const UStructProperty>(Property))
+	if (StructProperty)
 	{
 		for (auto LocalProperty : TFieldRange<const UProperty>(StructProperty->Struct))
 		{
@@ -693,20 +737,43 @@ void FEmitDefaultValueHelper::InnerGenerate(const UProperty* Property, const uin
 			}
 		}
 	}
-	else if (auto ArrayProperty = Cast<const UArrayProperty>(Property))
+	
+	if (ArrayProperty)
 	{
-		check(ArrayProperty->Inner);
 		FScriptArrayHelper ScriptArrayHelper(ArrayProperty, ValuePtr);
 		for (int32 Index = 0; Index < ScriptArrayHelper.Num(); ++Index)
 		{
 			const uint8* LocalValuePtr = ScriptArrayHelper.GetRawPtr(Index);
+
+			FString ValueStr;
+			const bool bComplete = OneLineConstruction(ArrayProperty->Inner, LocalValuePtr, ValueStr);
+			OutResult += FString::Printf(TEXT("%s.Add(%s);\n"), *PathToMember, *ValueStr);
+			if (!bComplete)
 			{
-				FString ValueStr;
-				ArrayProperty->Inner->ExportTextItem(ValueStr, LocalValuePtr, LocalValuePtr, nullptr, EPropertyPortFlags::PPF_ExportCpp);
-				OutResult += FString::Printf(TEXT("%s.Add(%s);\n"), *PathToMember, *ValueStr);
+				const FString LocalPathToMember = FString::Printf(TEXT("%s[%d]"), *PathToMember, Index);
+				InnerGenerate(ArrayProperty->Inner, LocalValuePtr, LocalPathToMember, OutResult);
 			}
-			const FString LocalPathToMember = FString::Printf(TEXT("%s[%d]"), *PathToMember, Index);
-			InnerGenerate(ArrayProperty->Inner, LocalValuePtr, LocalPathToMember, OutResult);
 		}
 	}
 }
+
+bool FEmitDefaultValueHelper::HandleSpecialTypes(const UProperty* Property, const uint8* ValuePtr, FString& OutResult)
+{
+	if (auto StructProperty = Cast<UStructProperty>(Property))
+	{
+		if (TBaseStructure<FTransform>::Get() == StructProperty->Struct)
+		{
+			check(ValuePtr);
+			const FTransform* Transform = reinterpret_cast<const FTransform*>(ValuePtr);
+			const auto Rotation = Transform->GetRotation();
+			const auto Translation = Transform->GetTranslation();
+			const auto Scale = Transform->GetScale3D();
+			OutResult = FString::Printf(TEXT("FTransform(FQuat(%f, %f, %f, %f), FVector(%f, %f, %f), FVector(%f, %f, %f))")
+				, Rotation.X, Rotation.Y, Rotation.Z, Rotation.W
+				, Translation.X, Translation.Y, Translation.Z
+				, Scale.X, Scale.Y, Scale.Z);
+			return true;
+		}
+	}
+	return false;
+}
\ No newline at end of file
diff --git a/Engine/Source/Developer/BlueprintCompilerCppBackend/Private/BlueprintCompilerCppBackendUtils.h b/Engine/Source/Developer/BlueprintCompilerCppBackend/Private/BlueprintCompilerCppBackendUtils.h
--- a/Engine/Source/Developer/BlueprintCompilerCppBackend/Private/BlueprintCompilerCppBackendUtils.h
+++ b/Engine/Source/Developer/BlueprintCompilerCppBackend/Private/BlueprintCompilerCppBackendUtils.h
@@ -67,4 +67,10 @@ struct FEmitDefaultValueHelper
 {
 	static FString GenerateGetDefaultValue(const UUserDefinedStruct* Struct);
 	static void InnerGenerate(const UProperty* Property, const uint8* ValuePtr, const FString& PathToMember, FString& OutResult);
+	
+	// Return if it's complete construction
+	static bool OneLineConstruction(const UProperty* Property, const uint8* ValuePtr, FString& OutResult);
+
+	//Return if handled
+	static bool HandleSpecialTypes(const UProperty* Property, const uint8* ValuePtr, FString& OutResult);
 };
\ No newline at end of file
diff --git a/Engine/Source/Editor/KismetCompiler/Private/KismetCompiler.cpp b/Engine/Source/Editor/KismetCompiler/Private/KismetCompiler.cpp
--- a/Engine/Source/Editor/KismetCompiler/Private/KismetCompiler.cpp
+++ b/Engine/Source/Editor/KismetCompiler/Private/KismetCompiler.cpp
@@ -121,6 +121,7 @@ void FKismetCompilerContext::SpawnNewClass(const FString& NewClassName)
 	else
 	{
 		// Already existed, but wasn't linked in the Blueprint yet due to load ordering issues
+		NewClass->ClassGeneratedBy = Blueprint;
 		FBlueprintCompileReinstancer::Create(NewClass);
 	}
 }
diff --git a/Engine/Source/Programs/UnrealHeaderTool/Private/HeaderParser.cpp b/Engine/Source/Programs/UnrealHeaderTool/Private/HeaderParser.cpp
--- a/Engine/Source/Programs/UnrealHeaderTool/Private/HeaderParser.cpp
+++ b/Engine/Source/Programs/UnrealHeaderTool/Private/HeaderParser.cpp
@@ -2608,14 +2608,14 @@ void FHeaderParser::CompileDirective(FClasses& AllClasses, FUnrealSourceFile& So
 
 void FHeaderParser::GetVarType
 (
-	FClasses&                       AllClasses,
-	FScope*							Scope,
-	FPropertyBase&                  VarProperty,
-	uint64                          Disallow,
-	FToken*                         OuterPropertyType,
-	EPropertyDeclarationStyle::Type PropertyDeclarationStyle,
-	EVariableCategory::Type         VariableCategory,
-	FIndexRange*                    ParsedVarIndexRange
+FClasses&                       AllClasses,
+FScope*							Scope,
+FPropertyBase&                  VarProperty,
+uint64                          Disallow,
+FToken*                         OuterPropertyType,
+EPropertyDeclarationStyle::Type PropertyDeclarationStyle,
+EVariableCategory::Type         VariableCategory,
+FIndexRange*                    ParsedVarIndexRange
 )
 {
 	UStruct* OwnerStruct = Scope->IsFileScope() ? nullptr : ((FStructScope*)Scope)->GetStruct();
@@ -2638,7 +2638,7 @@ void FHeaderParser::GetVarType
 	const bool bIsParamList = (VariableCategory != EVariableCategory::Member) && MatchIdentifier(TEXT("UPARAM"));
 
 	// No specifiers are allowed inside a TArray
-	if( (OuterPropertyType == NULL) || !OuterPropertyType->Matches(TEXT("TArray")) )
+	if ((OuterPropertyType == NULL) || !OuterPropertyType->Matches(TEXT("TArray")))
 	{
 		// New-style UPROPERTY() syntax 
 		if (PropertyDeclarationStyle == EPropertyDeclarationStyle::UPROPERTY || bIsParamList)
@@ -2656,7 +2656,7 @@ void FHeaderParser::GetVarType
 		}
 	}
 
-	if (CompilerDirectiveStack.Num() > 0 && (CompilerDirectiveStack.Last()&ECompilerDirective::WithEditorOnlyData)!=0)
+	if (CompilerDirectiveStack.Num() > 0 && (CompilerDirectiveStack.Last()&ECompilerDirective::WithEditorOnlyData) != 0)
 	{
 		Flags |= CPF_EditorOnly;
 	}
@@ -2677,284 +2677,284 @@ void FHeaderParser::GetVarType
 		{
 			switch (SpecID)
 			{
-				case EVariableSpecifier::EditAnywhere:
+			case EVariableSpecifier::EditAnywhere:
+			{
+				if (bSeenEditSpecifier)
 				{
-					if (bSeenEditSpecifier)
-					{
-						FError::Throwf(TEXT("Found more than one edit/visibility specifier (%s), only one is allowed"), *Specifier.Key);
-					}
-					Flags |= CPF_Edit;
-					bSeenEditSpecifier = true;
+					FError::Throwf(TEXT("Found more than one edit/visibility specifier (%s), only one is allowed"), *Specifier.Key);
 				}
+				Flags |= CPF_Edit;
+				bSeenEditSpecifier = true;
+			}
 				break;
 
-				case EVariableSpecifier::EditInstanceOnly:
+			case EVariableSpecifier::EditInstanceOnly:
+			{
+				if (bSeenEditSpecifier)
 				{
-					if (bSeenEditSpecifier)
-					{
-						FError::Throwf(TEXT("Found more than one edit/visibility specifier (%s), only one is allowed"), *Specifier.Key);
-					}
-					Flags |= CPF_Edit|CPF_DisableEditOnTemplate;
-					bSeenEditSpecifier = true;
+					FError::Throwf(TEXT("Found more than one edit/visibility specifier (%s), only one is allowed"), *Specifier.Key);
 				}
+				Flags |= CPF_Edit | CPF_DisableEditOnTemplate;
+				bSeenEditSpecifier = true;
+			}
 				break;
 
-				case EVariableSpecifier::EditDefaultsOnly: 
+			case EVariableSpecifier::EditDefaultsOnly:
+			{
+				if (bSeenEditSpecifier)
 				{
-					if (bSeenEditSpecifier)
-					{
-						FError::Throwf(TEXT("Found more than one edit/visibility specifier (%s), only one is allowed"), *Specifier.Key);
-					}
-					Flags |= CPF_Edit|CPF_DisableEditOnInstance;
-					bSeenEditSpecifier = true;
+					FError::Throwf(TEXT("Found more than one edit/visibility specifier (%s), only one is allowed"), *Specifier.Key);
 				}
+				Flags |= CPF_Edit | CPF_DisableEditOnInstance;
+				bSeenEditSpecifier = true;
+			}
 				break;
 
-				case EVariableSpecifier::VisibleAnywhere: 
+			case EVariableSpecifier::VisibleAnywhere:
+			{
+				if (bSeenEditSpecifier)
 				{
-					if (bSeenEditSpecifier)
-					{
-						FError::Throwf(TEXT("Found more than one edit/visibility specifier (%s), only one is allowed"), *Specifier.Key);
-					}
-					Flags |= CPF_Edit|CPF_EditConst;
-					bSeenEditSpecifier = true;
+					FError::Throwf(TEXT("Found more than one edit/visibility specifier (%s), only one is allowed"), *Specifier.Key);
 				}
+				Flags |= CPF_Edit | CPF_EditConst;
+				bSeenEditSpecifier = true;
+			}
 				break;
 
-				case EVariableSpecifier::VisibleInstanceOnly:
+			case EVariableSpecifier::VisibleInstanceOnly:
+			{
+				if (bSeenEditSpecifier)
 				{
-					if (bSeenEditSpecifier)
-					{
-						FError::Throwf(TEXT("Found more than one edit/visibility specifier (%s), only one is allowed"), *Specifier.Key);
-					}
-					Flags |= CPF_Edit|CPF_EditConst|CPF_DisableEditOnTemplate;
-					bSeenEditSpecifier = true;
+					FError::Throwf(TEXT("Found more than one edit/visibility specifier (%s), only one is allowed"), *Specifier.Key);
 				}
+				Flags |= CPF_Edit | CPF_EditConst | CPF_DisableEditOnTemplate;
+				bSeenEditSpecifier = true;
+			}
 				break;
 
-				case EVariableSpecifier::VisibleDefaultsOnly:
+			case EVariableSpecifier::VisibleDefaultsOnly:
+			{
+				if (bSeenEditSpecifier)
 				{
-					if (bSeenEditSpecifier)
-					{
-						FError::Throwf(TEXT("Found more than one edit/visibility specifier (%s), only one is allowed"), *Specifier.Key);
-					}
-					Flags |= CPF_Edit|CPF_EditConst|CPF_DisableEditOnInstance;
-					bSeenEditSpecifier = true;
+					FError::Throwf(TEXT("Found more than one edit/visibility specifier (%s), only one is allowed"), *Specifier.Key);
 				}
+				Flags |= CPF_Edit | CPF_EditConst | CPF_DisableEditOnInstance;
+				bSeenEditSpecifier = true;
+			}
 				break;
 
-				case EVariableSpecifier::BlueprintReadWrite:
+			case EVariableSpecifier::BlueprintReadWrite:
+			{
+				if (bSeenBlueprintEditSpecifier)
 				{
-					if (bSeenBlueprintEditSpecifier)
-					{
-						FError::Throwf(TEXT("Found more than one Blueprint read/write specifier (%s), only one is allowed"), *Specifier.Key);
-					}
-
-					const FString* PrivateAccessMD = MetaDataFromNewStyle.Find(TEXT("AllowPrivateAccess"));  // FBlueprintMetadata::MD_AllowPrivateAccess
-					const bool bAllowPrivateAccess = PrivateAccessMD ? (*PrivateAccessMD == TEXT("true")) : false;
-					if (CurrentAccessSpecifier == ACCESS_Private && !bAllowPrivateAccess)
-					{
-						FError::Throwf(TEXT("BlueprintReadWrite should not be used on private members"));
-					}
-					Flags |= CPF_BlueprintVisible;
-					bSeenBlueprintEditSpecifier = true;
+					FError::Throwf(TEXT("Found more than one Blueprint read/write specifier (%s), only one is allowed"), *Specifier.Key);
 				}
-				break;
 
-				case EVariableSpecifier::BlueprintReadOnly:
+				const FString* PrivateAccessMD = MetaDataFromNewStyle.Find(TEXT("AllowPrivateAccess"));  // FBlueprintMetadata::MD_AllowPrivateAccess
+				const bool bAllowPrivateAccess = PrivateAccessMD ? (*PrivateAccessMD == TEXT("true")) : false;
+				if (CurrentAccessSpecifier == ACCESS_Private && !bAllowPrivateAccess)
 				{
-					if (bSeenBlueprintEditSpecifier)
-					{
-						FError::Throwf(TEXT("Found more than one Blueprint read/write specifier (%s), only one is allowed"), *Specifier.Key);
-					}
-
-					const FString* PrivateAccessMD = MetaDataFromNewStyle.Find(TEXT("AllowPrivateAccess"));  // FBlueprintMetadata::MD_AllowPrivateAccess
-					const bool bAllowPrivateAccess = PrivateAccessMD ? (*PrivateAccessMD == TEXT("true")) : false;
-					if (CurrentAccessSpecifier == ACCESS_Private && !bAllowPrivateAccess)
-					{
-						FError::Throwf(TEXT("BlueprintReadOnly should not be used on private members"));
-					}
-					Flags |= CPF_BlueprintVisible|CPF_BlueprintReadOnly;
-					bSeenBlueprintEditSpecifier = true;
+					FError::Throwf(TEXT("BlueprintReadWrite should not be used on private members"));
 				}
+				Flags |= CPF_BlueprintVisible;
+				bSeenBlueprintEditSpecifier = true;
+			}
 				break;
 
-				case EVariableSpecifier::Config:
+			case EVariableSpecifier::BlueprintReadOnly:
+			{
+				if (bSeenBlueprintEditSpecifier)
 				{
-					Flags |= CPF_Config;
+					FError::Throwf(TEXT("Found more than one Blueprint read/write specifier (%s), only one is allowed"), *Specifier.Key);
 				}
-				break;
 
-				case EVariableSpecifier::GlobalConfig:
+				const FString* PrivateAccessMD = MetaDataFromNewStyle.Find(TEXT("AllowPrivateAccess"));  // FBlueprintMetadata::MD_AllowPrivateAccess
+				const bool bAllowPrivateAccess = PrivateAccessMD ? (*PrivateAccessMD == TEXT("true")) : false;
+				if (CurrentAccessSpecifier == ACCESS_Private && !bAllowPrivateAccess)
 				{
-					Flags |= CPF_GlobalConfig | CPF_Config;
+					FError::Throwf(TEXT("BlueprintReadOnly should not be used on private members"));
 				}
+				Flags |= CPF_BlueprintVisible | CPF_BlueprintReadOnly;
+				bSeenBlueprintEditSpecifier = true;
+			}
 				break;
 
-				case EVariableSpecifier::Localized:
-				{
-					FError::Throwf(TEXT("The Localized specifier is deprecated"));
-				}
+			case EVariableSpecifier::Config:
+			{
+				Flags |= CPF_Config;
+			}
 				break;
 
-				case EVariableSpecifier::Transient:
-				{
-					Flags |= CPF_Transient;
-				}
+			case EVariableSpecifier::GlobalConfig:
+			{
+				Flags |= CPF_GlobalConfig | CPF_Config;
+			}
 				break;
 
-				case EVariableSpecifier::DuplicateTransient:
-				{
-					Flags |= CPF_DuplicateTransient;
-				}
+			case EVariableSpecifier::Localized:
+			{
+				FError::Throwf(TEXT("The Localized specifier is deprecated"));
+			}
 				break;
 
-				case EVariableSpecifier::TextExportTransient:
-				{
-					Flags |= CPF_TextExportTransient;
-				}
+			case EVariableSpecifier::Transient:
+			{
+				Flags |= CPF_Transient;
+			}
 				break;
 
-				case EVariableSpecifier::NonPIETransient:
-				{
-					UE_LOG(LogCompile, Warning, TEXT("NonPIETransient is deprecated - NonPIEDuplicateTransient should be used instead"));
-					Flags |= CPF_NonPIEDuplicateTransient;
-				}
+			case EVariableSpecifier::DuplicateTransient:
+			{
+				Flags |= CPF_DuplicateTransient;
+			}
 				break;
 
-				case EVariableSpecifier::NonPIEDuplicateTransient:
-				{
-					Flags |= CPF_NonPIEDuplicateTransient;
-				}
+			case EVariableSpecifier::TextExportTransient:
+			{
+				Flags |= CPF_TextExportTransient;
+			}
 				break;
 
-				case EVariableSpecifier::Export:
-				{
-					Flags |= CPF_ExportObject;
-				}
+			case EVariableSpecifier::NonPIETransient:
+			{
+				UE_LOG(LogCompile, Warning, TEXT("NonPIETransient is deprecated - NonPIEDuplicateTransient should be used instead"));
+				Flags |= CPF_NonPIEDuplicateTransient;
+			}
 				break;
 
-				case EVariableSpecifier::EditInline:
-				{
-					FError::Throwf(TEXT("EditInline is deprecated. Remove it, or use Instanced instead."));
-				}
+			case EVariableSpecifier::NonPIEDuplicateTransient:
+			{
+				Flags |= CPF_NonPIEDuplicateTransient;
+			}
 				break;
 
-				case EVariableSpecifier::NoClear:
-				{
-					Flags |= CPF_NoClear;
-				}
+			case EVariableSpecifier::Export:
+			{
+				Flags |= CPF_ExportObject;
+			}
 				break;
 
-				case EVariableSpecifier::EditFixedSize:
-				{
-					Flags |= CPF_EditFixedSize;
-				}
+			case EVariableSpecifier::EditInline:
+			{
+				FError::Throwf(TEXT("EditInline is deprecated. Remove it, or use Instanced instead."));
+			}
 				break;
 
-				case EVariableSpecifier::Replicated:
-				case EVariableSpecifier::ReplicatedUsing:
-				{
-					if (OwnerStruct->IsA<UScriptStruct>())
-					{
-						FError::Throwf(TEXT("Struct members cannot be replicated"));
-					}
-
-					Flags |= CPF_Net;
+			case EVariableSpecifier::NoClear:
+			{
+				Flags |= CPF_NoClear;
+			}
+				break;
 
-					// See if we've specified a rep notification function
-					if (SpecID == EVariableSpecifier::ReplicatedUsing)
-					{
-						RepCallbackName = FName(*RequireExactlyOneSpecifierValue(Specifier));
-						Flags |= CPF_RepNotify;
-					}
-				}
+			case EVariableSpecifier::EditFixedSize:
+			{
+				Flags |= CPF_EditFixedSize;
+			}
 				break;
 
-				case EVariableSpecifier::NotReplicated:
+			case EVariableSpecifier::Replicated:
+			case EVariableSpecifier::ReplicatedUsing:
+			{
+				if (OwnerStruct->IsA<UScriptStruct>())
 				{
-					if (!OwnerStruct->IsA<UScriptStruct>())
-					{
-						FError::Throwf(TEXT("Only Struct members can be marked NotReplicated"));
-					}
-
-					Flags |= CPF_RepSkip;
+					FError::Throwf(TEXT("Struct members cannot be replicated"));
 				}
-				break;
 
-				case EVariableSpecifier::RepRetry:
+				Flags |= CPF_Net;
+
+				// See if we've specified a rep notification function
+				if (SpecID == EVariableSpecifier::ReplicatedUsing)
 				{
-					FError::Throwf(TEXT("'RepRetry' is deprecated."));
+					RepCallbackName = FName(*RequireExactlyOneSpecifierValue(Specifier));
+					Flags |= CPF_RepNotify;
 				}
+			}
 				break;
 
-				case EVariableSpecifier::Interp:
+			case EVariableSpecifier::NotReplicated:
+			{
+				if (!OwnerStruct->IsA<UScriptStruct>())
 				{
-					Flags |= CPF_Edit;
-					Flags |= CPF_BlueprintVisible;
-					Flags |= CPF_Interp;
+					FError::Throwf(TEXT("Only Struct members can be marked NotReplicated"));
 				}
+
+				Flags |= CPF_RepSkip;
+			}
 				break;
 
-				case EVariableSpecifier::NonTransactional:
-				{
-					Flags |= CPF_NonTransactional;
-				}
+			case EVariableSpecifier::RepRetry:
+			{
+				FError::Throwf(TEXT("'RepRetry' is deprecated."));
+			}
 				break;
 
-				case EVariableSpecifier::Instanced:
-				{
-					Flags |= CPF_PersistentInstance | CPF_ExportObject | CPF_InstancedReference;
-					AddEditInlineMetaData(MetaDataFromNewStyle);
-				}
+			case EVariableSpecifier::Interp:
+			{
+				Flags |= CPF_Edit;
+				Flags |= CPF_BlueprintVisible;
+				Flags |= CPF_Interp;
+			}
 				break;
 
-				case EVariableSpecifier::BlueprintAssignable:
-				{
-					Flags |= CPF_BlueprintAssignable;
-				}
+			case EVariableSpecifier::NonTransactional:
+			{
+				Flags |= CPF_NonTransactional;
+			}
 				break;
 
-				case EVariableSpecifier::BlueprintCallable:
-				{
-					Flags |= CPF_BlueprintCallable;
-				}
+			case EVariableSpecifier::Instanced:
+			{
+				Flags |= CPF_PersistentInstance | CPF_ExportObject | CPF_InstancedReference;
+				AddEditInlineMetaData(MetaDataFromNewStyle);
+			}
 				break;
 
-				case EVariableSpecifier::BlueprintAuthorityOnly:
-				{
-					Flags |= CPF_BlueprintAuthorityOnly;
-				}
+			case EVariableSpecifier::BlueprintAssignable:
+			{
+				Flags |= CPF_BlueprintAssignable;
+			}
 				break;
 
-				case EVariableSpecifier::AssetRegistrySearchable:
-				{
-					Flags |= CPF_AssetRegistrySearchable;
-				}
+			case EVariableSpecifier::BlueprintCallable:
+			{
+				Flags |= CPF_BlueprintCallable;
+			}
 				break;
 
-				case EVariableSpecifier::SimpleDisplay:
-				{
-					Flags |= CPF_SimpleDisplay;
-				}
+			case EVariableSpecifier::BlueprintAuthorityOnly:
+			{
+				Flags |= CPF_BlueprintAuthorityOnly;
+			}
 				break;
 
-				case EVariableSpecifier::AdvancedDisplay:
-				{
-					Flags |= CPF_AdvancedDisplay;
-				}
+			case EVariableSpecifier::AssetRegistrySearchable:
+			{
+				Flags |= CPF_AssetRegistrySearchable;
+			}
 				break;
 
-				case EVariableSpecifier::SaveGame:
-				{
-					Flags |= CPF_SaveGame;
-				}
+			case EVariableSpecifier::SimpleDisplay:
+			{
+				Flags |= CPF_SimpleDisplay;
+			}
 				break;
 
-				default:
-				{
-					FError::Throwf(TEXT("Unknown variable specifier '%s'"), *Specifier.Key);
-				}
+			case EVariableSpecifier::AdvancedDisplay:
+			{
+				Flags |= CPF_AdvancedDisplay;
+			}
+				break;
+
+			case EVariableSpecifier::SaveGame:
+			{
+				Flags |= CPF_SaveGame;
+			}
+				break;
+
+			default:
+			{
+				FError::Throwf(TEXT("Unknown variable specifier '%s'"), *Specifier.Key);
+			}
 				break;
 			}
 		}
@@ -2962,36 +2962,36 @@ void FHeaderParser::GetVarType
 		{
 			switch (SpecID)
 			{
-				case EVariableSpecifier::Const:
-				{
-					Flags |= CPF_ConstParm;
-				}
+			case EVariableSpecifier::Const:
+			{
+				Flags |= CPF_ConstParm;
+			}
 				break;
 
-				case EVariableSpecifier::Ref:
-				{
-					Flags |= CPF_OutParm | CPF_ReferenceParm;
-				}
+			case EVariableSpecifier::Ref:
+			{
+				Flags |= CPF_OutParm | CPF_ReferenceParm;
+			}
 				break;
 
-				case EVariableSpecifier::NotReplicated:
+			case EVariableSpecifier::NotReplicated:
+			{
+				if (VariableCategory == EVariableCategory::ReplicatedParameter)
 				{
-					if (VariableCategory == EVariableCategory::ReplicatedParameter)
-					{
-						VariableCategory = EVariableCategory::RegularParameter;
-						Flags |= CPF_RepSkip;
-					}
-					else
-					{
-						FError::Throwf(TEXT("Only parameters in service request functions can be marked NotReplicated"));
-					}
+					VariableCategory = EVariableCategory::RegularParameter;
+					Flags |= CPF_RepSkip;
 				}
-				break;
-
-				default:
+				else
 				{
-					FError::Throwf(TEXT("Unknown variable specifier '%s'"), *Specifier.Key);
+					FError::Throwf(TEXT("Only parameters in service request functions can be marked NotReplicated"));
 				}
+			}
+				break;
+
+			default:
+			{
+				FError::Throwf(TEXT("Unknown variable specifier '%s'"), *Specifier.Key);
+			}
 				break;
 			}
 		}
@@ -3019,18 +3019,27 @@ void FHeaderParser::GetVarType
 		Flags       &= ~CPF_Protected;
 		ExportFlags |= PROPEXPORT_Public;
 		ExportFlags &= ~(PROPEXPORT_Private|PROPEXPORT_Protected);
+
+		Flags &= ~CPF_NativeAccessSpecifiers;
+		Flags |= CPF_NativeAccessSpecifierPublic;
 	}
 	else if (CurrentAccessSpecifier == ACCESS_Protected)
 	{
 		Flags       |= CPF_Protected;
 		ExportFlags |= PROPEXPORT_Protected;
 		ExportFlags &= ~(PROPEXPORT_Public|PROPEXPORT_Private);
+
+		Flags &= ~CPF_NativeAccessSpecifiers;
+		Flags |= CPF_NativeAccessSpecifierProtected;
 	}
 	else if (CurrentAccessSpecifier == ACCESS_Private)
 	{
 		Flags       &= ~CPF_Protected;
 		ExportFlags |= PROPEXPORT_Private;
 		ExportFlags &= ~(PROPEXPORT_Public|PROPEXPORT_Protected);
+
+		Flags &= ~CPF_NativeAccessSpecifiers;
+		Flags |= CPF_NativeAccessSpecifierPrivate;
 	}
 	else
 	{
@@ -5050,7 +5059,7 @@ void FHeaderParser::ParseParameterList(FClasses& AllClasses, UFunction* Function
 		// Get parameter type.
 		FToken Property(CPT_None);
 		EVariableCategory::Type VariableCategory = (Function->FunctionFlags & FUNC_Net) ? EVariableCategory::ReplicatedParameter : EVariableCategory::RegularParameter;
-		GetVarType(AllClasses, GetCurrentScope(), Property, ~(CPF_ParmFlags | CPF_AutoWeak | CPF_RepSkip | CPF_UObjectWrapper), NULL, EPropertyDeclarationStyle::None, VariableCategory);
+		GetVarType(AllClasses, GetCurrentScope(), Property, ~(CPF_ParmFlags | CPF_AutoWeak | CPF_RepSkip | CPF_UObjectWrapper | CPF_NativeAccessSpecifiers), NULL, EPropertyDeclarationStyle::None, VariableCategory);
 		Property.PropertyFlags |= CPF_Parm;
 
 		if (bExpectCommaBeforeName)
diff --git a/Engine/Source/Runtime/CoreUObject/Private/UObject/PropertyStruct.cpp b/Engine/Source/Runtime/CoreUObject/Private/UObject/PropertyStruct.cpp
--- a/Engine/Source/Runtime/CoreUObject/Private/UObject/PropertyStruct.cpp
+++ b/Engine/Source/Runtime/CoreUObject/Private/UObject/PropertyStruct.cpp
@@ -330,7 +330,6 @@ void UStructProperty::ExportTextItem( FString& ValueStr, const void* PropertyVal
 
 	if (0 != (PortFlags & PPF_ExportCpp))
 	{
-		ValueStr += FString::Printf(TEXT("F%s()"), *Struct->GetName());
 		return;
 	}
 
diff --git a/Engine/Source/Runtime/CoreUObject/Public/UObject/ObjectBase.h b/Engine/Source/Runtime/CoreUObject/Public/UObject/ObjectBase.h
--- a/Engine/Source/Runtime/CoreUObject/Public/UObject/ObjectBase.h
+++ b/Engine/Source/Runtime/CoreUObject/Public/UObject/ObjectBase.h
@@ -396,6 +396,9 @@ typedef uint64 EClassCastFlags;
 #define CPF_PersistentInstance				DECLARE_UINT64(0x0002000000000000)		// A object referenced by the property is duplicated like a component. (Each actor should have an own instance.)
 #define CPF_UObjectWrapper					DECLARE_UINT64(0x0004000000000000)		// Property was parsed as a wrapper class like TSubobjectOf<T>, FScriptInterface etc., rather than a USomething*
 #define CPF_HasGetValueTypeHash				DECLARE_UINT64(0x0008000000000000)		// This property can generate a meaningful hash value.
+#define CPF_NativeAccessSpecifierPublic		DECLARE_UINT64(0x0010000000000000)		// Public native access specifier
+#define CPF_NativeAccessSpecifierProtected	DECLARE_UINT64(0x0020000000000000)		// Protected native access specifier
+#define CPF_NativeAccessSpecifierPrivate	DECLARE_UINT64(0x0040000000000000)		// Private native access specifier
 
 #define CPF_NonPIETransient \
 	EMIT_DEPRECATED_WARNING_MESSAGE("CPF_NonPIETransient is deprecated. Please use CPF_NonPIEDuplicateTransient instead.") \
@@ -403,6 +406,8 @@ typedef uint64 EClassCastFlags;
 
 /** @name Combinations flags */
 //@{
+#define CPF_NativeAccessSpecifiers	(CPF_NativeAccessSpecifierPublic | CPF_NativeAccessSpecifierProtected | CPF_NativeAccessSpecifierPrivate)
+
 #define CPF_ParmFlags				(CPF_Parm | CPF_OutParm | CPF_ReturnParm | CPF_ReferenceParm | CPF_ConstParm)
 #define CPF_PropagateToArrayInner	(CPF_ExportObject | CPF_PersistentInstance | CPF_InstancedReference | CPF_ContainsInstancedReference | CPF_Config | CPF_EditConst | CPF_Deprecated | CPF_EditorOnly | CPF_AutoWeak | CPF_UObjectWrapper )
 #define CPF_PropagateToMapValue		(CPF_ExportObject | CPF_PersistentInstance | CPF_InstancedReference | CPF_ContainsInstancedReference | CPF_Config | CPF_EditConst | CPF_Deprecated | CPF_EditorOnly | CPF_AutoWeak | CPF_UObjectWrapper | CPF_Edit )
diff --git a/Engine/Source/Runtime/InputCore/Private/InputCoreTypes.cpp b/Engine/Source/Runtime/InputCore/Private/InputCoreTypes.cpp
--- a/Engine/Source/Runtime/InputCore/Private/InputCoreTypes.cpp
+++ b/Engine/Source/Runtime/InputCore/Private/InputCoreTypes.cpp
@@ -1020,7 +1020,8 @@ bool FKey::ExportTextItem(FString& ValueStr, FKey const& DefaultValue, UObject*
 {
 	if (0 != (PortFlags & EPropertyPortFlags::PPF_ExportCpp))
 	{
-		return false;
+		ValueStr += FString::Printf(TEXT("FKey(TEXT(\"%s\"))"), *KeyName.ToString());
+		return true;
 	}
 
 	ValueStr += KeyName.ToString();
