CodeGenerator sets CppType in Enum
UByteProperty::GetCPPType returns proper type

diff --git a/Engine/Source/Programs/UnrealHeaderTool/Private/CodeGenerator.cpp b/Engine/Source/Programs/UnrealHeaderTool/Private/CodeGenerator.cpp
--- a/Engine/Source/Programs/UnrealHeaderTool/Private/CodeGenerator.cpp
+++ b/Engine/Source/Programs/UnrealHeaderTool/Private/CodeGenerator.cpp
@@ -2960,6 +2960,7 @@ void FNativeClassHeaderGenerator::ExportGeneratedEnumsInitCode(const TArray<UEnu
 			case UEnum::ECppForm::EnumClass:  EnumTypeStr = TEXT("UEnum::ECppForm::EnumClass");  break;
 			}
 			GeneratedEnumRegisterFunctionText.Logf(TEXT("            ReturnEnum->SetEnums(EnumNames, %s);\r\n"), *EnumTypeStr);
+			GeneratedEnumRegisterFunctionText.Logf(TEXT("            ReturnEnum->CppType = TEXT(\"%s\");\r\n"), *Enum->CppType);
 
 			FString Meta = GetMetaDataCodeForObject(Enum, TEXT("ReturnEnum"), TEXT("            "));
 			if (Meta.Len())
