Support for regular C++ int types in casts.

diff --git a/Engine/Source/Programs/UnrealHeaderTool/Private/HeaderParser.cpp b/Engine/Source/Programs/UnrealHeaderTool/Private/HeaderParser.cpp
--- a/Engine/Source/Programs/UnrealHeaderTool/Private/HeaderParser.cpp
+++ b/Engine/Source/Programs/UnrealHeaderTool/Private/HeaderParser.cpp
@@ -3990,7 +3990,12 @@ UProperty* FHeaderParser::GetVarNameAndDim
 					TEXT("(uint16)"),
 					TEXT("(int16)"),
 					TEXT("(uint8)"),
-					TEXT("(int8)")
+					TEXT("(int8)"),
+					TEXT("(int)"),
+					TEXT("(unsigned)"),
+					TEXT("(signed)"),
+					TEXT("(unsigned int)"),
+					TEXT("(signed int)")
 				};
 
 				// Remove any brackets
