PR #1460: Fixed LogCompile typo: Tabify time was was... (Contributed by [REDACTED_PERSON])

diff --git a/Engine/Source/Programs/UnrealHeaderTool/Private/CodeGenerator.cpp b/Engine/Source/Programs/UnrealHeaderTool/Private/CodeGenerator.cpp
--- a/Engine/Source/Programs/UnrealHeaderTool/Private/CodeGenerator.cpp
+++ b/Engine/Source/Programs/UnrealHeaderTool/Private/CodeGenerator.cpp
@@ -5659,7 +5659,7 @@ ECompilationResult::Type UnrealHeaderTool_Main(const FString& ModuleInfoFilename
 	UE_LOG(LogCompile, Log, TEXT("Code generation took %.2f seconds"), GHeaderCodeGenTime);
 	UE_LOG(LogCompile, Log, TEXT("ScriptPlugin overhead was %.2f seconds"), GPluginOverheadTime);
 	UE_LOG(LogCompile, Log, TEXT("Macroize time was %.2f seconds"), GMacroizeTime);
-	UE_LOG(LogCompile, Log, TEXT("Tabify time was was %.2f seconds"), GTabifyTime);
+	UE_LOG(LogCompile, Log, TEXT("Tabify time was %.2f seconds"), GTabifyTime);
 
 	if (bWriteContents)
 	{
