Fix multiline comments parsing in UHT, where end of comment was treated as start of one, e.g. /*/ Comment /*/

diff --git a/Engine/Source/Programs/UnrealHeaderTool/Private/BaseParser.cpp b/Engine/Source/Programs/UnrealHeaderTool/Private/BaseParser.cpp
--- a/Engine/Source/Programs/UnrealHeaderTool/Private/BaseParser.cpp
+++ b/Engine/Source/Programs/UnrealHeaderTool/Private/BaseParser.cpp
@@ -213,9 +213,15 @@ Loop:
 				PrevComment += c;
 				PrevComment += NextChar;
 				bInsideComment = true;
+
+				// Move past the star. Do it only when not in comment,
+				// otherwise end of comment might be missed e.g.
+				// /*/ Comment /*/
+				// ~~~~~~~~~~~~~^ Will report second /* as beginning of comment
+				// And throw error that end of file is found in comment.
+				InputPos++;
 			}
 
-			InputPos++;
 			goto Loop;
 		}
 		else if( c==TEXT('*') && NextChar==TEXT('/') )
