Fix multiline comment issues when parsing */// in UHT.

diff --git a/Engine/Source/Programs/UnrealHeaderTool/Private/HeaderParser.cpp b/Engine/Source/Programs/UnrealHeaderTool/Private/HeaderParser.cpp
--- a/Engine/Source/Programs/UnrealHeaderTool/Private/HeaderParser.cpp
+++ b/Engine/Source/Programs/UnrealHeaderTool/Private/HeaderParser.cpp
@@ -7280,6 +7280,12 @@ void FHeaderParser::SimplifiedClassParse(const TCHAR* InBuffer, TArray<FSimplifi
 
 			// Stub out the comments, ignoring anything inside literal strings.
 			Pos = StrLine.Find(TEXT("//"));
+
+			// Check if first slash is end of multiline comment and adjust position if necessary.
+			if (Pos > 0 && StrLine[Pos - 1] == TEXT('*'))
+			{
+				++Pos;
+			}
 			if (Pos >= 0)
 			{
 				if (StrBegin == INDEX_NONE || Pos < StrBegin || Pos > StrEnd)
