Fix generated _getUObject() function not having "override" when it should

diff --git a/Engine/Source/Programs/UnrealHeaderTool/Private/CodeGenerator.cpp b/Engine/Source/Programs/UnrealHeaderTool/Private/CodeGenerator.cpp
--- a/Engine/Source/Programs/UnrealHeaderTool/Private/CodeGenerator.cpp
+++ b/Engine/Source/Programs/UnrealHeaderTool/Private/CodeGenerator.cpp
@@ -2119,12 +2119,8 @@ void FNativeClassHeaderGenerator::ExportClassesFromSourceFileInner(FUnrealSource
 				ClassBoilerplate.Logf(TEXT("    static const TCHAR* StaticConfigName() {return TEXT(\"%s\");}\r\n\r\n"), *Class->ClassConfigName.ToString());
 			}
 
-			// @todo srobb: clean up _getUObject() code generation
-			if (((SuperClass == UClass::StaticClass()) || (SuperClass == UObject::StaticClass()) || (SuperClass && (SuperClass->ClassFlags & CLASS_Intrinsic))) && !Class->IsChildOf(UInterface::StaticClass()))
-			{
-				ClassBoilerplate.Logf(TEXT("\tvirtual UObject* _getUObject() const { return const_cast<%s*>(this); }\r\n"), ClassCPPName);
-			}
-			else
+			// export implementation of _getUObject for classes that implement interfaces
+			if (Class->Interfaces.Num() > 0)
 			{
 				ClassBoilerplate.Logf(TEXT("\tvirtual UObject* _getUObject() const override { return const_cast<%s*>(this); }\r\n"), ClassCPPName);
 			}
