Fix for FToken stack corruption when NAME_SIZE is changed.

diff --git a/Engine/Source/Programs/UnrealHeaderTool/Private/ParserHelper.h b/Engine/Source/Programs/UnrealHeaderTool/Private/ParserHelper.h
--- a/Engine/Source/Programs/UnrealHeaderTool/Private/ParserHelper.h
+++ b/Engine/Source/Programs/UnrealHeaderTool/Private/ParserHelper.h
@@ -732,7 +732,7 @@ public:
 		StartPos		= 0;
 		StartLine		= 0;
 		*Identifier		= 0;
-		FMemory::Memzero(String, sizeof(Identifier));
+		FMemory::Memzero(String);
 	}
 	bool Matches( const TCHAR* Str, ESearchCase::Type SearchCase = ESearchCase::IgnoreCase ) const
 	{
