Making FNativeClassHeaderGenerator::GetGeneratedFunctionTextDevice's MaxLinesPerCpp configurable via ini setting.

diff --git a/Engine/Programs/UnrealHeaderTool/Config/DefaultEngine.ini b/Engine/Programs/UnrealHeaderTool/Config/DefaultEngine.ini
--- a/Engine/Programs/UnrealHeaderTool/Config/DefaultEngine.ini
+++ b/Engine/Programs/UnrealHeaderTool/Config/DefaultEngine.ini
@@ -3,4 +3,7 @@
 
 +ScriptSupportedModules=CoreUObject
 +ScriptSupportedModules=Engine
-+ScriptSupportedModules=ScriptPlugin
\ No newline at end of file
++ScriptSupportedModules=ScriptPlugin
+
+[UnrealHeaderTool]
+MaxLinesPerCpp=30000
\ No newline at end of file
diff --git a/Engine/Source/Programs/UnrealHeaderTool/Private/CodeGenerator.cpp b/Engine/Source/Programs/UnrealHeaderTool/Private/CodeGenerator.cpp
--- a/Engine/Source/Programs/UnrealHeaderTool/Private/CodeGenerator.cpp
+++ b/Engine/Source/Programs/UnrealHeaderTool/Private/CodeGenerator.cpp
@@ -4672,13 +4672,23 @@ bool IsExportClass(FClass* Class)
 
 FUHTStringBuilder& FNativeClassHeaderGenerator::GetGeneratedFunctionTextDevice()
 {
+	static struct FMaxLinesPerCpp
+	{
+		int32 Value;
+		FMaxLinesPerCpp()
+		{
 #if ( PLATFORM_WINDOWS && defined(__clang__) )	// @todo clang: Clang r231657 often crashes with huge Engine.generated.cpp files, so we split using a smaller threshold
-	const int32 MaxLinesPerCpp = 15000;
+			Value = 15000;
 #else
-	const int32 MaxLinesPerCpp = 30000;
+			// We do this only for non-clang builds for now
+			Value = 30000;
+			check(GConfig);
+			GConfig->GetInt(TEXT("UnrealHeaderTool"), TEXT("MaxLinesPerCpp"), Value, GEngineIni);
 #endif
+		}
+	} MaxLinesPerCpp;
 
-	if ((GeneratedFunctionBodyTextSplit.Num() == 0) || (GeneratedFunctionBodyTextSplit[GeneratedFunctionBodyTextSplit.Num() - 1]->GetLineCount() > MaxLinesPerCpp))
+	if ((GeneratedFunctionBodyTextSplit.Num() == 0) || (GeneratedFunctionBodyTextSplit[GeneratedFunctionBodyTextSplit.Num() - 1]->GetLineCount() > MaxLinesPerCpp.Value))
 	{
 		GeneratedFunctionBodyTextSplit.Add( TUniqueObj<FUHTStringBuilderLineCounter>() );
 	}
