FPropertyBase::ImpliedPropertyFlags added to distinguish between property flags which are directly specified by UPROPERTY() and those which implicitly come from other places, like class flags or type.
Added a fix for const UCLASSes erroring out on deprecated properties, which no longer considers implied flags.

diff --git a/Engine/Source/Programs/UnrealHeaderTool/Private/HeaderParser.cpp b/Engine/Source/Programs/UnrealHeaderTool/Private/HeaderParser.cpp
--- a/Engine/Source/Programs/UnrealHeaderTool/Private/HeaderParser.cpp
+++ b/Engine/Source/Programs/UnrealHeaderTool/Private/HeaderParser.cpp
@@ -2622,11 +2622,19 @@ FIndexRange*                    ParsedVarIndexRange
 	FName RepCallbackName = FName(NAME_None);
 
 	// Get flags.
-	uint64 Flags = 0;
+	uint64 Flags        = 0;
+	uint64 ImpliedFlags = 0;
+
 	// force members to be 'blueprint read only' if in a const class
-	if (VariableCategory == EVariableCategory::Member && (Cast<UClass>(OwnerStruct) != nullptr) && (((UClass*)OwnerStruct)->ClassFlags & CLASS_Const))
+	if (VariableCategory == EVariableCategory::Member)
 	{
-		Flags |= CPF_BlueprintReadOnly;
+		if (UClass* OwnerClass = Cast<UClass>(OwnerStruct))
+		{
+			if (OwnerClass->ClassFlags & CLASS_Const)
+			{
+				ImpliedFlags |= CPF_BlueprintReadOnly;
+			}
+		}
 	}
 	uint32 ExportFlags = PROPEXPORT_Public;
 
@@ -2677,284 +2685,285 @@ FIndexRange*                    ParsedVarIndexRange
 		{
 			switch (SpecID)
 			{
-			case EVariableSpecifier::EditAnywhere:
-			{
-				if (bSeenEditSpecifier)
+				case EVariableSpecifier::EditAnywhere:
 				{
-					FError::Throwf(TEXT("Found more than one edit/visibility specifier (%s), only one is allowed"), *Specifier.Key);
+					if (bSeenEditSpecifier)
+					{
+						FError::Throwf(TEXT("Found more than one edit/visibility specifier (%s), only one is allowed"), *Specifier.Key);
+					}
+					Flags |= CPF_Edit;
+					bSeenEditSpecifier = true;
 				}
-				Flags |= CPF_Edit;
-				bSeenEditSpecifier = true;
-			}
 				break;
 
-			case EVariableSpecifier::EditInstanceOnly:
-			{
-				if (bSeenEditSpecifier)
+				case EVariableSpecifier::EditInstanceOnly:
 				{
-					FError::Throwf(TEXT("Found more than one edit/visibility specifier (%s), only one is allowed"), *Specifier.Key);
+					if (bSeenEditSpecifier)
+					{
+						FError::Throwf(TEXT("Found more than one edit/visibility specifier (%s), only one is allowed"), *Specifier.Key);
+					}
+					Flags |= CPF_Edit | CPF_DisableEditOnTemplate;
+					bSeenEditSpecifier = true;
 				}
-				Flags |= CPF_Edit | CPF_DisableEditOnTemplate;
-				bSeenEditSpecifier = true;
-			}
 				break;
 
-			case EVariableSpecifier::EditDefaultsOnly:
-			{
-				if (bSeenEditSpecifier)
+				case EVariableSpecifier::EditDefaultsOnly:
 				{
-					FError::Throwf(TEXT("Found more than one edit/visibility specifier (%s), only one is allowed"), *Specifier.Key);
+					if (bSeenEditSpecifier)
+					{
+						FError::Throwf(TEXT("Found more than one edit/visibility specifier (%s), only one is allowed"), *Specifier.Key);
+					}
+					Flags |= CPF_Edit | CPF_DisableEditOnInstance;
+					bSeenEditSpecifier = true;
 				}
-				Flags |= CPF_Edit | CPF_DisableEditOnInstance;
-				bSeenEditSpecifier = true;
-			}
 				break;
 
-			case EVariableSpecifier::VisibleAnywhere:
-			{
-				if (bSeenEditSpecifier)
+				case EVariableSpecifier::VisibleAnywhere:
 				{
-					FError::Throwf(TEXT("Found more than one edit/visibility specifier (%s), only one is allowed"), *Specifier.Key);
+					if (bSeenEditSpecifier)
+					{
+						FError::Throwf(TEXT("Found more than one edit/visibility specifier (%s), only one is allowed"), *Specifier.Key);
+					}
+					Flags |= CPF_Edit | CPF_EditConst;
+					bSeenEditSpecifier = true;
 				}
-				Flags |= CPF_Edit | CPF_EditConst;
-				bSeenEditSpecifier = true;
-			}
 				break;
 
-			case EVariableSpecifier::VisibleInstanceOnly:
-			{
-				if (bSeenEditSpecifier)
+				case EVariableSpecifier::VisibleInstanceOnly:
 				{
-					FError::Throwf(TEXT("Found more than one edit/visibility specifier (%s), only one is allowed"), *Specifier.Key);
+					if (bSeenEditSpecifier)
+					{
+						FError::Throwf(TEXT("Found more than one edit/visibility specifier (%s), only one is allowed"), *Specifier.Key);
+					}
+					Flags |= CPF_Edit | CPF_EditConst | CPF_DisableEditOnTemplate;
+					bSeenEditSpecifier = true;
 				}
-				Flags |= CPF_Edit | CPF_EditConst | CPF_DisableEditOnTemplate;
-				bSeenEditSpecifier = true;
-			}
 				break;
 
-			case EVariableSpecifier::VisibleDefaultsOnly:
-			{
-				if (bSeenEditSpecifier)
+				case EVariableSpecifier::VisibleDefaultsOnly:
 				{
-					FError::Throwf(TEXT("Found more than one edit/visibility specifier (%s), only one is allowed"), *Specifier.Key);
+					if (bSeenEditSpecifier)
+					{
+						FError::Throwf(TEXT("Found more than one edit/visibility specifier (%s), only one is allowed"), *Specifier.Key);
+					}
+					Flags |= CPF_Edit | CPF_EditConst | CPF_DisableEditOnInstance;
+					bSeenEditSpecifier = true;
 				}
-				Flags |= CPF_Edit | CPF_EditConst | CPF_DisableEditOnInstance;
-				bSeenEditSpecifier = true;
-			}
 				break;
 
-			case EVariableSpecifier::BlueprintReadWrite:
-			{
-				if (bSeenBlueprintEditSpecifier)
+				case EVariableSpecifier::BlueprintReadWrite:
 				{
-					FError::Throwf(TEXT("Found more than one Blueprint read/write specifier (%s), only one is allowed"), *Specifier.Key);
-				}
+					if (bSeenBlueprintEditSpecifier)
+					{
+						FError::Throwf(TEXT("Found more than one Blueprint read/write specifier (%s), only one is allowed"), *Specifier.Key);
+					}
 
-				const FString* PrivateAccessMD = MetaDataFromNewStyle.Find(TEXT("AllowPrivateAccess"));  // FBlueprintMetadata::MD_AllowPrivateAccess
-				const bool bAllowPrivateAccess = PrivateAccessMD ? (*PrivateAccessMD == TEXT("true")) : false;
-				if (CurrentAccessSpecifier == ACCESS_Private && !bAllowPrivateAccess)
-				{
-					FError::Throwf(TEXT("BlueprintReadWrite should not be used on private members"));
+					const FString* PrivateAccessMD = MetaDataFromNewStyle.Find(TEXT("AllowPrivateAccess"));  // FBlueprintMetadata::MD_AllowPrivateAccess
+					const bool bAllowPrivateAccess = PrivateAccessMD ? (*PrivateAccessMD == TEXT("true")) : false;
+					if (CurrentAccessSpecifier == ACCESS_Private && !bAllowPrivateAccess)
+					{
+						FError::Throwf(TEXT("BlueprintReadWrite should not be used on private members"));
+					}
+					Flags |= CPF_BlueprintVisible;
+					bSeenBlueprintEditSpecifier = true;
 				}
-				Flags |= CPF_BlueprintVisible;
-				bSeenBlueprintEditSpecifier = true;
-			}
 				break;
 
-			case EVariableSpecifier::BlueprintReadOnly:
-			{
-				if (bSeenBlueprintEditSpecifier)
+				case EVariableSpecifier::BlueprintReadOnly:
 				{
-					FError::Throwf(TEXT("Found more than one Blueprint read/write specifier (%s), only one is allowed"), *Specifier.Key);
-				}
+					if (bSeenBlueprintEditSpecifier)
+					{
+						FError::Throwf(TEXT("Found more than one Blueprint read/write specifier (%s), only one is allowed"), *Specifier.Key);
+					}
 
-				const FString* PrivateAccessMD = MetaDataFromNewStyle.Find(TEXT("AllowPrivateAccess"));  // FBlueprintMetadata::MD_AllowPrivateAccess
-				const bool bAllowPrivateAccess = PrivateAccessMD ? (*PrivateAccessMD == TEXT("true")) : false;
-				if (CurrentAccessSpecifier == ACCESS_Private && !bAllowPrivateAccess)
-				{
-					FError::Throwf(TEXT("BlueprintReadOnly should not be used on private members"));
+					const FString* PrivateAccessMD = MetaDataFromNewStyle.Find(TEXT("AllowPrivateAccess"));  // FBlueprintMetadata::MD_AllowPrivateAccess
+					const bool bAllowPrivateAccess = PrivateAccessMD ? (*PrivateAccessMD == TEXT("true")) : false;
+					if (CurrentAccessSpecifier == ACCESS_Private && !bAllowPrivateAccess)
+					{
+						FError::Throwf(TEXT("BlueprintReadOnly should not be used on private members"));
+					}
+					Flags        |= CPF_BlueprintVisible | CPF_BlueprintReadOnly;
+					ImpliedFlags &= ~CPF_BlueprintReadOnly;
+					bSeenBlueprintEditSpecifier = true;
 				}
-				Flags |= CPF_BlueprintVisible | CPF_BlueprintReadOnly;
-				bSeenBlueprintEditSpecifier = true;
-			}
 				break;
 
-			case EVariableSpecifier::Config:
-			{
-				Flags |= CPF_Config;
-			}
+				case EVariableSpecifier::Config:
+				{
+					Flags |= CPF_Config;
+				}
 				break;
 
-			case EVariableSpecifier::GlobalConfig:
-			{
-				Flags |= CPF_GlobalConfig | CPF_Config;
-			}
+				case EVariableSpecifier::GlobalConfig:
+				{
+					Flags |= CPF_GlobalConfig | CPF_Config;
+				}
 				break;
 
-			case EVariableSpecifier::Localized:
-			{
-				FError::Throwf(TEXT("The Localized specifier is deprecated"));
-			}
+				case EVariableSpecifier::Localized:
+				{
+					FError::Throwf(TEXT("The Localized specifier is deprecated"));
+				}
 				break;
 
-			case EVariableSpecifier::Transient:
-			{
-				Flags |= CPF_Transient;
-			}
+				case EVariableSpecifier::Transient:
+				{
+					Flags |= CPF_Transient;
+				}
 				break;
 
-			case EVariableSpecifier::DuplicateTransient:
-			{
-				Flags |= CPF_DuplicateTransient;
-			}
+				case EVariableSpecifier::DuplicateTransient:
+				{
+					Flags |= CPF_DuplicateTransient;
+				}
 				break;
 
-			case EVariableSpecifier::TextExportTransient:
-			{
-				Flags |= CPF_TextExportTransient;
-			}
+				case EVariableSpecifier::TextExportTransient:
+				{
+					Flags |= CPF_TextExportTransient;
+				}
 				break;
 
-			case EVariableSpecifier::NonPIETransient:
-			{
-				UE_LOG(LogCompile, Warning, TEXT("NonPIETransient is deprecated - NonPIEDuplicateTransient should be used instead"));
-				Flags |= CPF_NonPIEDuplicateTransient;
-			}
+				case EVariableSpecifier::NonPIETransient:
+				{
+					UE_LOG(LogCompile, Warning, TEXT("NonPIETransient is deprecated - NonPIEDuplicateTransient should be used instead"));
+					Flags |= CPF_NonPIEDuplicateTransient;
+				}
 				break;
 
-			case EVariableSpecifier::NonPIEDuplicateTransient:
-			{
-				Flags |= CPF_NonPIEDuplicateTransient;
-			}
+				case EVariableSpecifier::NonPIEDuplicateTransient:
+				{
+					Flags |= CPF_NonPIEDuplicateTransient;
+				}
 				break;
 
-			case EVariableSpecifier::Export:
-			{
-				Flags |= CPF_ExportObject;
-			}
+				case EVariableSpecifier::Export:
+				{
+					Flags |= CPF_ExportObject;
+				}
 				break;
 
-			case EVariableSpecifier::EditInline:
-			{
-				FError::Throwf(TEXT("EditInline is deprecated. Remove it, or use Instanced instead."));
-			}
+				case EVariableSpecifier::EditInline:
+				{
+					FError::Throwf(TEXT("EditInline is deprecated. Remove it, or use Instanced instead."));
+				}
 				break;
 
-			case EVariableSpecifier::NoClear:
-			{
-				Flags |= CPF_NoClear;
-			}
+				case EVariableSpecifier::NoClear:
+				{
+					Flags |= CPF_NoClear;
+				}
 				break;
 
-			case EVariableSpecifier::EditFixedSize:
-			{
-				Flags |= CPF_EditFixedSize;
-			}
+				case EVariableSpecifier::EditFixedSize:
+				{
+					Flags |= CPF_EditFixedSize;
+				}
 				break;
 
-			case EVariableSpecifier::Replicated:
-			case EVariableSpecifier::ReplicatedUsing:
-			{
-				if (OwnerStruct->IsA<UScriptStruct>())
+				case EVariableSpecifier::Replicated:
+				case EVariableSpecifier::ReplicatedUsing:
 				{
-					FError::Throwf(TEXT("Struct members cannot be replicated"));
-				}
+					if (OwnerStruct->IsA<UScriptStruct>())
+					{
+						FError::Throwf(TEXT("Struct members cannot be replicated"));
+					}
 
-				Flags |= CPF_Net;
+					Flags |= CPF_Net;
 
-				// See if we've specified a rep notification function
-				if (SpecID == EVariableSpecifier::ReplicatedUsing)
-				{
-					RepCallbackName = FName(*RequireExactlyOneSpecifierValue(Specifier));
-					Flags |= CPF_RepNotify;
+					// See if we've specified a rep notification function
+					if (SpecID == EVariableSpecifier::ReplicatedUsing)
+					{
+						RepCallbackName = FName(*RequireExactlyOneSpecifierValue(Specifier));
+						Flags |= CPF_RepNotify;
+					}
 				}
-			}
 				break;
 
-			case EVariableSpecifier::NotReplicated:
-			{
-				if (!OwnerStruct->IsA<UScriptStruct>())
+				case EVariableSpecifier::NotReplicated:
 				{
-					FError::Throwf(TEXT("Only Struct members can be marked NotReplicated"));
-				}
+					if (!OwnerStruct->IsA<UScriptStruct>())
+					{
+						FError::Throwf(TEXT("Only Struct members can be marked NotReplicated"));
+					}
 
-				Flags |= CPF_RepSkip;
-			}
+					Flags |= CPF_RepSkip;
+				}
 				break;
 
-			case EVariableSpecifier::RepRetry:
-			{
-				FError::Throwf(TEXT("'RepRetry' is deprecated."));
-			}
+				case EVariableSpecifier::RepRetry:
+				{
+					FError::Throwf(TEXT("'RepRetry' is deprecated."));
+				}
 				break;
 
-			case EVariableSpecifier::Interp:
-			{
-				Flags |= CPF_Edit;
-				Flags |= CPF_BlueprintVisible;
-				Flags |= CPF_Interp;
-			}
+				case EVariableSpecifier::Interp:
+				{
+					Flags |= CPF_Edit;
+					Flags |= CPF_BlueprintVisible;
+					Flags |= CPF_Interp;
+				}
 				break;
 
-			case EVariableSpecifier::NonTransactional:
-			{
-				Flags |= CPF_NonTransactional;
-			}
+				case EVariableSpecifier::NonTransactional:
+				{
+					Flags |= CPF_NonTransactional;
+				}
 				break;
 
-			case EVariableSpecifier::Instanced:
-			{
-				Flags |= CPF_PersistentInstance | CPF_ExportObject | CPF_InstancedReference;
-				AddEditInlineMetaData(MetaDataFromNewStyle);
-			}
+				case EVariableSpecifier::Instanced:
+				{
+					Flags |= CPF_PersistentInstance | CPF_ExportObject | CPF_InstancedReference;
+					AddEditInlineMetaData(MetaDataFromNewStyle);
+				}
 				break;
 
-			case EVariableSpecifier::BlueprintAssignable:
-			{
-				Flags |= CPF_BlueprintAssignable;
-			}
+				case EVariableSpecifier::BlueprintAssignable:
+				{
+					Flags |= CPF_BlueprintAssignable;
+				}
 				break;
 
-			case EVariableSpecifier::BlueprintCallable:
-			{
-				Flags |= CPF_BlueprintCallable;
-			}
+				case EVariableSpecifier::BlueprintCallable:
+				{
+					Flags |= CPF_BlueprintCallable;
+				}
 				break;
 
-			case EVariableSpecifier::BlueprintAuthorityOnly:
-			{
-				Flags |= CPF_BlueprintAuthorityOnly;
-			}
+				case EVariableSpecifier::BlueprintAuthorityOnly:
+				{
+					Flags |= CPF_BlueprintAuthorityOnly;
+				}
 				break;
 
-			case EVariableSpecifier::AssetRegistrySearchable:
-			{
-				Flags |= CPF_AssetRegistrySearchable;
-			}
+				case EVariableSpecifier::AssetRegistrySearchable:
+				{
+					Flags |= CPF_AssetRegistrySearchable;
+				}
 				break;
 
-			case EVariableSpecifier::SimpleDisplay:
-			{
-				Flags |= CPF_SimpleDisplay;
-			}
+				case EVariableSpecifier::SimpleDisplay:
+				{
+					Flags |= CPF_SimpleDisplay;
+				}
 				break;
 
-			case EVariableSpecifier::AdvancedDisplay:
-			{
-				Flags |= CPF_AdvancedDisplay;
-			}
+				case EVariableSpecifier::AdvancedDisplay:
+				{
+					Flags |= CPF_AdvancedDisplay;
+				}
 				break;
 
-			case EVariableSpecifier::SaveGame:
-			{
-				Flags |= CPF_SaveGame;
-			}
+				case EVariableSpecifier::SaveGame:
+				{
+					Flags |= CPF_SaveGame;
+				}
 				break;
 
-			default:
-			{
-				FError::Throwf(TEXT("Unknown variable specifier '%s'"), *Specifier.Key);
-			}
+				default:
+				{
+					FError::Throwf(TEXT("Unknown variable specifier '%s'"), *Specifier.Key);
+				}
 				break;
 			}
 		}
@@ -2962,36 +2971,36 @@ FIndexRange*                    ParsedVarIndexRange
 		{
 			switch (SpecID)
 			{
-			case EVariableSpecifier::Const:
-			{
-				Flags |= CPF_ConstParm;
-			}
-				break;
-
-			case EVariableSpecifier::Ref:
-			{
-				Flags |= CPF_OutParm | CPF_ReferenceParm;
-			}
+				case EVariableSpecifier::Const:
+				{
+					Flags |= CPF_ConstParm;
+				}
 				break;
 
-			case EVariableSpecifier::NotReplicated:
-			{
-				if (VariableCategory == EVariableCategory::ReplicatedParameter)
+				case EVariableSpecifier::Ref:
 				{
-					VariableCategory = EVariableCategory::RegularParameter;
-					Flags |= CPF_RepSkip;
+					Flags |= CPF_OutParm | CPF_ReferenceParm;
 				}
-				else
+				break;
+
+				case EVariableSpecifier::NotReplicated:
 				{
-					FError::Throwf(TEXT("Only parameters in service request functions can be marked NotReplicated"));
+					if (VariableCategory == EVariableCategory::ReplicatedParameter)
+					{
+						VariableCategory = EVariableCategory::RegularParameter;
+						Flags |= CPF_RepSkip;
+					}
+					else
+					{
+						FError::Throwf(TEXT("Only parameters in service request functions can be marked NotReplicated"));
+					}
 				}
-			}
 				break;
 
-			default:
-			{
-				FError::Throwf(TEXT("Unknown variable specifier '%s'"), *Specifier.Key);
-			}
+				default:
+				{
+					FError::Throwf(TEXT("Unknown variable specifier '%s'"), *Specifier.Key);
+				}
 				break;
 			}
 		}
@@ -3723,7 +3732,8 @@ FIndexRange*                    ParsedVarIndexRange
 	VarProperty.PropertyExportFlags = ExportFlags;
 
 	// Set FPropertyBase info.
-	VarProperty.PropertyFlags |= Flags;
+	VarProperty.PropertyFlags        |= Flags | ImpliedFlags;
+	VarProperty.ImpliedPropertyFlags |= ImpliedFlags;
 
 	// Set the RepNotify name, if the variable needs it
 	if( VarProperty.PropertyFlags & CPF_RepNotify )
@@ -3891,7 +3901,7 @@ UProperty* FHeaderParser::GetVarNameAndDim
 			}
 
 			// Warn if a deprecated property is visible
-			if (VarProperty.PropertyFlags & (CPF_Edit | CPF_EditConst | CPF_BlueprintVisible | CPF_BlueprintReadOnly))
+			if (VarProperty.PropertyFlags & (CPF_Edit | CPF_EditConst | CPF_BlueprintVisible | CPF_BlueprintReadOnly) && !(VarProperty.ImpliedPropertyFlags & CPF_BlueprintReadOnly))
 			{
 				UE_LOG(LogCompile, Warning, TEXT("%s: Deprecated property '%s' should not be marked as visible or editable"), HintText, *VarName);
 			}
diff --git a/Engine/Source/Programs/UnrealHeaderTool/Private/ParserHelper.h b/Engine/Source/Programs/UnrealHeaderTool/Private/ParserHelper.h
--- a/Engine/Source/Programs/UnrealHeaderTool/Private/ParserHelper.h
+++ b/Engine/Source/Programs/UnrealHeaderTool/Private/ParserHelper.h
@@ -72,6 +72,7 @@ public:
 	EPropertyType       Type;
 	EArrayType::Type    ArrayType;
 	uint64              PropertyFlags;
+	uint64              ImpliedPropertyFlags;
 	ERefQualifier::Type RefQualifier; // This is needed because of legacy stuff - FString mangles the flags for reasons that have become lost in time but we need this info for testing for invalid replicated function signatures.
 
 	TSharedPtr<FPropertyBase> MapKeyProp;
@@ -109,44 +110,47 @@ public:
 	/** @name Constructors */
 	//@{
 	explicit FPropertyBase(EPropertyType InType)
-	: Type               (InType)
-	, ArrayType          (EArrayType::None)
-	, PropertyFlags      (0)
-	, RefQualifier       (ERefQualifier::None)
-	, PropertyExportFlags(PROPEXPORT_Public)
-	, StringSize         (0)
-	, MetaClass          (NULL)
-	, DelegateName       (NAME_None)
-	, RepNotifyName      (NAME_None)
-	, PointerType        (EPointerType::None)
+	: Type                (InType)
+	, ArrayType           (EArrayType::None)
+	, PropertyFlags       (0)
+	, ImpliedPropertyFlags(0)
+	, RefQualifier        (ERefQualifier::None)
+	, PropertyExportFlags (PROPEXPORT_Public)
+	, StringSize          (0)
+	, MetaClass           (NULL)
+	, DelegateName        (NAME_None)
+	, RepNotifyName       (NAME_None)
+	, PointerType         (EPointerType::None)
 	{
 	}
 
 	explicit FPropertyBase(UEnum* InEnum, EPropertyType InType)
-	: Type               (InType)
-	, ArrayType          (EArrayType::None)
-	, PropertyFlags      (0)
-	, RefQualifier       (ERefQualifier::None)
-	, PropertyExportFlags(PROPEXPORT_Public)
-	, Enum               (InEnum)
-	, MetaClass          (NULL)
-	, DelegateName       (NAME_None)
-	, RepNotifyName      (NAME_None)
-	, PointerType        (EPointerType::None)
+	: Type                (InType)
+	, ArrayType           (EArrayType::None)
+	, PropertyFlags       (0)
+	, ImpliedPropertyFlags(0)
+	, RefQualifier        (ERefQualifier::None)
+	, PropertyExportFlags (PROPEXPORT_Public)
+	, Enum                (InEnum)
+	, MetaClass           (NULL)
+	, DelegateName        (NAME_None)
+	, RepNotifyName       (NAME_None)
+	, PointerType         (EPointerType::None)
 	{
 	}
 
 	explicit FPropertyBase(UClass* InClass, UClass* InMetaClass=NULL, bool bAllowWeak = false, bool bIsWeak = false, bool bWeakIsAuto = false, bool bIsLazy = false, bool bIsAsset = false)
-	: Type               (CPT_ObjectReference)
-	, ArrayType          (EArrayType::None)
-	, PropertyFlags      (0)
-	, RefQualifier       (ERefQualifier::None)
-	, PropertyExportFlags(PROPEXPORT_Public)
-	, PropertyClass      (InClass)
-	, MetaClass          (InMetaClass)
-	, DelegateName       (NAME_None)
-	, RepNotifyName      (NAME_None)
-	, PointerType        (EPointerType::None)
+	: Type                (CPT_ObjectReference)
+	, ArrayType           (EArrayType::None)
+	, PropertyFlags       (0)
+	, ImpliedPropertyFlags(0)
+	, RefQualifier        (ERefQualifier::None)
+	, PropertyExportFlags (PROPEXPORT_Public)
+	, PropertyClass       (InClass)
+	, MetaClass           (InMetaClass)
+	, DelegateName        (NAME_None)
+	, RepNotifyName       (NAME_None)
+	, PointerType         (EPointerType::None)
 	{
 		// if this is an interface class, we use the UInterfaceProperty class instead of UObjectProperty
 		if ( InClass->HasAnyClassFlags(CLASS_Interface) )
@@ -194,16 +198,17 @@ public:
 	}
 
 	explicit FPropertyBase(UScriptStruct* InStruct)
-	: Type               (CPT_Struct)
-	, ArrayType          (EArrayType::None)
-	, PropertyFlags      (0)
-	, RefQualifier       (ERefQualifier::None)
-	, PropertyExportFlags(PROPEXPORT_Public)
-	, Struct             (InStruct)
-	, MetaClass          (NULL)
-	, DelegateName       (NAME_None)
-	, RepNotifyName      (NAME_None)
-	, PointerType        (EPointerType::None)
+	: Type                (CPT_Struct)
+	, ArrayType           (EArrayType::None)
+	, PropertyFlags       (0)
+	, ImpliedPropertyFlags(0)
+	, RefQualifier        (ERefQualifier::None)
+	, PropertyExportFlags (PROPEXPORT_Public)
+	, Struct              (InStruct)
+	, MetaClass           (NULL)
+	, DelegateName        (NAME_None)
+	, RepNotifyName       (NAME_None)
+	, PointerType         (EPointerType::None)
 	{
 	}
 
@@ -365,10 +370,11 @@ public:
 		{
 			UE_LOG(LogCompile, Fatal, TEXT("Unknown property type '%s'"), *Property->GetFullName() );
 		}
-		ArrayType     = ArrType;
-		PropertyFlags = Property->PropertyFlags | PropagateFlags;
-		RefQualifier  = ERefQualifier::None;
-		PointerType   = EPointerType::None;
+		ArrayType            = ArrType;
+		PropertyFlags        = Property->PropertyFlags | PropagateFlags;
+		ImpliedPropertyFlags = 0;
+		RefQualifier         = ERefQualifier::None;
+		PointerType          = EPointerType::None;
 	}
 	//@}
 
@@ -574,8 +580,8 @@ public:
 	FString Describe()
 	{
 		return FString::Printf(
-			TEXT("Type:%s  Flags:%lli  Enum:%s  PropertyClass:%s  Struct:%s  Function:%s  MetaClass:%s"),
-			GetPropertyTypeText(Type), PropertyFlags,
+			TEXT("Type:%s  Flags:%lli  ImpliedFlags:%lli  Enum:%s  PropertyClass:%s  Struct:%s  Function:%s  MetaClass:%s"),
+			GetPropertyTypeText(Type), PropertyFlags, ImpliedPropertyFlags,
 			Enum!=NULL?*Enum->GetName():TEXT(""),
 			PropertyClass!=NULL?*PropertyClass->GetName():TEXT("NULL"),
 			Struct!=NULL?*Struct->GetName():TEXT("NULL"),
@@ -833,11 +839,12 @@ public:
 	FString Describe()
 	{
 		return FString::Printf(
-			TEXT("Property:%s  Type:%s  TokenName:%s  ConstValue:%s  Struct:%s  Flags:%lli"),
+			TEXT("Property:%s  Type:%s  TokenName:%s  ConstValue:%s  Struct:%s  Flags:%lli  Implied:%lli"),
 			TokenProperty!=NULL?*TokenProperty->GetName():TEXT("NULL"),
 			GetPropertyTypeText(Type), *TokenName.ToString(), *GetConstantValue(),
 			Struct!=NULL?*Struct->GetName():TEXT("NULL"),
-			PropertyFlags
+			PropertyFlags,
+			ImpliedPropertyFlags
 			);
 	}
 };
