Basic support for 'dynamic' Utypes (types that don't need to be constructed at runtime, can be GC'd, etc). Native blueprint support for UHT.

diff --git a/Engine/Source/Programs/UnrealHeaderTool/Private/CodeGenerator.cpp b/Engine/Source/Programs/UnrealHeaderTool/Private/CodeGenerator.cpp
--- a/Engine/Source/Programs/UnrealHeaderTool/Private/CodeGenerator.cpp
+++ b/Engine/Source/Programs/UnrealHeaderTool/Private/CodeGenerator.cpp
@@ -1312,7 +1312,11 @@ void FNativeClassHeaderGenerator::ExportGeneratedPackageInitCode(UPackage* InPac
 				ScriptType->IsA<UDelegateFunction>()
 				)
 			{
-				GeneratedFunctionText.Logf(TEXT("            %s;\r\n"), *GetSingletonName(ScriptType));
+				UField* FieldOuter = Cast<UField>(ScriptType->GetOuter());
+				if (!FieldOuter || !FClass::IsDynamic(FieldOuter))
+				{
+					GeneratedFunctionText.Logf(TEXT("            %s;\r\n"), *GetSingletonName(ScriptType));
+				}
 			}
 		}
 	}
@@ -1328,7 +1332,8 @@ void FNativeClassHeaderGenerator::ExportNativeGeneratedInitCode(FClass* Class, F
 	FUHTStringBuilder& GeneratedFunctionText = GetGeneratedFunctionTextDevice();
 	check(!OutFriendText.Len());
 	// Emit code to build the UObjects that used to be in .u files
-	bool bIsNoExport = Class->HasAnyClassFlags(CLASS_NoExport);
+	const bool bIsNoExport = Class->HasAnyClassFlags(CLASS_NoExport);
+	const bool bIsDynamic = FClass::IsDynamic(Class);
 	FUHTStringBuilder BodyText;
 	FUHTStringBuilder CallSingletons;
 	FString ApiString = GetAPIString();
@@ -1377,7 +1382,15 @@ void FNativeClassHeaderGenerator::ExportNativeGeneratedInitCode(FClass* Class, F
 
 		GeneratedClassRegisterFunctionText.Logf(TEXT("    UClass* %s\r\n"), *SingletonName);
 		GeneratedClassRegisterFunctionText.Logf(TEXT("    {\r\n"));
-		GeneratedClassRegisterFunctionText.Logf(TEXT("        static UClass* OuterClass = NULL;\r\n"));
+		if (!bIsDynamic)
+		{
+			GeneratedClassRegisterFunctionText.Logf(TEXT("        static UClass* OuterClass = NULL;\r\n"));
+		}
+		else
+		{
+			GeneratedClassRegisterFunctionText.Logf(TEXT("        UClass* OuterClass = Cast<UClass>(StaticFindObjectFast(UClass::StaticClass(), %s, TEXT(\"%s\")));\r\n"),
+				*GetPackageSingletonName(CastChecked<UPackage>(Class->GetOutermost())), *Class->GetName());
+		}
 		GeneratedClassRegisterFunctionText.Logf(TEXT("        if (!OuterClass)\r\n"));
 		GeneratedClassRegisterFunctionText.Logf(TEXT("        {\r\n"));
 		if (Class->GetSuperClass() && Class->GetSuperClass() != Class)
@@ -1489,8 +1502,10 @@ void FNativeClassHeaderGenerator::ExportNativeGeneratedInitCode(FClass* Class, F
 
 	{	
 		const TCHAR* ClassNameCPP = NameLookupCPP.GetNameCPP(Class);
-		GeneratedFunctionText.Logf(TEXT("    static FCompiledInDefer Z_CompiledInDefer_UClass_%s(%s, TEXT(\"%s\"));\r\n"), 
-			ClassNameCPP, *SingletonName, ClassNameCPP);
+
+		GeneratedFunctionText.Logf(TEXT("    static FCompiledInDefer Z_CompiledInDefer_UClass_%s(%s, TEXT(\"%s\"), %s);\r\n"),
+			ClassNameCPP, *SingletonName, ClassNameCPP,
+			bIsDynamic ? TEXT("true") : TEXT("false"));
 		
 		// Append base class' CRC at the end of the generated code, this will force update derived classes
 		// when base class changes during hot-reload.
@@ -1505,7 +1520,14 @@ void FNativeClassHeaderGenerator::ExportNativeGeneratedInitCode(FClass* Class, F
 		uint32 ClassCrc = GenerateTextCRC(*GeneratedClassRegisterFunctionText);
 		GGeneratedCodeCRCs.Add(Class, ClassCrc);
 		// Emit the IMPLEMENT_CLASS macro to go in the generated cpp file.
-		GeneratedPackageCPP.Logf(TEXT("    IMPLEMENT_CLASS(%s, %u);\r\n"), ClassNameCPP, ClassCrc);
+		if (!bIsDynamic)
+		{
+			GeneratedPackageCPP.Logf(TEXT("    IMPLEMENT_CLASS(%s, %u);\r\n"), ClassNameCPP, ClassCrc);
+		}
+		else
+		{
+			GeneratedPackageCPP.Logf(TEXT("    IMPLEMENT_DYNAMIC_CLASS(%s, %u);\r\n"), ClassNameCPP, ClassCrc);			
+		}
 	}
 }
 
@@ -2719,6 +2741,7 @@ void FNativeClassHeaderGenerator::ExportGeneratedStructBodyMacros(FUnrealSourceF
 	for (int32 i = NativeStructs.Num() - 1; i >= 0; --i)
 	{
 		UScriptStruct* Struct = NativeStructs[i];
+		const bool bIsDynamic = FClass::IsDynamic(Struct);
 
 		// Export struct.
 		if (Struct->StructFlags & STRUCT_Native)
@@ -2746,34 +2769,48 @@ void FNativeClassHeaderGenerator::ExportGeneratedStructBodyMacros(FUnrealSourceF
 			const TCHAR* StructNameCPP = NameLookupCPP.GetNameCPP(Struct);
 			FString SingletonName = StaticConstructionString.Replace(TEXT("()"), TEXT(""), ESearchCase::CaseSensitive); // function address
 			FString GetCRCName = FString::Printf(TEXT("Get_%s_CRC"), *SingletonName);
-
+			
 			GeneratedPackageCPP.Logf(TEXT("class UScriptStruct* %s::StaticStruct()\r\n"), StructNameCPP);
 			GeneratedPackageCPP.Logf(TEXT("{\r\n"));
-			GeneratedPackageCPP.Logf(TEXT("    static class UScriptStruct* Singleton = NULL;\r\n"));
-			GeneratedPackageCPP.Logf(TEXT("    if (!Singleton)\r\n"));
-			GeneratedPackageCPP.Logf(TEXT("    {\r\n"));
-			GeneratedPackageCPP.Logf(TEXT("        extern %sclass UScriptStruct* %s;\r\n"), *FriendApiString, *StaticConstructionString);
-			GeneratedPackageCPP.Logf(TEXT("        extern %suint32 %s();\r\n"), *FriendApiString, *GetCRCName);
 
 			// UStructs can have UClass or UPackage outer (if declared in non-UClass headers).
+			FString OuterName;
 			if (Struct->GetOuter()->IsA(UStruct::StaticClass()))
 			{
-				GeneratedPackageCPP.Logf(TEXT("        Singleton = GetStaticStruct(%s, %s::StaticClass(), TEXT(\"%s\"), sizeof(%s), %s());\r\n"),
-					*SingletonName, NameLookupCPP.GetNameCPP(CastChecked<UStruct>(Struct->GetOuter())), *Struct->GetName(), StructNameCPP, *GetCRCName);
+				OuterName = NameLookupCPP.GetNameCPP(CastChecked<UStruct>(Struct->GetOuter()));
+				OuterName += TEXT("::StaticClass()");
+			}
+			else
+			{
+				OuterName = GetPackageSingletonName(CastChecked<UPackage>(Struct->GetOuter()));
+				GeneratedPackageCPP.Logf(TEXT("    extern %sclass UPackage* %s;\r\n"), *FriendApiString, *OuterName);
+			}
+
+			if (!bIsDynamic)
+			{
+				GeneratedPackageCPP.Logf(TEXT("    static class UScriptStruct* Singleton = NULL;\r\n"));
 			}
 			else
 			{
 				FString PackageSingletonName = GetPackageSingletonName(CastChecked<UPackage>(Struct->GetOuter()));
-				GeneratedPackageCPP.Logf(TEXT("        extern %sclass UPackage* %s;\r\n"), *FriendApiString, *PackageSingletonName);
-				GeneratedPackageCPP.Logf(TEXT("        Singleton = GetStaticStruct(%s, %s, TEXT(\"%s\"), sizeof(%s), %s());\r\n"),
-					*SingletonName, *PackageSingletonName, *Struct->GetName(), StructNameCPP, *GetCRCName);
+				GeneratedPackageCPP.Logf(TEXT("    class UScriptStruct* Singleton = Cast<UScriptStruct>(StaticFindObjectFast(UScriptStruct::StaticClass(), %s, TEXT(\"%s\")));\r\n"),
+					*OuterName, *Struct->GetName());
 			}
+			GeneratedPackageCPP.Logf(TEXT("    if (!Singleton)\r\n"));
+			GeneratedPackageCPP.Logf(TEXT("    {\r\n"));
+			GeneratedPackageCPP.Logf(TEXT("        extern %sclass UScriptStruct* %s;\r\n"), *FriendApiString, *StaticConstructionString);
+			GeneratedPackageCPP.Logf(TEXT("        extern %suint32 %s();\r\n"), *FriendApiString, *GetCRCName);
+
+			GeneratedPackageCPP.Logf(TEXT("        Singleton = GetStaticStruct(%s, %s, TEXT(\"%s\"), sizeof(%s), %s());\r\n"),
+				*SingletonName, *OuterName, *Struct->GetName(), StructNameCPP, *GetCRCName);
 
 			GeneratedPackageCPP.Logf(TEXT("    }\r\n"));
 			GeneratedPackageCPP.Logf(TEXT("    return Singleton;\r\n"));
 			GeneratedPackageCPP.Logf(TEXT("}\r\n"));
-			GeneratedPackageCPP.Logf(TEXT("static FCompiledInDeferStruct Z_CompiledInDeferStruct_UScriptStruct_%s(%s::StaticStruct, TEXT(\"%s\"));\r\n"),
-				StructNameCPP, StructNameCPP, *Struct->GetOutermost()->GetName());
+
+			GeneratedPackageCPP.Logf(TEXT("static FCompiledInDeferStruct Z_CompiledInDeferStruct_UScriptStruct_%s(%s::StaticStruct, TEXT(\"%s\"), TEXT(\"%s\"), %s);\r\n"),
+				StructNameCPP, StructNameCPP, *Struct->GetOutermost()->GetName(), *Struct->GetName(),
+				bIsDynamic ? TEXT("true") : TEXT("false"));
 
 			// Generate StaticRegisterNatives equivalent for structs without classes.
 			if (!Struct->GetOuter()->IsA(UStruct::StaticClass()))
@@ -2821,7 +2858,15 @@ void FNativeClassHeaderGenerator::ExportGeneratedStructBodyMacros(FUnrealSourceF
 			{
 				GeneratedStructRegisterFunctionText.Logf(TEXT("        UPackage* Outer=%s;\r\n"), *GetPackageSingletonName(CastChecked<UPackage>(ScriptStruct->GetOuter())));
 			}
-			GeneratedStructRegisterFunctionText.Logf(TEXT("        static UScriptStruct* ReturnStruct = NULL;\r\n"));
+			if (!bIsDynamic)
+			{
+				GeneratedStructRegisterFunctionText.Logf(TEXT("        static UScriptStruct* ReturnStruct = NULL;\r\n"));
+			}
+			else
+			{
+				GeneratedStructRegisterFunctionText.Logf(TEXT("        UScriptStruct* ReturnStruct = Cast<UScriptStruct>(StaticFindObjectFast(UScriptStruct::StaticClass(), Outer, TEXT(\"%s\")));\r\n"),
+					*ScriptStruct->GetName());
+			}
 			GeneratedStructRegisterFunctionText.Logf(TEXT("        if (!ReturnStruct)\r\n"));
 			GeneratedStructRegisterFunctionText.Logf(TEXT("        {\r\n"));
 			FString BaseStructString(TEXT("NULL"));
@@ -2899,26 +2944,38 @@ void FNativeClassHeaderGenerator::ExportGeneratedEnumsInitCode(const TArray<UEnu
 			continue;
 		}
 
+		const bool bIsDynamic = FClass::IsDynamic(Enum);
 		const FString FriendApiString = GetAPIString();
 		const FString StaticConstructionString = GetSingletonName(Enum);
 
 		FString SingletonName = StaticConstructionString.Replace(TEXT("()"), TEXT(""), ESearchCase::CaseSensitive); // function address
+		FString PackageSingletonName = GetPackageSingletonName(CastChecked<UPackage>(Enum->GetOuter()));
 
 		GeneratedPackageCPP.Logf(TEXT("static class UEnum* %s_StaticEnum()\r\n"), *Enum->GetName());
 		GeneratedPackageCPP.Logf(TEXT("{\r\n"));
-		GeneratedPackageCPP.Logf(TEXT("    static class UEnum* Singleton = NULL;\r\n"));
+		GeneratedPackageCPP.Logf(TEXT("    extern %sclass UPackage* %s;\r\n"), *FriendApiString, *PackageSingletonName);
+
+		if (!bIsDynamic)
+		{
+			GeneratedPackageCPP.Logf(TEXT("    static class UEnum* Singleton = NULL;\r\n"));
+		}
+		else
+		{
+			GeneratedPackageCPP.Logf(TEXT("    class UEnum* Singleton = Cast<UEnum>(StaticFindObjectFast(UEnum::StaticClass(), %s, TEXT(\"%s\")));\r\n"),
+				*PackageSingletonName, *Enum->GetName());
+		}
 		GeneratedPackageCPP.Logf(TEXT("    if (!Singleton)\r\n"));
 		GeneratedPackageCPP.Logf(TEXT("    {\r\n"));
-		GeneratedPackageCPP.Logf(TEXT("        extern %sclass UEnum* %s;\r\n"), *FriendApiString, *StaticConstructionString);
-
-		FString PackageSingletonName = GetPackageSingletonName(CastChecked<UPackage>(Enum->GetOuter()));
-		GeneratedPackageCPP.Logf(TEXT("        extern %sclass UPackage* %s;\r\n"), *FriendApiString, *PackageSingletonName);
+		GeneratedPackageCPP.Logf(TEXT("        extern %sclass UEnum* %s;\r\n"), *FriendApiString, *StaticConstructionString);		
 		GeneratedPackageCPP.Logf(TEXT("        Singleton = GetStaticEnum(%s, %s, TEXT(\"%s\"));\r\n"), *SingletonName, *PackageSingletonName, *Enum->GetName());
 			
 		GeneratedPackageCPP.Logf(TEXT("    }\r\n"));
 		GeneratedPackageCPP.Logf(TEXT("    return Singleton;\r\n"));
 		GeneratedPackageCPP.Logf(TEXT("}\r\n"));
-		GeneratedPackageCPP.Logf(TEXT("static FCompiledInDeferEnum Z_CompiledInDeferEnum_UEnum_%s(%s_StaticEnum, TEXT(\"%s\"));\r\n"), *Enum->GetName(), *Enum->GetName(), *Enum->GetOutermost()->GetName());
+
+		GeneratedPackageCPP.Logf(TEXT("static FCompiledInDeferEnum Z_CompiledInDeferEnum_UEnum_%s(%s_StaticEnum, TEXT(\"%s\"), TEXT(\"%s\"), %s);\r\n"), 
+			*Enum->GetName(), *Enum->GetName(), *Enum->GetOutermost()->GetName(), *Enum->GetName(),
+			bIsDynamic ? TEXT("true") : TEXT("false"));
 
 		{
 			const FString EnumSingletonName = GetSingletonName(Enum);
diff --git a/Engine/Source/Programs/UnrealHeaderTool/Private/ParserClass.cpp b/Engine/Source/Programs/UnrealHeaderTool/Private/ParserClass.cpp
--- a/Engine/Source/Programs/UnrealHeaderTool/Private/ParserClass.cpp
+++ b/Engine/Source/Programs/UnrealHeaderTool/Private/ParserClass.cpp
@@ -86,3 +86,9 @@ void FClass::GetShowCategories(TArray<FString>& OutShowCategories) const
 		ShowCategories.ParseIntoArray(OutShowCategories, TEXT(" "), true);
 	}
 }
+
+bool FClass::IsDynamic(UField* Field)
+{
+	static const FName NAME_ReplaceConverted(TEXT("ReplaceConverted"));
+	return Field->HasMetaData(NAME_ReplaceConverted);
+}
\ No newline at end of file
diff --git a/Engine/Source/Programs/UnrealHeaderTool/Private/ParserClass.h b/Engine/Source/Programs/UnrealHeaderTool/Private/ParserClass.h
--- a/Engine/Source/Programs/UnrealHeaderTool/Private/ParserClass.h
+++ b/Engine/Source/Programs/UnrealHeaderTool/Private/ParserClass.h
@@ -44,4 +44,7 @@ public:
 
 	void GetHideCategories(TArray<FString>& OutHideCategories) const;
 	void GetShowCategories(TArray<FString>& OutShowCategories) const;
+
+	/** Helper function that checks if the field is a dynamic type (can be constructed post-startup) */
+	static bool IsDynamic(UField* Field);
 };
diff --git a/Engine/Source/Runtime/CoreUObject/Private/UObject/Class.cpp b/Engine/Source/Runtime/CoreUObject/Private/UObject/Class.cpp
--- a/Engine/Source/Runtime/CoreUObject/Private/UObject/Class.cpp
+++ b/Engine/Source/Runtime/CoreUObject/Private/UObject/Class.cpp
@@ -52,7 +52,16 @@ COREUOBJECT_API void InitializePrivateStaticClass(
 
 	// Register the class's dependencies, then itself.
 	TClass_PrivateStaticClass->RegisterDependencies();
-	TClass_PrivateStaticClass->Register(PackageName,Name);
+	if (!TClass_PrivateStaticClass->HasAnyFlags(RF_Dynamic))
+	{
+		// Defer
+		TClass_PrivateStaticClass->Register(PackageName, Name);
+	}
+	else
+	{
+		// Register immediately
+		TClass_PrivateStaticClass->DeferredRegister(UClass::StaticClass(), PackageName, Name);
+	}
 }
 
 void FNativeFunctionRegistrar::RegisterFunction(class UClass* Class, const ANSICHAR* InName, Native InPointer)
diff --git a/Engine/Source/Runtime/CoreUObject/Private/UObject/LinkerLoad.cpp b/Engine/Source/Runtime/CoreUObject/Private/UObject/LinkerLoad.cpp
--- a/Engine/Source/Runtime/CoreUObject/Private/UObject/LinkerLoad.cpp
+++ b/Engine/Source/Runtime/CoreUObject/Private/UObject/LinkerLoad.cpp
@@ -3913,11 +3913,7 @@ bool FLinkerLoad::IsImportNative(const int32 Index) const
 	return bIsImportNative;
 }
 
-// Temporary placeholder: constructs native blieprint class given its name
-void ConstructNativeBlueprintClass(FName ClassName)
-{
-	// @todo: this needs to find a generated function by class name.
-}
+UObject* ConstructDynamicType(FName TypeName, FName TypeClass);
 
 // Return the loaded object corresponding to an import index; any errors are fatal.
 UObject* FLinkerLoad::CreateImport( int32 Index )
@@ -3964,7 +3960,7 @@ UObject* FLinkerLoad::CreateImport( int32 Index )
 							// At this point we know the class package and class of the import exist, this can potentially be a native blueprint
 							// (class is just a UClass and the package is /Script/CoreUObject)
 							// @todo: maybe we can further filter out classes by checking FindClass somehow (e.g. once it's a specialized UCookedBlueprintClass or something)
-							ConstructNativeBlueprintClass(Import.ObjectName);
+							ConstructDynamicType(Import.ObjectName, Import.ClassName);
 						}
 						
 						// Find the imports' outer.
diff --git a/Engine/Source/Runtime/CoreUObject/Private/UObject/UObjectBase.cpp b/Engine/Source/Runtime/CoreUObject/Private/UObject/UObjectBase.cpp
--- a/Engine/Source/Runtime/CoreUObject/Private/UObject/UObjectBase.cpp
+++ b/Engine/Source/Runtime/CoreUObject/Private/UObject/UObjectBase.cpp
@@ -375,12 +375,25 @@ static TArray<FPendingStructRegistrant>& GetDeferredCompiledInStructRegistration
 	return DeferredCompiledInRegistration;
 }
 
-void UObjectCompiledInDeferStruct(class UScriptStruct *(*InRegister)(), const TCHAR* PackageName)
+TMap<FName, UScriptStruct *(*)()> GetDynamicStructMap()
 {
-	// we do reregister StaticStruct in hot reload
-	FPendingStructRegistrant Registrant(InRegister, PackageName);
-	checkSlow(!GetDeferredCompiledInStructRegistration().Contains(Registrant));
-	GetDeferredCompiledInStructRegistration().Add(Registrant);
+	static TMap<FName, UScriptStruct *(*)()> DynamicStructMap;
+	return DynamicStructMap;
+}
+
+void UObjectCompiledInDeferStruct(class UScriptStruct *(*InRegister)(), const TCHAR* PackageName, const TCHAR* Name, bool bDynamic)
+{
+	if (!bDynamic)
+	{
+		// we do reregister StaticStruct in hot reload
+		FPendingStructRegistrant Registrant(InRegister, PackageName);
+		checkSlow(!GetDeferredCompiledInStructRegistration().Contains(Registrant));
+		GetDeferredCompiledInStructRegistration().Add(Registrant);
+	}
+	else
+	{
+		GetDynamicStructMap().Add(FName(Name), InRegister);
+	}
 }
 
 #if WITH_HOT_RELOAD
@@ -485,12 +498,25 @@ static TArray<FPendingEnumRegistrant>& GetDeferredCompiledInEnumRegistration()
 	return DeferredCompiledInRegistration;
 }
 
-void UObjectCompiledInDeferEnum(class UEnum *(*InRegister)(), const TCHAR* PackageName)
+TMap<FName, UEnum *(*)()> GetDynamicEnumMap()
+{
+	static TMap<FName, UEnum *(*)()> DynamicEnumMap;
+	return DynamicEnumMap;
+}
+
+void UObjectCompiledInDeferEnum(class UEnum *(*InRegister)(), const TCHAR* PackageName, const TCHAR* Name, bool bDynamic)
 {
-	// we do reregister StaticStruct in hot reload
-	FPendingEnumRegistrant Registrant(InRegister, PackageName);
-	checkSlow(!GetDeferredCompiledInEnumRegistration().Contains(Registrant));
-	GetDeferredCompiledInEnumRegistration().Add(Registrant);
+	if (!bDynamic)
+	{
+		// we do reregister StaticStruct in hot reload
+		FPendingEnumRegistrant Registrant(InRegister, PackageName);
+		checkSlow(!GetDeferredCompiledInEnumRegistration().Contains(Registrant));
+		GetDeferredCompiledInEnumRegistration().Add(Registrant);
+	}
+	else
+	{
+		GetDynamicEnumMap().Add(FName(Name), InRegister);
+	}
 }
 
 class UEnum *GetStaticEnum(class UEnum *(*InRegister)(), UObject* EnumOuter, const TCHAR* EnumName)
@@ -631,15 +657,28 @@ void UClassCompiledInDefer(FFieldCompiledInInfo* ClassInfo, const TCHAR* Name, S
 	GetDeferredClassRegistration().Add(ClassInfo);
 }
 
-void UObjectCompiledInDefer(class UClass *(*InRegister)(), const TCHAR* Name)
+TMap<FName, UClass *(*)()> GetDynamicClassMap()
+{
+	static TMap<FName, UClass *(*)()> DynamicClassMap;
+	return DynamicClassMap;
+}
+
+void UObjectCompiledInDefer(class UClass *(*InRegister)(), const TCHAR* Name, bool bDynamic)
 {
+	if (!bDynamic)
+	{
 #if WITH_HOT_RELOAD
-	// Either add all classes if not hot-reloading, or those which have changed
-	if (!GIsHotReload || GetDeferRegisterClassMap().FindChecked(Name)->bHasChanged)
+		// Either add all classes if not hot-reloading, or those which have changed
+		if (!GIsHotReload || GetDeferRegisterClassMap().FindChecked(Name)->bHasChanged)
 #endif
+		{
+			checkSlow(!GetDeferredCompiledInRegistration().Contains(InRegister));
+			GetDeferredCompiledInRegistration().Add(InRegister);
+		}
+	}
+	else
 	{
-		checkSlow(!GetDeferredCompiledInRegistration().Contains(InRegister));
-		GetDeferredCompiledInRegistration().Add(InRegister);
+		GetDynamicClassMap().Add(FName(Name), InRegister);
 	}
 }
 
@@ -982,4 +1021,56 @@ const TCHAR* DebugFullName(UObject* Object)
 	{
 		return TEXT("None");
 	}
+}
+
+UClass* ConstructDynamicClass(FName ClassName)
+{
+	UClass* Result = nullptr;
+	UClass *(**StaticClassFNPtr)() = GetDynamicClassMap().Find(ClassName);
+	if (StaticClassFNPtr)
+	{
+		Result = (*StaticClassFNPtr)();
+	}
+	return Result;
+}
+
+UScriptStruct* ConstructDynamicStruct(FName StructName)
+{
+	UScriptStruct* Result = nullptr;
+	UScriptStruct *(**StaticStructFNPtr)() = GetDynamicStructMap().Find(StructName);
+	if (StaticStructFNPtr)
+	{
+		Result = (*StaticStructFNPtr)();
+	}
+	return Result;
+}
+
+UEnum* ConstructDynamicEnum(FName EnumName)
+{
+	UEnum* Result = nullptr;
+	UEnum *(**StaticEnumFNPtr)() = GetDynamicEnumMap().Find(EnumName);
+	if (StaticEnumFNPtr)
+	{
+		Result = (*StaticEnumFNPtr)();
+	}
+	return Result;
+}
+
+/** Constructs a dynamic class/enum/struct given its name */
+UObject* ConstructDynamicType(FName TypeName, FName TypeClass)
+{
+	UObject* Result = nullptr;
+	if (TypeClass == UClass::StaticClass()->GetFName())
+	{
+		Result = ConstructDynamicClass(TypeName);
+	}
+	else if (TypeClass == UScriptStruct::StaticClass()->GetFName())
+	{
+		Result = ConstructDynamicStruct(TypeName);
+	}
+	else if (TypeClass == UEnum::StaticClass()->GetFName())
+	{
+		Result = ConstructDynamicEnum(TypeName);
+	}
+	return Result;
 }
\ No newline at end of file
diff --git a/Engine/Source/Runtime/CoreUObject/Public/UObject/Class.h b/Engine/Source/Runtime/CoreUObject/Public/UObject/Class.h
--- a/Engine/Source/Runtime/CoreUObject/Public/UObject/Class.h
+++ b/Engine/Source/Runtime/CoreUObject/Public/UObject/Class.h
@@ -2406,13 +2406,16 @@ COREUOBJECT_API void InitializePrivateStaticClass(
  * @param PackageName name of the package this class will be inside
  * @param Name of the class
  * @param ReturnClass reference to pointer to result. This must be PrivateStaticClass.
+ * @param RegisterNativeFunc Native function registration funcfion pointer.
+ * @param bIsDynamic true if the class can be constructed dynamically at runtime
  */
 template<class TClass>
-void GetPrivateStaticClassBody( const TCHAR* PackageName, const TCHAR* Name, UClass*& ReturnClass, void (*RegisterNativeFunc)() )
+void GetPrivateStaticClassBody( const TCHAR* PackageName, const TCHAR* Name, UClass*& ReturnClass, void (*RegisterNativeFunc)(), bool bIsDynamic = false )
 { 
 #if WITH_HOT_RELOAD
 	if (GIsHotReload)
 	{
+		check(!bIsDynamic);
 		UPackage* Package = FindPackage(NULL, PackageName);
 		if (!Package)
 		{
@@ -2448,6 +2451,16 @@ void GetPrivateStaticClassBody( const TCHAR* PackageName, const TCHAR* Name, UCl
 	}
 #endif
 
+	EObjectFlags Flags = EObjectFlags(RF_Public | RF_Standalone | RF_Transient | RF_Native);
+	if (!bIsDynamic)
+	{
+		Flags = EObjectFlags(Flags | RF_RootSet);
+	}
+	else
+	{
+		Flags = EObjectFlags(Flags | RF_Dynamic);
+	}
+
 	ReturnClass = ::new (GUObjectAllocator.AllocateUObject(sizeof(UClass),ALIGNOF(UClass),true)) 
 		UClass
 		(
@@ -2457,7 +2470,7 @@ void GetPrivateStaticClassBody( const TCHAR* PackageName, const TCHAR* Name, UCl
 		TClass::StaticClassFlags,
 		TClass::StaticClassCastFlags(),
 		TClass::StaticConfigName(),
-		EObjectFlags(RF_Public | RF_Standalone | RF_Transient | RF_Native | RF_RootSet),
+		Flags,
 		(UClass::ClassConstructorType)InternalConstructor<TClass>,
 #if WITH_HOT_RELOAD_CTORS
 		(UClass::ClassVTableHelperCtorCallerType)InternalVTableHelperCtorCaller<TClass>,
diff --git a/Engine/Source/Runtime/CoreUObject/Public/UObject/ObjectBase.h b/Engine/Source/Runtime/CoreUObject/Public/UObject/ObjectBase.h
--- a/Engine/Source/Runtime/CoreUObject/Public/UObject/ObjectBase.h
+++ b/Engine/Source/Runtime/CoreUObject/Public/UObject/ObjectBase.h
@@ -469,6 +469,7 @@ enum EObjectFlags
 	RF_Async = 0x00800000, ///< Object exists only on a different thread than the game thread.
 	RF_StrongRefOnFrame			= 0x01000000,	///< References to this object from persistent function frame are handled as strong ones.
 	RF_NoStrongReference		= 0x02000000,  ///< The object is not referenced by any strong reference. The flag is used by GC.
+	RF_Dynamic = 0x04000000, // Field Only. Dynamic field - doesn't get constructed during static initialization, can be constructed multiple times
 };
 
 	// Special all and none masks
@@ -1477,11 +1478,30 @@ public: \
 		check(Class->GetClass()); \
 		return Class; \
 	} \
-	static FCompiledInDefer Z_CompiledInDefer_UClass_##TClass(Z_Construct_UClass_##TClass, TEXT(#TClass));
+	static FCompiledInDefer Z_CompiledInDefer_UClass_##TClass(Z_Construct_UClass_##TClass, TEXT(#TClass), false);
 
 #define IMPLEMENT_CORE_INTRINSIC_CLASS(TClass, TSuperClass, InitCode) \
 	IMPLEMENT_INTRINSIC_CLASS(TClass, COREUOBJECT_API, TSuperClass, COREUOBJECT_API, InitCode)
 
+// Register a dynamic class (created at runtime, not startup).
+#define IMPLEMENT_DYNAMIC_CLASS(TClass, TClassCrc) \
+	UClass* TClass::GetPrivateStaticClass(const TCHAR* Package) \
+	{ \
+		UPackage* PrivateStaticClassOuter = CastChecked<UPackage>(StaticFindObjectFast(UPackage::StaticClass(), nullptr, Package)); \
+		UClass* PrivateStaticClass = Cast<UClass>(StaticFindObjectFast(UClass::StaticClass(), PrivateStaticClassOuter, (TCHAR*)TEXT(#TClass) + 1)); \
+		if (!PrivateStaticClass) \
+		{ \
+			/* this could be handled with templates, but we want it external to avoid code bloat */ \
+			GetPrivateStaticClassBody<TClass>(\
+			Package, \
+			(TCHAR*)TEXT(#TClass) + 1 + ((StaticClassFlags & CLASS_Deprecated) ? 11 : 0), \
+			PrivateStaticClass, \
+			StaticRegisterNatives##TClass, \
+			true \
+			); \
+		} \
+		return PrivateStaticClass; \
+	}
 
 /*-----------------------------------------------------------------------------
 	ERenameFlags.
diff --git a/Engine/Source/Runtime/CoreUObject/Public/UObject/UObjectBase.h b/Engine/Source/Runtime/CoreUObject/Public/UObject/UObjectBase.h
--- a/Engine/Source/Runtime/CoreUObject/Public/UObject/UObjectBase.h
+++ b/Engine/Source/Runtime/CoreUObject/Public/UObject/UObjectBase.h
@@ -326,26 +326,26 @@ struct TClassCompiledInDefer : public FFieldCompiledInInfo
 /**
  * Stashes the singleton function that builds a compiled in class. Later, this is executed.
  */
-COREUOBJECT_API void UObjectCompiledInDefer(class UClass *(*InRegister)(), const TCHAR* Name);
+COREUOBJECT_API void UObjectCompiledInDefer(class UClass *(*InRegister)(), const TCHAR* Name, bool bDynamic);
 
 struct FCompiledInDefer
 {
-	FCompiledInDefer(class UClass *(*InRegister)(), const TCHAR* Name)
+	FCompiledInDefer(class UClass *(*InRegister)(), const TCHAR* Name, bool bDynamic)
 	{
-		UObjectCompiledInDefer(InRegister, Name);
+		UObjectCompiledInDefer(InRegister, Name, bDynamic);
 	}
 };
 
 /**
  * Stashes the singleton function that builds a compiled in struct (StaticStruct). Later, this is executed.
  */
-COREUOBJECT_API void UObjectCompiledInDeferStruct(class UScriptStruct *(*InRegister)(), const TCHAR* PackageName);
+COREUOBJECT_API void UObjectCompiledInDeferStruct(class UScriptStruct *(*InRegister)(), const TCHAR* PackageName, const TCHAR* Name, bool bDynamic);
 
 struct FCompiledInDeferStruct
 {
-	FCompiledInDeferStruct(class UScriptStruct *(*InRegister)(), const TCHAR* PackageName)
+	FCompiledInDeferStruct(class UScriptStruct *(*InRegister)(), const TCHAR* PackageName, const TCHAR* Name, bool bDynamic)
 	{
-		UObjectCompiledInDeferStruct(InRegister, PackageName);
+		UObjectCompiledInDeferStruct(InRegister, PackageName, Name, bDynamic);
 	}
 };
 
@@ -357,13 +357,13 @@ COREUOBJECT_API class UScriptStruct *GetStaticStruct(class UScriptStruct *(*InRe
 /**
  * Stashes the singleton function that builds a compiled in enum. Later, this is executed.
  */
-COREUOBJECT_API void UObjectCompiledInDeferEnum(class UEnum *(*InRegister)(), const TCHAR* PackageName);
+COREUOBJECT_API void UObjectCompiledInDeferEnum(class UEnum *(*InRegister)(), const TCHAR* PackageName, const TCHAR* Name, bool bDynamic);
 
 struct FCompiledInDeferEnum
 {
-	FCompiledInDeferEnum(class UEnum *(*InRegister)(), const TCHAR* PackageName)
+	FCompiledInDeferEnum(class UEnum *(*InRegister)(), const TCHAR* PackageName, const TCHAR* Name, bool bDynamic)
 	{
-		UObjectCompiledInDeferEnum(InRegister, PackageName);
+		UObjectCompiledInDeferEnum(InRegister, PackageName, Name, bDynamic);
 	}
 };
 
