Code generator now supports UProperties that have a name that is different than the c++ symbol name

diff --git a/Engine/Source/Programs/UnrealHeaderTool/Private/CodeGenerator.cpp b/Engine/Source/Programs/UnrealHeaderTool/Private/CodeGenerator.cpp
--- a/Engine/Source/Programs/UnrealHeaderTool/Private/CodeGenerator.cpp
+++ b/Engine/Source/Programs/UnrealHeaderTool/Private/CodeGenerator.cpp
@@ -959,6 +959,16 @@ FString FNativeClassHeaderGenerator::GetSingletonName(FClass* Item, bool bRequir
 	return GetSingletonName((UClass*)Item, bRequiresValidObject);
 }
 
+FString FNativeClassHeaderGenerator::GetOverriddenName(const UField* Item)
+{
+	FString OverriddenName = Item->GetMetaData(TEXT("OverrideNativeName"));
+	if (!OverriddenName.IsEmpty())
+	{
+		return OverriddenName;
+	}
+	return Item->GetName();
+}
+
 FString FNativeClassHeaderGenerator::PropertyNew(FString& Meta, UProperty* Prop, const FString& OuterString, const FString& PropMacro, const TCHAR* NameSuffix, const TCHAR* Spaces, const TCHAR* SourceStruct)
 {
 	FString ExtraArgs;
@@ -1027,7 +1037,7 @@ FString FNativeClassHeaderGenerator::PropertyNew(FString& Meta, UProperty* Prop,
 
 	FString Constructor = FString::Printf(TEXT("new(EC_InternalUseOnlyConstructor, %s, TEXT(\"%s\"), RF_Public|RF_Transient|RF_Native) U%s(%s, 0x%016llx%s);"),
 		*OuterString,
-		*Prop->GetName(), 
+		*FNativeClassHeaderGenerator::GetOverriddenName(Prop),
 		*Prop->GetClass()->GetName(), 
 		*PropMacro,
 		Prop->PropertyFlags & ~CPF_ComputedFlags, 
@@ -1440,7 +1450,7 @@ void FNativeClassHeaderGenerator::ExportNativeGeneratedInitCode(FClass* Class, F
 			// Emit code to construct each UFunction and rebuild the function map at runtime
 			for (UFunction* Function : FunctionsInMap)
 			{
-				GeneratedClassRegisterFunctionText.Logf(TEXT("                OuterClass->AddFunctionToFunctionMap(%s);%s\r\n"), *GetSingletonName(Function), *GetGeneratedCodeCRCTag(Function));
+				GeneratedClassRegisterFunctionText.Logf(TEXT("                OuterClass->AddFunctionToFunctionMapWithOverriddenName(%s, \"%s\");%s\r\n"), *GetSingletonName(Function), *FNativeClassHeaderGenerator::GetOverriddenName(Function), *GetGeneratedCodeCRCTag(Function));
 			}
 		}
 
@@ -1605,7 +1615,7 @@ void FNativeClassHeaderGenerator::ExportFunction(UFunction* Function, FScope* Sc
 	FString UFunctionType = bIsDelegate ? TEXT("UDelegateFunction") : TEXT("UFunction");
 
 	CurrentFunctionText.Logf(TEXT("            ReturnFunction = new(EC_InternalUseOnlyConstructor, Outer, TEXT(\"%s\"), RF_Public|RF_Transient|RF_Native) %s(FObjectInitializer(), %s, 0x%08X, %d%s);\r\n"),
-		*Function->GetName(),
+		*FNativeClassHeaderGenerator::GetOverriddenName(Function),
 		*UFunctionType,
 		*SuperFunctionString,
 		Function->FunctionFlags,
@@ -1677,7 +1687,7 @@ void FNativeClassHeaderGenerator::ExportNatives(FClass* Class)
 		{
 			GeneratedPackageCPP.Logf(TEXT("        FNativeFunctionRegistrar::RegisterFunction(%s::StaticClass(),\"%s\",(Native)&%s::exec%s);\r\n"),
 				NameLookupCPP.GetNameCPP(Class),
-				*Func->GetName(),
+				*FNativeClassHeaderGenerator::GetOverriddenName(Func),
 				Class->HasAnyClassFlags(CLASS_Interface) ? *FString::Printf(TEXT("I%s"), *Class->GetName()) : NameLookupCPP.GetNameCPP(Class),
 				*Func->GetName()
 			);
diff --git a/Engine/Source/Programs/UnrealHeaderTool/Private/NativeClassExporter.h b/Engine/Source/Programs/UnrealHeaderTool/Private/NativeClassExporter.h
--- a/Engine/Source/Programs/UnrealHeaderTool/Private/NativeClassExporter.h
+++ b/Engine/Source/Programs/UnrealHeaderTool/Private/NativeClassExporter.h
@@ -187,6 +187,12 @@ private:
 	/** Return the name of the singleton function that returns the UObject for Item */
 	FString GetSingletonName(FClass* Item, bool bRequiresValidObject=true);
 
+	/** 
+	 * Properties in source files generated from blueprint assets have a symbol name that differs from the source asset.
+	 * This function returns the original name of the field (rather than the native, symbol name).
+	 */
+	FString GetOverriddenName(const UField* Item);
+
 	/** 
 	 * Export functions used to find and call C++ or script implementation of a script function in the interface 
 	 */
