Fix for running Z_ConstructUClass* for dynamic classes while constructing the UClass object.

diff --git a/Engine/Source/Programs/UnrealHeaderTool/Private/CodeGenerator.cpp b/Engine/Source/Programs/UnrealHeaderTool/Private/CodeGenerator.cpp
--- a/Engine/Source/Programs/UnrealHeaderTool/Private/CodeGenerator.cpp
+++ b/Engine/Source/Programs/UnrealHeaderTool/Private/CodeGenerator.cpp
@@ -1395,13 +1395,15 @@ void FNativeClassHeaderGenerator::ExportNativeGeneratedInitCode(FClass* Class, F
 		if (!bIsDynamic)
 		{
 			GeneratedClassRegisterFunctionText.Logf(TEXT("        static UClass* OuterClass = NULL;\r\n"));
+			GeneratedClassRegisterFunctionText.Logf(TEXT("        if (!OuterClass)\r\n"));
 		}
 		else
 		{
 			GeneratedClassRegisterFunctionText.Logf(TEXT("        UClass* OuterClass = Cast<UClass>(StaticFindObjectFast(UClass::StaticClass(), %s, TEXT(\"%s\")));\r\n"),
 				*GetPackageSingletonName(CastChecked<UPackage>(Class->GetOutermost())), *Class->GetName());
+			GeneratedClassRegisterFunctionText.Logf(TEXT("        if (!OuterClass || !(OuterClass->ClassFlags & CLASS_Constructed))\r\n"));
 		}
-		GeneratedClassRegisterFunctionText.Logf(TEXT("        if (!OuterClass)\r\n"));
+		
 		GeneratedClassRegisterFunctionText.Logf(TEXT("        {\r\n"));
 		if (Class->GetSuperClass() && Class->GetSuperClass() != Class)
 		{
