Merging using Orion->UE4

CL# 2693983
Skip inline keywords/macros when parsing a return type.

diff --git a/Engine/Source/Programs/UnrealHeaderTool/Private/HeaderParser.cpp b/Engine/Source/Programs/UnrealHeaderTool/Private/HeaderParser.cpp
--- a/Engine/Source/Programs/UnrealHeaderTool/Private/HeaderParser.cpp
+++ b/Engine/Source/Programs/UnrealHeaderTool/Private/HeaderParser.cpp
@@ -3055,6 +3055,23 @@ FIndexRange*                    ParsedVarIndexRange
 		FError::Throwf(TEXT("Unknown access level"));
 	}
 
+	// Swallow inline keywords
+	if (VariableCategory == EVariableCategory::Return)
+	{
+		FToken InlineToken;
+		if (!GetIdentifier(InlineToken, true))
+		{
+			FError::Throwf(TEXT("%s: Missing variable type"), GetHintText(VariableCategory));
+		}
+
+		if (FCString::Strcmp(InlineToken.Identifier, TEXT("inline")) != 0
+			&& FCString::Strcmp(InlineToken.Identifier, TEXT("FORCENOINLINE")) != 0
+			&& FCString::Strncmp(InlineToken.Identifier, TEXT("FORCEINLINE"), 11) != 0)
+		{
+			UngetToken(InlineToken);
+		}
+	}
+
 	// Get variable type.
 	bool bUnconsumedStructKeyword = false;
 	bool bUnconsumedClassKeyword  = false;
diff --git a/Engine/Source/Programs/UnrealHeaderTool/Resources/UHTDebugging/TestObject.h b/Engine/Source/Programs/UnrealHeaderTool/Resources/UHTDebugging/TestObject.h
--- a/Engine/Source/Programs/UnrealHeaderTool/Resources/UHTDebugging/TestObject.h
+++ b/Engine/Source/Programs/UnrealHeaderTool/Resources/UHTDebugging/TestObject.h
@@ -31,7 +31,25 @@ public:
 	TArray<UClass*> ReturnArrayOfUClassPtrs();
 
 	UFUNCTION()
-	int32 InlineFunc()
+	inline int32 InlineFunc1()
+	{
+		return FString("Hello").Len();
+	}
+
+	UFUNCTION()
+	FORCEINLINE int32 InlineFunc2()
+	{
+		return FString("Hello").Len();
+	}
+
+	UFUNCTION()
+	FORCEINLINE_WHATEVER int32 InlineFunc3()
+	{
+		return FString("Hello").Len();
+	}
+
+	UFUNCTION()
+	FORCENOINLINE int32 NoInlineFunc()
 	{
 		return FString("Hello").Len();
 	}
