ModuleRelativePath meta data in UInterface.

diff --git a/Engine/Source/Programs/UnrealHeaderTool/Private/HeaderParser.cpp b/Engine/Source/Programs/UnrealHeaderTool/Private/HeaderParser.cpp
--- a/Engine/Source/Programs/UnrealHeaderTool/Private/HeaderParser.cpp
+++ b/Engine/Source/Programs/UnrealHeaderTool/Private/HeaderParser.cpp
@@ -5056,6 +5056,7 @@ void FHeaderParser::CompileInterfaceDeclaration(FClasses& AllClasses)
 	ClassData->SetPrologLine(PrologFinishLine);
 
 	// Register the metadata
+	AddModuleRelativePathToMetadata(InterfaceClass, MetaData);
 	AddMetaDataToClassData(InterfaceClass, MetaData);
 
 	// Handle the start of the rest of the interface
