Merging using Orion->UE4

CL# 2698724:
UHT speed improvements:

Code is now generated with tabs, rather than going through a separate Tabify step.
Some expensive temporaries created during iteration and during iteration have been removed.
FFunctionData::FunctionDataMap now uses unique ownership instead of shared ownership.
TCString::Spc (and the new TCString::Tab) now uses a compile-time array instead of one generated at runtime with a static.
TSet::Compact no longer rehashes if no compaction was performed.
TArray assignment no longer reallocs to the same capacity.

diff --git a/Engine/Source/Programs/UnrealHeaderTool/Private/CodeGenerator.cpp b/Engine/Source/Programs/UnrealHeaderTool/Private/CodeGenerator.cpp
--- a/Engine/Source/Programs/UnrealHeaderTool/Private/CodeGenerator.cpp
+++ b/Engine/Source/Programs/UnrealHeaderTool/Private/CodeGenerator.cpp
@@ -19,7 +19,6 @@
 FManifest GManifest;
 
 double GMacroizeTime = 0.0;
-double GTabifyTime = 0.0;
 
 static TArray<FString> ChangeMessages;
 static bool bWriteContents = false;
@@ -35,17 +34,6 @@ FCompilerMetadataManager GScriptHelper;
 /** C++ name lookup helper */
 FNameLookupCPP NameLookupCPP;
 
-static const int32 MaxNumTabs = 64;
-static const TCHAR TabString[MaxNumTabs + 1] = TEXT("\t\t\t\t\t\t\t\t\t\t\t\t\t\t\t\t\t\t\t\t\t\t\t\t\t\t\t\t\t\t\t\t\t\t\t\t\t\t\t\t\t\t\t\t\t\t\t\t\t\t\t\t\t\t\t\t\t\t\t\t\t\t\t\t");
-
-const TCHAR* GetTabs(int32 NumTabs)
-{
-	checkSlow((NumTabs >= 0) && (NumTabs < MaxNumTabs));
-
-	const uint32 TabPos = MaxNumTabs - NumTabs;
-	return TabString + TabPos;
-}
-
 /**
  * Finds exact match of Identifier in string. Returns nullptr if none is found.
  *
@@ -242,75 +230,6 @@ static struct FFlagAudit
 	}
 } TheFlagAudit;
 
-static FString Tabify(const TCHAR* Input)
-{
-	FScopedDurationTimer Tracker(GTabifyTime);
-
-	FString Result(Input); 
-	const int32 SpacesPerTab = 4;
-	Result.ReplaceInline(TEXT("\r\n"), TEXT("\n"), ESearchCase::CaseSensitive);
-
-	TArray<FString> Lines;
-	Result.ParseIntoArray(Lines, TEXT("\n"), false);
-
-	Result.Reset();
-
-	for (int32 Index = 0; Index < Lines.Num(); Index++)
-	{
-		if (Index)
-		{
-			Result += TEXT("\r\n");
-		}
-
-		FString& Line = Lines[Index];
-
-		int32 NumSpaces = 0;
-		int32 FirstNonWS = 0;
-		int32 NumTrueTabs = 0;
-		while (1)
-		{
-			TCHAR c = *(*Line + FirstNonWS);
-
-			if (c == TEXT('\t'))
-			{
-				++NumTrueTabs;
-				NumSpaces += SpacesPerTab;
-			}
-			else if (c == TEXT(' '))
-			{
-				++NumSpaces;
-			}
-			else
-			{
-				break;
-			}
-			FirstNonWS++;
-		}
-
-		const int32 NumTabs = NumSpaces / SpacesPerTab;
-		const int32 LeftoverSpaces = NumSpaces % SpacesPerTab;
-		if ((NumTrueTabs == NumTabs) && (LeftoverSpaces == 0))
-		{
-			Result += Line;
-		}
-		else
-		{
-			FString Remainder = *Line + FirstNonWS;
-			for (int32 Tab = 0; Tab < NumTabs; Tab++)
-			{
-				Result += TEXT("\t");
-			}
-			if (LeftoverSpaces > 0)
-			{
-				Result += FCString::Spc(LeftoverSpaces);
-			}
-			Result += Remainder;
-		}
-	}
-
-	return Result;
-}
-
 void ConvertToBuildIncludePath(UPackage* Package, FString& LocalPath)
 {
 	FPaths::MakePathRelativeTo(LocalPath, *GPackageToManifestModuleMap.FindChecked(Package)->IncludeBase);
@@ -698,7 +617,7 @@ void FNativeClassHeaderGenerator::ExportProperties(UStruct* Struct, int32 TextIn
 				ApplyAlternatePropertyExportText(*It, JustPropertyDecl);
 
 				// Finish up line.
-				PropertyText.Logf(TEXT("%s%s;\r\n"), GetTabs(TextIndent + 1), *JustPropertyDecl);
+				PropertyText.Logf(TEXT("%s%s;\r\n"), FCString::Tab(TextIndent + 1), *JustPropertyDecl);
 			}
 
 			LastInSuper	= NULL;
@@ -771,19 +690,19 @@ private:
 		
 		if (InType->GetClass() == UClass::StaticClass())
 		{
-			FormatStr = TEXT("    %s_API class UClass* %s;\r\n");
+			FormatStr = TEXT("\t%s_API class UClass* %s;\r\n");
 		}
 		else if (InType->GetClass() == UFunction::StaticClass() || InType->GetClass() == UDelegateFunction::StaticClass())
 		{
-			FormatStr = TEXT("    %s_API class UFunction* %s;\r\n");
+			FormatStr = TEXT("\t%s_API class UFunction* %s;\r\n");
 		}
 		else if (InType->GetClass() == UScriptStruct::StaticClass())
 		{
-			FormatStr = TEXT("    %s_API class UScriptStruct* %s;\r\n");
+			FormatStr = TEXT("\t%s_API class UScriptStruct* %s;\r\n");
 		}
 		else if (InType->GetClass() == UEnum::StaticClass())
 		{
-			FormatStr = TEXT("    %s_API class UEnum* %s;\r\n");
+			FormatStr = TEXT("\t%s_API class UEnum* %s;\r\n");
 		}
 		else
 		{
@@ -1175,8 +1094,8 @@ void FNativeClassHeaderGenerator::OutputProperty(FString& Meta, FOutputDevice& O
 
 static FString BackSlashAndIndent(const FString& Text)
 {
-	FString Trailer(TEXT("    \\\r\n    "));
-	FString Result = FString(TEXT("    ")) + Text.Replace(TEXT("\r\n"), *Trailer, ESearchCase::CaseSensitive);
+	FString Trailer(TEXT("\t\\\r\n\t"));
+	FString Result = FString(TEXT("\t")) + Text.Replace(TEXT("\r\n"), *Trailer, ESearchCase::CaseSensitive);
 	return Result.LeftChop(Trailer.Len()) + TEXT("\r\n");
 }
 
@@ -1261,25 +1180,25 @@ void FNativeClassHeaderGenerator::ExportGeneratedPackageInitCode(UPackage* InPac
 
 	FUHTStringBuilder& GeneratedFunctionText = GetGeneratedFunctionTextDevice();
 
-	GeneratedFunctionDeclarations.Logf(TEXT("    %sclass UPackage* %s;\r\n"), *ApiString, *SingletonName);
+	GeneratedFunctionDeclarations.Logf(TEXT("\t%sclass UPackage* %s;\r\n"), *ApiString, *SingletonName);
 
-	GeneratedFunctionText.Logf(TEXT("    UPackage* %s\r\n"), *SingletonName);
-	GeneratedFunctionText.Logf(TEXT("    {\r\n"));
-	GeneratedFunctionText.Logf(TEXT("        static UPackage* ReturnPackage = NULL;\r\n"));
-	GeneratedFunctionText.Logf(TEXT("        if (!ReturnPackage)\r\n"));
-	GeneratedFunctionText.Logf(TEXT("        {\r\n"));
-	GeneratedFunctionText.Logf(TEXT("            ReturnPackage = CastChecked<UPackage>(StaticFindObjectFast(UPackage::StaticClass(), NULL, FName(TEXT(\"%s\")), false, false));\r\n"), *InPackage->GetName());
+	GeneratedFunctionText.Logf(TEXT("\tUPackage* %s\r\n"), *SingletonName);
+	GeneratedFunctionText.Logf(TEXT("\t{\r\n"));
+	GeneratedFunctionText.Logf(TEXT("\t\tstatic UPackage* ReturnPackage = NULL;\r\n"));
+	GeneratedFunctionText.Logf(TEXT("\t\tif (!ReturnPackage)\r\n"));
+	GeneratedFunctionText.Logf(TEXT("\t\t{\r\n"));
+	GeneratedFunctionText.Logf(TEXT("\t\t\tReturnPackage = CastChecked<UPackage>(StaticFindObjectFast(UPackage::StaticClass(), NULL, FName(TEXT(\"%s\")), false, false));\r\n"), *InPackage->GetName());
 
-	FString Meta = GetMetaDataCodeForObject(InPackage, TEXT("ReturnPackage"), TEXT("            "));
+	FString Meta = GetMetaDataCodeForObject(InPackage, TEXT("ReturnPackage"), TEXT("\t\t\t"));
 	if (Meta.Len())
 	{
 		GeneratedFunctionText.Logf(TEXT("#if WITH_METADATA\r\n"));
-		GeneratedFunctionText.Logf(TEXT("            UMetaData* MetaData = ReturnPackage->GetMetaData();\r\n"));
+		GeneratedFunctionText.Logf(TEXT("\t\t\tUMetaData* MetaData = ReturnPackage->GetMetaData();\r\n"));
 		GeneratedFunctionText.Log(*Meta);
 		GeneratedFunctionText.Logf(TEXT("#endif\r\n"));
 	}
 
-	GeneratedFunctionText.Logf(TEXT("            ReturnPackage->PackageFlags |= PKG_CompiledIn | 0x%08X;\r\n"), InPackage->PackageFlags & (PKG_ClientOptional | PKG_ServerSideOnly | PKG_EditorOnly));
+	GeneratedFunctionText.Logf(TEXT("\t\t\tReturnPackage->PackageFlags |= PKG_CompiledIn | 0x%08X;\r\n"), InPackage->PackageFlags & (PKG_ClientOptional | PKG_ServerSideOnly | PKG_EditorOnly));
 	TheFlagAudit.Add(InPackage, TEXT("PackageFlags"), InPackage->PackageFlags);
 	{
 		FGuid Guid;
@@ -1301,12 +1220,12 @@ void FNativeClassHeaderGenerator::ExportGeneratedPackageInitCode(UPackage* InPac
 
 		Guid.A = CombinedCRC;
 		Guid.B = GenerateTextCRC(*GeneratedFunctionDeclarations.ToUpper());
-		GeneratedFunctionText.Logf(TEXT("            FGuid Guid;\r\n"));
-		GeneratedFunctionText.Logf(TEXT("            Guid.A = 0x%08X;\r\n"), Guid.A);
-		GeneratedFunctionText.Logf(TEXT("            Guid.B = 0x%08X;\r\n"), Guid.B);
-		GeneratedFunctionText.Logf(TEXT("            Guid.C = 0x%08X;\r\n"), Guid.C);
-		GeneratedFunctionText.Logf(TEXT("            Guid.D = 0x%08X;\r\n"), Guid.D);
-		GeneratedFunctionText.Logf(TEXT("            ReturnPackage->SetGuid(Guid);\r\n"), Guid.D);
+		GeneratedFunctionText.Logf(TEXT("\t\t\tFGuid Guid;\r\n"));
+		GeneratedFunctionText.Logf(TEXT("\t\t\tGuid.A = 0x%08X;\r\n"), Guid.A);
+		GeneratedFunctionText.Logf(TEXT("\t\t\tGuid.B = 0x%08X;\r\n"), Guid.B);
+		GeneratedFunctionText.Logf(TEXT("\t\t\tGuid.C = 0x%08X;\r\n"), Guid.C);
+		GeneratedFunctionText.Logf(TEXT("\t\t\tGuid.D = 0x%08X;\r\n"), Guid.D);
+		GeneratedFunctionText.Logf(TEXT("\t\t\tReturnPackage->SetGuid(Guid);\r\n"), Guid.D);
 
 		GeneratedFunctionText.Log(TEXT("\r\n"));
 
@@ -1325,15 +1244,15 @@ void FNativeClassHeaderGenerator::ExportGeneratedPackageInitCode(UPackage* InPac
 				UField* FieldOuter = Cast<UField>(ScriptType->GetOuter());
 				if (!FieldOuter || !FClass::IsDynamic(FieldOuter))
 				{
-					GeneratedFunctionText.Logf(TEXT("            %s;\r\n"), *GetSingletonName(ScriptType));
+					GeneratedFunctionText.Logf(TEXT("\t\t\t%s;\r\n"), *GetSingletonName(ScriptType));
 				}
 			}
 		}
 	}
 
-	GeneratedFunctionText.Logf(TEXT("        }\r\n"));
-	GeneratedFunctionText.Logf(TEXT("        return ReturnPackage;\r\n"));
-	GeneratedFunctionText.Logf(TEXT("    }\r\n"));
+	GeneratedFunctionText.Logf(TEXT("\t\t}\r\n"));
+	GeneratedFunctionText.Logf(TEXT("\t\treturn ReturnPackage;\r\n"));
+	GeneratedFunctionText.Logf(TEXT("\t}\r\n"));
 
 }
 
@@ -1367,7 +1286,7 @@ void FNativeClassHeaderGenerator::ExportNativeGeneratedInitCode(FClass* Class, F
 			ExportFunction(Function, &FScope::GetTypeScope(Class).Get(), bIsNoExport);
 		}
 
-		CallSingletons.Logf(TEXT("                OuterClass->LinkChild(%s);\r\n"), *GetSingletonName(Function));
+		CallSingletons.Logf(TEXT("\t\t\t\tOuterClass->LinkChild(%s);\r\n"), *GetSingletonName(Function));
 	}
 
 	FUHTStringBuilder GeneratedClassRegisterFunctionText;
@@ -1380,49 +1299,49 @@ void FNativeClassHeaderGenerator::ExportNativeGeneratedInitCode(FClass* Class, F
 
 			GeneratedFunctionDeclarations.Log(FTypeSingletonCache::Get(Class, false).GetExternDecl());
 
-			GeneratedClassRegisterFunctionText.Logf(TEXT("    UClass* %s\r\n"), *SingletonNameNoRegister);
-			GeneratedClassRegisterFunctionText.Logf(TEXT("    {\r\n"));
-			GeneratedClassRegisterFunctionText.Logf(TEXT("        return %s::StaticClass();\r\n"), NameLookupCPP.GetNameCPP(Class));
-			GeneratedClassRegisterFunctionText.Logf(TEXT("    }\r\n"));
+			GeneratedClassRegisterFunctionText.Logf(TEXT("\tUClass* %s\r\n"), *SingletonNameNoRegister);
+			GeneratedClassRegisterFunctionText.Logf(TEXT("\t{\r\n"));
+			GeneratedClassRegisterFunctionText.Logf(TEXT("\t\treturn %s::StaticClass();\r\n"), NameLookupCPP.GetNameCPP(Class));
+			GeneratedClassRegisterFunctionText.Logf(TEXT("\t}\r\n"));
 		}
 		FString SingletonName(GetSingletonName(Class));
 
-		OutFriendText.Logf(TEXT("    friend %sclass UClass* %s;\r\n"), *ApiString, *SingletonName);
+		OutFriendText.Logf(TEXT("\tfriend %sclass UClass* %s;\r\n"), *ApiString, *SingletonName);
 		GeneratedFunctionDeclarations.Log(FTypeSingletonCache::Get(Class).GetExternDecl());
 
-		GeneratedClassRegisterFunctionText.Logf(TEXT("    UClass* %s\r\n"), *SingletonName);
-		GeneratedClassRegisterFunctionText.Logf(TEXT("    {\r\n"));
+		GeneratedClassRegisterFunctionText.Logf(TEXT("\tUClass* %s\r\n"), *SingletonName);
+		GeneratedClassRegisterFunctionText.Logf(TEXT("\t{\r\n"));
 		if (!bIsDynamic)
 		{
-			GeneratedClassRegisterFunctionText.Logf(TEXT("        static UClass* OuterClass = NULL;\r\n"));
-			GeneratedClassRegisterFunctionText.Logf(TEXT("        if (!OuterClass)\r\n"));
+			GeneratedClassRegisterFunctionText.Logf(TEXT("\t\tstatic UClass* OuterClass = NULL;\r\n"));
+			GeneratedClassRegisterFunctionText.Logf(TEXT("\t\tif (!OuterClass)\r\n"));
 		}
 		else
 		{
-			GeneratedClassRegisterFunctionText.Logf(TEXT("        UClass* OuterClass = Cast<UClass>(StaticFindObjectFast(UClass::StaticClass(), %s, TEXT(\"%s\")));\r\n"),
+			GeneratedClassRegisterFunctionText.Logf(TEXT("\t\tUClass* OuterClass = Cast<UClass>(StaticFindObjectFast(UClass::StaticClass(), %s, TEXT(\"%s\")));\r\n"),
 				*GetPackageSingletonName(CastChecked<UPackage>(Class->GetOutermost())), *Class->GetName());
-			GeneratedClassRegisterFunctionText.Logf(TEXT("        if (!OuterClass || !(OuterClass->ClassFlags & CLASS_Constructed))\r\n"));
+			GeneratedClassRegisterFunctionText.Logf(TEXT("\t\tif (!OuterClass || !(OuterClass->ClassFlags & CLASS_Constructed))\r\n"));
 		}
 		
-		GeneratedClassRegisterFunctionText.Logf(TEXT("        {\r\n"));
+		GeneratedClassRegisterFunctionText.Logf(TEXT("\t\t{\r\n"));
 		if (Class->GetSuperClass() && Class->GetSuperClass() != Class)
 		{
-			GeneratedClassRegisterFunctionText.Logf(TEXT("            %s;\r\n"), *GetSingletonName(Class->GetSuperClass()));
+			GeneratedClassRegisterFunctionText.Logf(TEXT("\t\t\t%s;\r\n"), *GetSingletonName(Class->GetSuperClass()));
 		}
-		GeneratedClassRegisterFunctionText.Logf(TEXT("            %s;\r\n"), *GetPackageSingletonName(CastChecked<UPackage>(Class->GetOutermost())));
-		GeneratedClassRegisterFunctionText.Logf(TEXT("            OuterClass = %s::StaticClass();\r\n"), NameLookupCPP.GetNameCPP(Class));
-		GeneratedClassRegisterFunctionText.Logf(TEXT("            if (!(OuterClass->ClassFlags & CLASS_Constructed))\r\n"), NameLookupCPP.GetNameCPP(Class));
-		GeneratedClassRegisterFunctionText.Logf(TEXT("            {\r\n"), NameLookupCPP.GetNameCPP(Class));
-		GeneratedClassRegisterFunctionText.Logf(TEXT("                UObjectForceRegistration(OuterClass);\r\n"));
+		GeneratedClassRegisterFunctionText.Logf(TEXT("\t\t\t%s;\r\n"), *GetPackageSingletonName(CastChecked<UPackage>(Class->GetOutermost())));
+		GeneratedClassRegisterFunctionText.Logf(TEXT("\t\t\tOuterClass = %s::StaticClass();\r\n"), NameLookupCPP.GetNameCPP(Class));
+		GeneratedClassRegisterFunctionText.Logf(TEXT("\t\t\tif (!(OuterClass->ClassFlags & CLASS_Constructed))\r\n"), NameLookupCPP.GetNameCPP(Class));
+		GeneratedClassRegisterFunctionText.Logf(TEXT("\t\t\t{\r\n"), NameLookupCPP.GetNameCPP(Class));
+		GeneratedClassRegisterFunctionText.Logf(TEXT("\t\t\t\tUObjectForceRegistration(OuterClass);\r\n"));
 		uint32 Flags = (Class->ClassFlags & CLASS_SaveInCompiledInClasses) | CLASS_Constructed;
-		GeneratedClassRegisterFunctionText.Logf(TEXT("                OuterClass->ClassFlags |= 0x%08X;\r\n"), Flags);
+		GeneratedClassRegisterFunctionText.Logf(TEXT("\t\t\t\tOuterClass->ClassFlags |= 0x%08X;\r\n"), Flags);
 		TheFlagAudit.Add(Class, TEXT("ClassFlags"), Flags);
 		GeneratedClassRegisterFunctionText.Logf(TEXT("\r\n"));
 		GeneratedClassRegisterFunctionText.Log(CallSingletons);
 		GeneratedClassRegisterFunctionText.Logf(TEXT("\r\n"));
 
 		FString OuterString = FString(TEXT("OuterClass"));
-		FString Meta = GetMetaDataCodeForObject(Class, *OuterString, TEXT("                "));
+		FString Meta = GetMetaDataCodeForObject(Class, *OuterString, TEXT("\t\t\t\t"));
 		// properties
 		{
 			TArray<UProperty*> Props;
@@ -1434,7 +1353,7 @@ void FNativeClassHeaderGenerator::ExportNativeGeneratedInitCode(FClass* Class, F
 			if (Props.Num() > 0)
 			{
 				GeneratedClassRegisterFunctionText.Logf(TEXT("PRAGMA_DISABLE_DEPRECATION_WARNINGS\r\n"));
-				OutputProperties(Meta, GeneratedClassRegisterFunctionText, OuterString, Props, TEXT("                "));
+				OutputProperties(Meta, GeneratedClassRegisterFunctionText, OuterString, Props, TEXT("\t\t\t\t"));
 				GeneratedClassRegisterFunctionText.Logf(TEXT("PRAGMA_ENABLE_DEPRECATION_WARNINGS\r\n"));
 			}
 		}
@@ -1452,15 +1371,15 @@ void FNativeClassHeaderGenerator::ExportNativeGeneratedInitCode(FClass* Class, F
 			// Emit code to construct each UFunction and rebuild the function map at runtime
 			for (UFunction* Function : FunctionsInMap)
 			{
-				GeneratedClassRegisterFunctionText.Logf(TEXT("                OuterClass->AddFunctionToFunctionMapWithOverriddenName(%s, \"%s\");%s\r\n"), *GetSingletonName(Function), *FNativeClassHeaderGenerator::GetOverriddenName(Function), *GetGeneratedCodeCRCTag(Function));
+				GeneratedClassRegisterFunctionText.Logf(TEXT("\t\t\t\tOuterClass->AddFunctionToFunctionMapWithOverriddenName(%s, \"%s\");%s\r\n"), *GetSingletonName(Function), *FNativeClassHeaderGenerator::GetOverriddenName(Function), *GetGeneratedCodeCRCTag(Function));
 			}
 		}
 
 		// class flags are handled by the intrinsic bootstrap code
-		//GeneratedClassRegisterFunctionText.Logf(TEXT("            OuterClass->ClassFlags = 0x%08X;\r\n"), Class->ClassFlags);
+		//GeneratedClassRegisterFunctionText.Logf(TEXT("\t\t\tOuterClass->ClassFlags = 0x%08X;\r\n"), Class->ClassFlags);
 		if (Class->ClassConfigName != NAME_None)
 		{
-			GeneratedClassRegisterFunctionText.Logf(TEXT("                OuterClass->ClassConfigName = FName(TEXT(\"%s\"));\r\n"), *Class->ClassConfigName.ToString());
+			GeneratedClassRegisterFunctionText.Logf(TEXT("\t\t\t\tOuterClass->ClassConfigName = FName(TEXT(\"%s\"));\r\n"), *Class->ClassConfigName.ToString());
 		}
 
 		for (auto& Inter : Class->Interfaces)
@@ -1471,7 +1390,7 @@ void FNativeClassHeaderGenerator::ExportNativeGeneratedInitCode(FClass* Class, F
 			{
 				OffsetString = FString::Printf(TEXT("VTABLE_OFFSET(%s, %s)"), NameLookupCPP.GetNameCPP(Class), NameLookupCPP.GetNameCPP(Inter.Class, true));
 			}
-			GeneratedClassRegisterFunctionText.Logf(TEXT("                OuterClass->Interfaces.Add(FImplementedInterface(%s, %s, %s ));\r\n"), 
+			GeneratedClassRegisterFunctionText.Logf(TEXT("\t\t\t\tOuterClass->Interfaces.Add(FImplementedInterface(%s, %s, %s ));\r\n"), 
 				*GetSingletonName(Inter.Class, false),
 				*OffsetString,
 				Inter.bImplementedByK2 ? TEXT("true") : TEXT("false")
@@ -1480,33 +1399,33 @@ void FNativeClassHeaderGenerator::ExportNativeGeneratedInitCode(FClass* Class, F
 		if (Class->ClassGeneratedBy)
 		{
 			UE_LOG(LogCompile, Fatal, TEXT("For intrinsic and compiled-in classes, ClassGeneratedBy should always be NULL"));
-			GeneratedClassRegisterFunctionText.Logf(TEXT("                OuterClass->ClassGeneratedBy = %s;\r\n"), *GetSingletonName(CastChecked<UClass>(Class->ClassGeneratedBy), false));
+			GeneratedClassRegisterFunctionText.Logf(TEXT("\t\t\t\tOuterClass->ClassGeneratedBy = %s;\r\n"), *GetSingletonName(CastChecked<UClass>(Class->ClassGeneratedBy), false));
 		}
 
-		GeneratedClassRegisterFunctionText.Logf(TEXT("                OuterClass->StaticLink();\r\n"));
+		GeneratedClassRegisterFunctionText.Logf(TEXT("\t\t\t\tOuterClass->StaticLink();\r\n"));
 
 		if (Meta.Len())
 		{
 			GeneratedClassRegisterFunctionText.Logf(TEXT("#if WITH_METADATA\r\n"));
-			GeneratedClassRegisterFunctionText.Logf(TEXT("                UMetaData* MetaData = OuterClass->GetOutermost()->GetMetaData();\r\n"));
+			GeneratedClassRegisterFunctionText.Logf(TEXT("\t\t\t\tUMetaData* MetaData = OuterClass->GetOutermost()->GetMetaData();\r\n"));
 			GeneratedClassRegisterFunctionText.Log(*Meta);
 			GeneratedClassRegisterFunctionText.Logf(TEXT("#endif\r\n"));
 		}
 
-		GeneratedClassRegisterFunctionText.Logf(TEXT("            }\r\n"));
-		GeneratedClassRegisterFunctionText.Logf(TEXT("        }\r\n"));
-		GeneratedClassRegisterFunctionText.Logf(TEXT("        check(OuterClass->GetClass());\r\n"));
-		GeneratedClassRegisterFunctionText.Logf(TEXT("        return OuterClass;\r\n"));
-		GeneratedClassRegisterFunctionText.Logf(TEXT("    }\r\n"));
+		GeneratedClassRegisterFunctionText.Logf(TEXT("\t\t\t}\r\n"));
+		GeneratedClassRegisterFunctionText.Logf(TEXT("\t\t}\r\n"));
+		GeneratedClassRegisterFunctionText.Logf(TEXT("\t\tcheck(OuterClass->GetClass());\r\n"));
+		GeneratedClassRegisterFunctionText.Logf(TEXT("\t\treturn OuterClass;\r\n"));
+		GeneratedClassRegisterFunctionText.Logf(TEXT("\t}\r\n"));
 
 		GeneratedFunctionText += GeneratedClassRegisterFunctionText;
 	}
 
 	if (OutFriendText.Len() && bIsNoExport)
 	{
-		GeneratedFunctionText.Logf(TEXT("    /* friend declarations for pasting into noexport class %s\r\n"), NameLookupCPP.GetNameCPP(Class));
+		GeneratedFunctionText.Logf(TEXT("\t/* friend declarations for pasting into noexport class %s\r\n"), NameLookupCPP.GetNameCPP(Class));
 		GeneratedFunctionText.Log(OutFriendText);
-		GeneratedFunctionText.Logf(TEXT("    */\r\n"));
+		GeneratedFunctionText.Logf(TEXT("\t*/\r\n"));
 		OutFriendText.Reset();
 	}
 	FString SingletonName(GetSingletonName(Class));
@@ -1515,7 +1434,7 @@ void FNativeClassHeaderGenerator::ExportNativeGeneratedInitCode(FClass* Class, F
 	{	
 		const TCHAR* ClassNameCPP = NameLookupCPP.GetNameCPP(Class);
 
-		GeneratedFunctionText.Logf(TEXT("    static FCompiledInDefer Z_CompiledInDefer_UClass_%s(%s, &%s::StaticClass, TEXT(\"%s\"), %s);\r\n"),
+		GeneratedFunctionText.Logf(TEXT("\tstatic FCompiledInDefer Z_CompiledInDefer_UClass_%s(%s, &%s::StaticClass, TEXT(\"%s\"), %s);\r\n"),
 			ClassNameCPP, *SingletonName, ClassNameCPP, ClassNameCPP,
 			bIsDynamic ? TEXT("true") : TEXT("false"));
 		
@@ -1534,11 +1453,11 @@ void FNativeClassHeaderGenerator::ExportNativeGeneratedInitCode(FClass* Class, F
 		// Emit the IMPLEMENT_CLASS macro to go in the generated cpp file.
 		if (!bIsDynamic)
 		{
-			GeneratedPackageCPP.Logf(TEXT("    IMPLEMENT_CLASS(%s, %u);\r\n"), ClassNameCPP, ClassCrc);
+			GeneratedPackageCPP.Logf(TEXT("\tIMPLEMENT_CLASS(%s, %u);\r\n"), ClassNameCPP, ClassCrc);
 		}
 		else
 		{
-			GeneratedPackageCPP.Logf(TEXT("    IMPLEMENT_DYNAMIC_CLASS(%s, %u);\r\n"), ClassNameCPP, ClassCrc);			
+			GeneratedPackageCPP.Logf(TEXT("\tIMPLEMENT_DYNAMIC_CLASS(%s, %u);\r\n"), ClassNameCPP, ClassCrc);			
 		}
 	}
 }
@@ -1558,8 +1477,8 @@ void FNativeClassHeaderGenerator::ExportFunction(UFunction* Function, FScope* Sc
 
 	FUHTStringBuilder CurrentFunctionText;
 
-	CurrentFunctionText.Logf(TEXT("    UFunction* %s\r\n"), *SingletonName);
-	CurrentFunctionText.Logf(TEXT("    {\r\n"));
+	CurrentFunctionText.Logf(TEXT("\tUFunction* %s\r\n"), *SingletonName);
+	CurrentFunctionText.Logf(TEXT("\t{\r\n"));
 
 	if (bIsNoExport || !(Function->FunctionFlags&FUNC_Event))  // non-events do not export a params struct, so lets do that locally for offset determination
 	{
@@ -1577,16 +1496,16 @@ void FNativeClassHeaderGenerator::ExportFunction(UFunction* Function, FScope* Sc
 
 	if (UObject* Outer = Function->GetOuter())
 	{
-		CurrentFunctionText.Logf(TEXT("        UObject* Outer=%s;\r\n"), Outer->IsA<UPackage>() ? *GetPackageSingletonName((UPackage*)Outer) : *GetSingletonName(Function->GetOwnerClass()));
+		CurrentFunctionText.Logf(TEXT("\t\tUObject* Outer=%s;\r\n"), Outer->IsA<UPackage>() ? *GetPackageSingletonName((UPackage*)Outer) : *GetSingletonName(Function->GetOwnerClass()));
 	}
 	else
 	{
-		CurrentFunctionText.Logf(TEXT("        UObject* Outer=nullptr;\r\n"));
+		CurrentFunctionText.Logf(TEXT("\t\tUObject* Outer=nullptr;\r\n"));
 	}
 
-	CurrentFunctionText.Logf(TEXT("        static UFunction* ReturnFunction = NULL;\r\n"));
-	CurrentFunctionText.Logf(TEXT("        if (!ReturnFunction)\r\n"));
-	CurrentFunctionText.Logf(TEXT("        {\r\n"));
+	CurrentFunctionText.Logf(TEXT("\t\tstatic UFunction* ReturnFunction = NULL;\r\n"));
+	CurrentFunctionText.Logf(TEXT("\t\tif (!ReturnFunction)\r\n"));
+	CurrentFunctionText.Logf(TEXT("\t\t{\r\n"));
 	FString SuperFunctionString(TEXT("NULL"));
 	if (SuperFunction)
 	{
@@ -1616,7 +1535,7 @@ void FNativeClassHeaderGenerator::ExportFunction(UFunction* Function, FScope* Sc
 
 	FString UFunctionType = bIsDelegate ? TEXT("UDelegateFunction") : TEXT("UFunction");
 
-	CurrentFunctionText.Logf(TEXT("            ReturnFunction = new(EC_InternalUseOnlyConstructor, Outer, TEXT(\"%s\"), RF_Public|RF_Transient|RF_Native) %s(FObjectInitializer(), %s, 0x%08X, %d%s);\r\n"),
+	CurrentFunctionText.Logf(TEXT("\t\t\tReturnFunction = new(EC_InternalUseOnlyConstructor, Outer, TEXT(\"%s\"), RF_Public|RF_Transient|RF_Native) %s(FObjectInitializer(), %s, 0x%08X, %d%s);\r\n"),
 		*FNativeClassHeaderGenerator::GetOverriddenName(Function),
 		*UFunctionType,
 		*SuperFunctionString,
@@ -1628,33 +1547,33 @@ void FNativeClassHeaderGenerator::ExportFunction(UFunction* Function, FScope* Sc
 
 
 	FString OuterString = FString(TEXT("ReturnFunction"));
-	FString Meta = GetMetaDataCodeForObject(Function, *OuterString, TEXT("            "));
+	FString Meta = GetMetaDataCodeForObject(Function, *OuterString, TEXT("\t\t\t"));
 
 	for (int32 Index = Props.Num() - 1; Index >= 0; Index--)
 	{
-		OutputProperty(Meta, CurrentFunctionText, OuterString, Props[Index], TEXT("            "));
+		OutputProperty(Meta, CurrentFunctionText, OuterString, Props[Index], TEXT("\t\t\t"));
 	}
 
 	if (FunctionData.FunctionFlags & (FUNC_NetRequest | FUNC_NetResponse))
 	{
-		CurrentFunctionText.Logf(TEXT("            ReturnFunction->RPCId=%d;\r\n"), FunctionData.RPCId);
-		CurrentFunctionText.Logf(TEXT("            ReturnFunction->RPCResponseId=%d;\r\n"), FunctionData.RPCResponseId);
+		CurrentFunctionText.Logf(TEXT("\t\t\tReturnFunction->RPCId=%d;\r\n"), FunctionData.RPCId);
+		CurrentFunctionText.Logf(TEXT("\t\t\tReturnFunction->RPCResponseId=%d;\r\n"), FunctionData.RPCResponseId);
 	}
 
-	CurrentFunctionText.Logf(TEXT("            ReturnFunction->Bind();\r\n"));
-	CurrentFunctionText.Logf(TEXT("            ReturnFunction->StaticLink();\r\n"));
+	CurrentFunctionText.Logf(TEXT("\t\t\tReturnFunction->Bind();\r\n"));
+	CurrentFunctionText.Logf(TEXT("\t\t\tReturnFunction->StaticLink();\r\n"));
 
 	if (Meta.Len())
 	{
 		CurrentFunctionText.Logf(TEXT("#if WITH_METADATA\r\n"));
-		CurrentFunctionText.Logf(TEXT("            UMetaData* MetaData = ReturnFunction->GetOutermost()->GetMetaData();\r\n"));
+		CurrentFunctionText.Logf(TEXT("\t\t\tUMetaData* MetaData = ReturnFunction->GetOutermost()->GetMetaData();\r\n"));
 		CurrentFunctionText.Log(*Meta);
 		CurrentFunctionText.Logf(TEXT("#endif\r\n"));
 	}
 
-	CurrentFunctionText.Logf(TEXT("        }\r\n"));
-	CurrentFunctionText.Logf(TEXT("        return ReturnFunction;\r\n"));
-	CurrentFunctionText.Logf(TEXT("    }\r\n"));
+	CurrentFunctionText.Logf(TEXT("\t\t}\r\n"));
+	CurrentFunctionText.Logf(TEXT("\t\treturn ReturnFunction;\r\n"));
+	CurrentFunctionText.Logf(TEXT("\t}\r\n"));
 
 	uint32 FunctionCrc = GenerateTextCRC(*CurrentFunctionText);
 	GGeneratedCodeCRCs.Add(Function, FunctionCrc);
@@ -1670,8 +1589,8 @@ void FNativeClassHeaderGenerator::ExportNatives(FClass* Class)
 		return;
 	}
 
-	GeneratedPackageCPP.Logf(TEXT("    void %s::StaticRegisterNatives%s()\r\n"),NameLookupCPP.GetNameCPP(Class),NameLookupCPP.GetNameCPP(Class));
-	GeneratedPackageCPP.Logf(TEXT("    {\r\n"));
+	GeneratedPackageCPP.Logf(TEXT("\tvoid %s::StaticRegisterNatives%s()\r\n"),NameLookupCPP.GetNameCPP(Class),NameLookupCPP.GetNameCPP(Class));
+	GeneratedPackageCPP.Logf(TEXT("\t{\r\n"));
 
 	{
 		TArray<UFunction*> FunctionsToExport;
@@ -1687,7 +1606,7 @@ void FNativeClassHeaderGenerator::ExportNatives(FClass* Class)
 
 		for (auto Func : FunctionsToExport)
 		{
-			GeneratedPackageCPP.Logf(TEXT("        FNativeFunctionRegistrar::RegisterFunction(%s::StaticClass(),\"%s\",(Native)&%s::exec%s);\r\n"),
+			GeneratedPackageCPP.Logf(TEXT("\t\tFNativeFunctionRegistrar::RegisterFunction(%s::StaticClass(),\"%s\",(Native)&%s::exec%s);\r\n"),
 				NameLookupCPP.GetNameCPP(Class),
 				*FNativeClassHeaderGenerator::GetOverriddenName(Func),
 				Class->HasAnyClassFlags(CLASS_Interface) ? *FString::Printf(TEXT("I%s"), *Class->GetName()) : NameLookupCPP.GetNameCPP(Class),
@@ -1700,11 +1619,11 @@ void FNativeClassHeaderGenerator::ExportNatives(FClass* Class)
 	{
 		if (Struct->StructFlags & STRUCT_Native)
 		{
-			GeneratedPackageCPP.Logf( TEXT("        UScriptStruct::DeferCppStructOps(FName(TEXT(\"%s\")),new UScriptStruct::TCppStructOps<%s%s>);\r\n"), *Struct->GetName(), Struct->GetPrefixCPP(), *Struct->GetName() );
+			GeneratedPackageCPP.Logf( TEXT("\t\tUScriptStruct::DeferCppStructOps(FName(TEXT(\"%s\")),new UScriptStruct::TCppStructOps<%s%s>);\r\n"), *Struct->GetName(), Struct->GetPrefixCPP(), *Struct->GetName() );
 		}
 	}
 
-	GeneratedPackageCPP.Logf(TEXT("    }\r\n"));
+	GeneratedPackageCPP.Logf(TEXT("\t}\r\n"));
 }
 
 void FNativeClassHeaderGenerator::ExportInterfaceCallFunctions(const TArray<UFunction*>& InCallbackFunctions, UClass* Class, FClassMetaData* ClassData, FUHTStringBuilder& HeaderOutput)
@@ -2006,8 +1925,8 @@ void FNativeClassHeaderGenerator::ExportClassesFromSourceFileInner(FUnrealSource
 				FUHTStringBuilder UInterfaceBoilerplate;
 
 				// Export the class's native function registration.
-				UInterfaceBoilerplate.Logf(TEXT("    private:\r\n"));
-				UInterfaceBoilerplate.Logf(TEXT("    static void StaticRegisterNatives%s();\r\n"), NameLookupCPP.GetNameCPP(Class));
+				UInterfaceBoilerplate.Logf(TEXT("\tprivate:\r\n"));
+				UInterfaceBoilerplate.Logf(TEXT("\tstatic void StaticRegisterNatives%s();\r\n"), NameLookupCPP.GetNameCPP(Class));
 				UInterfaceBoilerplate.Log(*FriendText);
 				FriendText.Reset();
 				UInterfaceBoilerplate.Logf(TEXT("public:\r\n"));
@@ -2020,7 +1939,7 @@ void FNativeClassHeaderGenerator::ExportClassesFromSourceFileInner(FUnrealSource
 				}
 
 				const bool bCastedClass = Class->HasAnyCastFlag(CASTCLASS_AllFlags) && SuperClass && Class->ClassCastFlags != SuperClass->ClassCastFlags;
-				UInterfaceBoilerplate.Logf(TEXT("    DECLARE_CLASS(%s, %s, COMPILED_IN_FLAGS(CLASS_Abstract%s), %s, %s, %s_API)\r\n"),
+				UInterfaceBoilerplate.Logf(TEXT("\tDECLARE_CLASS(%s, %s, COMPILED_IN_FLAGS(CLASS_Abstract%s), %s, %s, %s_API)\r\n"),
 					ClassCPPName,
 					SuperClassCPPName,
 					*GetClassFlagExportText(Class),
@@ -2028,12 +1947,12 @@ void FNativeClassHeaderGenerator::ExportClassesFromSourceFileInner(FUnrealSource
 					*FPackageName::GetShortName(Class->GetOuter()->GetName()),
 					*APIArg);
 
-				UInterfaceBoilerplate.Logf(TEXT("    DECLARE_SERIALIZER(%s)\r\n"), ClassCPPName);
-				UInterfaceBoilerplate.Log(TEXT("    enum {IsIntrinsic=COMPILED_IN_INTRINSIC};\r\n"));
+				UInterfaceBoilerplate.Logf(TEXT("\tDECLARE_SERIALIZER(%s)\r\n"), ClassCPPName);
+				UInterfaceBoilerplate.Log(TEXT("\tenum {IsIntrinsic=COMPILED_IN_INTRINSIC};\r\n"));
 
 				if (Class->ClassWithin != Class->GetSuperClass()->ClassWithin)
 				{
-					UInterfaceBoilerplate.Logf(TEXT("    DECLARE_WITHIN(%s)\r\n"), NameLookupCPP.GetNameCPP(Class->GetClassWithin()));
+					UInterfaceBoilerplate.Logf(TEXT("\tDECLARE_WITHIN(%s)\r\n"), NameLookupCPP.GetNameCPP(Class->GetClassWithin()));
 				}
 
 				ExportConstructorsMacros(SourceFile.GetGeneratedMacroName(ClassData), Class);
@@ -2117,10 +2036,10 @@ void FNativeClassHeaderGenerator::ExportClassesFromSourceFileInner(FUnrealSource
 			FUHTStringBuilder ClassBoilerplate;
 
 			// Export the class's native function registration.
-			ClassBoilerplate.Logf(TEXT("    private:\r\n"));
-			ClassBoilerplate.Logf(TEXT("    static void StaticRegisterNatives%s();\r\n"), NameLookupCPP.GetNameCPP(Class));
+			ClassBoilerplate.Logf(TEXT("\tprivate:\r\n"));
+			ClassBoilerplate.Logf(TEXT("\tstatic void StaticRegisterNatives%s();\r\n"), NameLookupCPP.GetNameCPP(Class));
 			ClassBoilerplate.Log(*FriendText);
-			ClassBoilerplate.Logf(TEXT("    public:\r\n"));
+			ClassBoilerplate.Logf(TEXT("\tpublic:\r\n"));
 
 			// Build the DECLARE_CLASS line
 			FString APIArg(API);
@@ -2129,7 +2048,7 @@ void FNativeClassHeaderGenerator::ExportClassesFromSourceFileInner(FUnrealSource
 				APIArg = TEXT("NO");
 			}
 			const bool bCastedClass = Class->HasAnyCastFlag(CASTCLASS_AllFlags) && SuperClass && Class->ClassCastFlags != SuperClass->ClassCastFlags;
-			ClassBoilerplate.Logf(TEXT("    DECLARE_CLASS(%s, %s, COMPILED_IN_FLAGS(%s%s), %s, %s, %s_API)\r\n"),
+			ClassBoilerplate.Logf(TEXT("\tDECLARE_CLASS(%s, %s, COMPILED_IN_FLAGS(%s%s), %s, %s, %s_API)\r\n"),
 				ClassCPPName,
 				SuperClassCPPName ? SuperClassCPPName : TEXT("None"),
 				Class->HasAnyClassFlags(CLASS_Abstract) ? TEXT("CLASS_Abstract") : TEXT("0"),
@@ -2138,19 +2057,19 @@ void FNativeClassHeaderGenerator::ExportClassesFromSourceFileInner(FUnrealSource
 				*FPackageName::GetShortName(*Class->GetOuter()->GetName()),
 				*APIArg);
 
-			ClassBoilerplate.Logf(TEXT("    DECLARE_SERIALIZER(%s)\r\n"), ClassCPPName);
-			ClassBoilerplate.Log(TEXT("    /** Indicates whether the class is compiled into the engine */"));
-			ClassBoilerplate.Log(TEXT("    enum {IsIntrinsic=COMPILED_IN_INTRINSIC};\r\n"));
+			ClassBoilerplate.Logf(TEXT("\tDECLARE_SERIALIZER(%s)\r\n"), ClassCPPName);
+			ClassBoilerplate.Log(TEXT("\t/** Indicates whether the class is compiled into the engine */\r\n"));
+			ClassBoilerplate.Log(TEXT("\tenum {IsIntrinsic=COMPILED_IN_INTRINSIC};\r\n"));
 
 			if (SuperClass && Class->ClassWithin != SuperClass->ClassWithin)
 			{
-				ClassBoilerplate.Logf(TEXT("    DECLARE_WITHIN(%s)\r\n"), NameLookupCPP.GetNameCPP(Class->GetClassWithin()));
+				ClassBoilerplate.Logf(TEXT("\tDECLARE_WITHIN(%s)\r\n"), NameLookupCPP.GetNameCPP(Class->GetClassWithin()));
 			}
 
 			// export the class's config name
 			if (SuperClass && Class->ClassConfigName != NAME_None && Class->ClassConfigName != SuperClass->ClassConfigName)
 			{
-				ClassBoilerplate.Logf(TEXT("    static const TCHAR* StaticConfigName() {return TEXT(\"%s\");}\r\n\r\n"), *Class->ClassConfigName.ToString());
+				ClassBoilerplate.Logf(TEXT("\tstatic const TCHAR* StaticConfigName() {return TEXT(\"%s\");}\r\n\r\n"), *Class->ClassConfigName.ToString());
 			}
 
 			// export implementation of _getUObject for classes that implement interfaces
@@ -2274,7 +2193,7 @@ void ExportVTableHelperConstructorDecl(FStringOutputDevice& Out, FClass* Class,
 	check(ClassData);
 	if (!ClassData->bCustomVTableHelperConstructorDeclared)
 	{
-		Out.Logf(TEXT("%sDECLARE_VTABLE_PTR_HELPER_CTOR(%s_API, %s);" LINE_TERMINATOR), FCString::Spc(4), *API, NameLookupCPP.GetNameCPP(Class));
+		Out.Logf(TEXT("\tDECLARE_VTABLE_PTR_HELPER_CTOR(%s_API, %s);" LINE_TERMINATOR), *API, NameLookupCPP.GetNameCPP(Class));
 	}
 }
 
@@ -2312,7 +2231,7 @@ void ExportVTableHelperConstructorBody(FStringOutputDevice& Out, FClass* Class)
 	check(ClassData);
 	if (!ClassData->bCustomVTableHelperConstructorDeclared)
 	{
-		Out.Logf(TEXT("%sDEFINE_VTABLE_PTR_HELPER_CTOR(%s);" LINE_TERMINATOR), FCString::Spc(4), NameLookupCPP.GetNameCPP(Class));
+		Out.Logf(TEXT("\tDEFINE_VTABLE_PTR_HELPER_CTOR(%s);" LINE_TERMINATOR), NameLookupCPP.GetNameCPP(Class));
 	}
 }
 
@@ -2569,9 +2488,9 @@ void FNativeClassHeaderGenerator::ExportClassesFromSourceFileWrapper(FUnrealSour
 	GeneratedHeaderTextWithCopyright.Log(
 		TEXT("// Copyright 1998-2015 Epic Games, Inc. All Rights Reserved.\r\n")
 		TEXT("/*===========================================================================\r\n")
-		TEXT("    C++ class header boilerplate exported from UnrealHeaderTool.\r\n")
-		TEXT("    This is automatically generated by the tools.\r\n")
-		TEXT("    DO NOT modify this manually! Edit the corresponding .h files instead!\r\n")
+		TEXT("\tC++ class header boilerplate exported from UnrealHeaderTool.\r\n")
+		TEXT("\tThis is automatically generated by the tools.\r\n")
+		TEXT("\tDO NOT modify this manually! Edit the corresponding .h files instead!\r\n")
 		TEXT("===========================================================================*/\r\n"));
 	GeneratedHeaderTextWithCopyright.Log(
 		LINE_TERMINATOR
@@ -2609,7 +2528,7 @@ void FNativeClassHeaderGenerator::ExportClassesFromSourceFileWrapper(FUnrealSour
 
 	auto IncludeStr = FString::Printf(
 		TEXT("#ifndef %s")			LINE_TERMINATOR
-		TEXT("    #include \"%s\"")	LINE_TERMINATOR
+		TEXT("\t#include \"%s\"")	LINE_TERMINATOR
 		TEXT("#endif")				LINE_TERMINATOR,
 		*SourceFile.GetFileDefineName(), *NewFileName);
 
@@ -2740,7 +2659,7 @@ void FNativeClassHeaderGenerator::ExportEnums( const TArray<UEnum*>& Enums )
 		for (int32 i = 0; i < Enum->NumEnums() - 1; i++)
 		{
 			const FString QualifiedEnumValue = Enum->GetNameByIndex(i).ToString();
-			EnumForeachText.Logf( TEXT("\\\r\n    op(%s) "), *QualifiedEnumValue );
+			EnumForeachText.Logf( TEXT("\\\r\n\top(%s) "), *QualifiedEnumValue );
 		}
 		EnumForeachText.Logf( TEXT("\r\n") );
 	}
@@ -2769,8 +2688,8 @@ void FNativeClassHeaderGenerator::ExportGeneratedStructBodyMacros(FUnrealSourceF
 				RequiredAPI = FriendApiString;
 			}
 
-			const FString FriendLine = FString::Printf(TEXT("    friend %sclass UScriptStruct* %s;\r\n"), *FriendApiString, *StaticConstructionString);
-			const FString StaticClassLine = FString::Printf(TEXT("    %sstatic class UScriptStruct* StaticStruct();\r\n"), *RequiredAPI);
+			const FString FriendLine = FString::Printf(TEXT("\tfriend %sclass UScriptStruct* %s;\r\n"), *FriendApiString, *StaticConstructionString);
+			const FString StaticClassLine = FString::Printf(TEXT("\t%sstatic class UScriptStruct* StaticStruct();\r\n"), *RequiredAPI);
 
 			const FString CombinedLine = FriendLine + StaticClassLine;
 			const FString MacroName = SourceFile.GetGeneratedBodyMacroName(Struct->StructMacroDeclaredLineNumber);
@@ -2795,29 +2714,29 @@ void FNativeClassHeaderGenerator::ExportGeneratedStructBodyMacros(FUnrealSourceF
 			else
 			{
 				OuterName = GetPackageSingletonName(CastChecked<UPackage>(Struct->GetOuter()));
-				GeneratedPackageCPP.Logf(TEXT("    extern %sclass UPackage* %s;\r\n"), *FriendApiString, *OuterName);
+				GeneratedPackageCPP.Logf(TEXT("\textern %sclass UPackage* %s;\r\n"), *FriendApiString, *OuterName);
 			}
 
 			if (!bIsDynamic)
 			{
-				GeneratedPackageCPP.Logf(TEXT("    static class UScriptStruct* Singleton = NULL;\r\n"));
+				GeneratedPackageCPP.Logf(TEXT("\tstatic class UScriptStruct* Singleton = NULL;\r\n"));
 			}
 			else
 			{
 				FString PackageSingletonName = GetPackageSingletonName(CastChecked<UPackage>(Struct->GetOuter()));
-				GeneratedPackageCPP.Logf(TEXT("    class UScriptStruct* Singleton = Cast<UScriptStruct>(StaticFindObjectFast(UScriptStruct::StaticClass(), %s, TEXT(\"%s\")));\r\n"),
+				GeneratedPackageCPP.Logf(TEXT("\tclass UScriptStruct* Singleton = Cast<UScriptStruct>(StaticFindObjectFast(UScriptStruct::StaticClass(), %s, TEXT(\"%s\")));\r\n"),
 					*OuterName, *Struct->GetName());
 			}
-			GeneratedPackageCPP.Logf(TEXT("    if (!Singleton)\r\n"));
-			GeneratedPackageCPP.Logf(TEXT("    {\r\n"));
-			GeneratedPackageCPP.Logf(TEXT("        extern %sclass UScriptStruct* %s;\r\n"), *FriendApiString, *StaticConstructionString);
-			GeneratedPackageCPP.Logf(TEXT("        extern %suint32 %s();\r\n"), *FriendApiString, *GetCRCName);
+			GeneratedPackageCPP.Logf(TEXT("\tif (!Singleton)\r\n"));
+			GeneratedPackageCPP.Logf(TEXT("\t{\r\n"));
+			GeneratedPackageCPP.Logf(TEXT("\t\textern %sclass UScriptStruct* %s;\r\n"), *FriendApiString, *StaticConstructionString);
+			GeneratedPackageCPP.Logf(TEXT("\t\textern %suint32 %s();\r\n"), *FriendApiString, *GetCRCName);
 
-			GeneratedPackageCPP.Logf(TEXT("        Singleton = GetStaticStruct(%s, %s, TEXT(\"%s\"), sizeof(%s), %s());\r\n"),
+			GeneratedPackageCPP.Logf(TEXT("\t\tSingleton = GetStaticStruct(%s, %s, TEXT(\"%s\"), sizeof(%s), %s());\r\n"),
 				*SingletonName, *OuterName, *Struct->GetName(), StructNameCPP, *GetCRCName);
 
-			GeneratedPackageCPP.Logf(TEXT("    }\r\n"));
-			GeneratedPackageCPP.Logf(TEXT("    return Singleton;\r\n"));
+			GeneratedPackageCPP.Logf(TEXT("\t}\r\n"));
+			GeneratedPackageCPP.Logf(TEXT("\treturn Singleton;\r\n"));
 			GeneratedPackageCPP.Logf(TEXT("}\r\n"));
 
 			GeneratedPackageCPP.Logf(TEXT("static FCompiledInDeferStruct Z_CompiledInDeferStruct_UScriptStruct_%s(%s::StaticStruct, TEXT(\"%s\"), TEXT(\"%s\"), %s);\r\n"),
@@ -2849,8 +2768,8 @@ void FNativeClassHeaderGenerator::ExportGeneratedStructBodyMacros(FUnrealSourceF
 
 			FUHTStringBuilder GeneratedStructRegisterFunctionText;
 
-			GeneratedStructRegisterFunctionText.Logf(TEXT("    UScriptStruct* %s\r\n"), *SingletonName);
-			GeneratedStructRegisterFunctionText.Logf(TEXT("    {\r\n"));
+			GeneratedStructRegisterFunctionText.Logf(TEXT("\tUScriptStruct* %s\r\n"), *SingletonName);
+			GeneratedStructRegisterFunctionText.Logf(TEXT("\t{\r\n"));
 
 			// if this is a no export struct, we will put a local struct here for offset determination
 			TArray<UScriptStruct*> Structs;
@@ -2864,23 +2783,23 @@ void FNativeClassHeaderGenerator::ExportGeneratedStructBodyMacros(FUnrealSourceF
 			// Structs can either have a UClass or UPackage as outer (if delcared in non-UClass header).
 			if (ScriptStruct->GetOuter()->IsA(UStruct::StaticClass()))
 			{
-				GeneratedStructRegisterFunctionText.Logf(TEXT("        UStruct* Outer=%s;\r\n"), *GetSingletonName(CastChecked<UStruct>(ScriptStruct->GetOuter())));
+				GeneratedStructRegisterFunctionText.Logf(TEXT("\t\tUStruct* Outer=%s;\r\n"), *GetSingletonName(CastChecked<UStruct>(ScriptStruct->GetOuter())));
 			}
 			else
 			{
-				GeneratedStructRegisterFunctionText.Logf(TEXT("        UPackage* Outer=%s;\r\n"), *GetPackageSingletonName(CastChecked<UPackage>(ScriptStruct->GetOuter())));
+				GeneratedStructRegisterFunctionText.Logf(TEXT("\t\tUPackage* Outer=%s;\r\n"), *GetPackageSingletonName(CastChecked<UPackage>(ScriptStruct->GetOuter())));
 			}
 			if (!bIsDynamic)
 			{
-				GeneratedStructRegisterFunctionText.Logf(TEXT("        static UScriptStruct* ReturnStruct = NULL;\r\n"));
+				GeneratedStructRegisterFunctionText.Logf(TEXT("\t\tstatic UScriptStruct* ReturnStruct = NULL;\r\n"));
 			}
 			else
 			{
-				GeneratedStructRegisterFunctionText.Logf(TEXT("        UScriptStruct* ReturnStruct = Cast<UScriptStruct>(StaticFindObjectFast(UScriptStruct::StaticClass(), Outer, TEXT(\"%s\")));\r\n"),
+				GeneratedStructRegisterFunctionText.Logf(TEXT("\t\tUScriptStruct* ReturnStruct = Cast<UScriptStruct>(StaticFindObjectFast(UScriptStruct::StaticClass(), Outer, TEXT(\"%s\")));\r\n"),
 					*ScriptStruct->GetName());
 			}
-			GeneratedStructRegisterFunctionText.Logf(TEXT("        if (!ReturnStruct)\r\n"));
-			GeneratedStructRegisterFunctionText.Logf(TEXT("        {\r\n"));
+			GeneratedStructRegisterFunctionText.Logf(TEXT("\t\tif (!ReturnStruct)\r\n"));
+			GeneratedStructRegisterFunctionText.Logf(TEXT("\t\t{\r\n"));
 			FString BaseStructString(TEXT("NULL"));
 			if (BaseStruct)
 			{
@@ -2900,7 +2819,7 @@ void FNativeClassHeaderGenerator::ExportGeneratedStructBodyMacros(FUnrealSourceF
 				ExplicitSizeString = FString::Printf(TEXT(", sizeof(%s), ALIGNOF(%s)"), NameLookupCPP.GetNameCPP(ScriptStruct), NameLookupCPP.GetNameCPP(ScriptStruct));
 			}
 
-			GeneratedStructRegisterFunctionText.Logf(TEXT("            ReturnStruct = new(EC_InternalUseOnlyConstructor, Outer, TEXT(\"%s\"), RF_Public|RF_Transient|RF_Native) UScriptStruct(FObjectInitializer(), %s, %s, EStructFlags(0x%08X)%s);\r\n"),
+			GeneratedStructRegisterFunctionText.Logf(TEXT("\t\t\tReturnStruct = new(EC_InternalUseOnlyConstructor, Outer, TEXT(\"%s\"), RF_Public|RF_Transient|RF_Native) UScriptStruct(FObjectInitializer(), %s, %s, EStructFlags(0x%08X)%s);\r\n"),
 				*ScriptStruct->GetName(),
 				*BaseStructString,
 				*CppStructOpsString,
@@ -2915,30 +2834,30 @@ void FNativeClassHeaderGenerator::ExportGeneratedStructBodyMacros(FUnrealSourceF
 				Props.Add(*ItInner);
 			}
 			FString OuterString = FString(TEXT("ReturnStruct"));
-			FString Meta = GetMetaDataCodeForObject(ScriptStruct, *OuterString, TEXT("            "));
-			OutputProperties(Meta, GeneratedStructRegisterFunctionText, OuterString, Props, TEXT("            "));
-			GeneratedStructRegisterFunctionText.Logf(TEXT("            ReturnStruct->StaticLink();\r\n"));
+			FString Meta = GetMetaDataCodeForObject(ScriptStruct, *OuterString, TEXT("\t\t\t"));
+			OutputProperties(Meta, GeneratedStructRegisterFunctionText, OuterString, Props, TEXT("\t\t\t"));
+			GeneratedStructRegisterFunctionText.Logf(TEXT("\t\t\tReturnStruct->StaticLink();\r\n"));
 
 			if (Meta.Len())
 			{
 				GeneratedStructRegisterFunctionText.Logf(TEXT("#if WITH_METADATA\r\n"));
-				GeneratedStructRegisterFunctionText.Logf(TEXT("            UMetaData* MetaData = ReturnStruct->GetOutermost()->GetMetaData();\r\n"));
+				GeneratedStructRegisterFunctionText.Logf(TEXT("\t\t\tUMetaData* MetaData = ReturnStruct->GetOutermost()->GetMetaData();\r\n"));
 				GeneratedStructRegisterFunctionText.Log(*Meta);
 				GeneratedStructRegisterFunctionText.Logf(TEXT("#endif\r\n"));
 			}
 
-			GeneratedStructRegisterFunctionText.Logf(TEXT("        }\r\n"));
-			GeneratedStructRegisterFunctionText.Logf(TEXT("        return ReturnStruct;\r\n"));
-			GeneratedStructRegisterFunctionText.Logf(TEXT("    }\r\n"));
+			GeneratedStructRegisterFunctionText.Logf(TEXT("\t\t}\r\n"));
+			GeneratedStructRegisterFunctionText.Logf(TEXT("\t\treturn ReturnStruct;\r\n"));
+			GeneratedStructRegisterFunctionText.Logf(TEXT("\t}\r\n"));
 
 			uint32 StructCrc = GenerateTextCRC(*GeneratedStructRegisterFunctionText);
 			GGeneratedCodeCRCs.Add(ScriptStruct, StructCrc);
 
 			auto& GeneratedFunctionText = GetGeneratedFunctionTextDevice();
 			GeneratedFunctionText += GeneratedStructRegisterFunctionText;
-			GeneratedFunctionText.Logf(TEXT("    uint32 Get_%s_CRC() { return %uU; }\r\n"), *SingletonName.Replace(TEXT("()"), TEXT(""), ESearchCase::CaseSensitive), StructCrc);
+			GeneratedFunctionText.Logf(TEXT("\tuint32 Get_%s_CRC() { return %uU; }\r\n"), *SingletonName.Replace(TEXT("()"), TEXT(""), ESearchCase::CaseSensitive), StructCrc);
 
-			//CallSingletons.Logf(TEXT("                OuterClass->LinkChild(%s); // %u\r\n"), *SingletonName, StructCrc);
+			//CallSingletons.Logf(TEXT("\t\t\t\tOuterClass->LinkChild(%s); // %u\r\n"), *SingletonName, StructCrc);
 		}
 	}
 }
@@ -2965,24 +2884,24 @@ void FNativeClassHeaderGenerator::ExportGeneratedEnumsInitCode(const TArray<UEnu
 
 		GeneratedPackageCPP.Logf(TEXT("static class UEnum* %s_StaticEnum()\r\n"), *Enum->GetName());
 		GeneratedPackageCPP.Logf(TEXT("{\r\n"));
-		GeneratedPackageCPP.Logf(TEXT("    extern %sclass UPackage* %s;\r\n"), *FriendApiString, *PackageSingletonName);
+		GeneratedPackageCPP.Logf(TEXT("\textern %sclass UPackage* %s;\r\n"), *FriendApiString, *PackageSingletonName);
 
 		if (!bIsDynamic)
 		{
-			GeneratedPackageCPP.Logf(TEXT("    static class UEnum* Singleton = NULL;\r\n"));
+			GeneratedPackageCPP.Logf(TEXT("\tstatic class UEnum* Singleton = NULL;\r\n"));
 		}
 		else
 		{
-			GeneratedPackageCPP.Logf(TEXT("    class UEnum* Singleton = Cast<UEnum>(StaticFindObjectFast(UEnum::StaticClass(), %s, TEXT(\"%s\")));\r\n"),
+			GeneratedPackageCPP.Logf(TEXT("\tclass UEnum* Singleton = Cast<UEnum>(StaticFindObjectFast(UEnum::StaticClass(), %s, TEXT(\"%s\")));\r\n"),
 				*PackageSingletonName, *Enum->GetName());
 		}
-		GeneratedPackageCPP.Logf(TEXT("    if (!Singleton)\r\n"));
-		GeneratedPackageCPP.Logf(TEXT("    {\r\n"));
-		GeneratedPackageCPP.Logf(TEXT("        extern %sclass UEnum* %s;\r\n"), *FriendApiString, *StaticConstructionString);		
-		GeneratedPackageCPP.Logf(TEXT("        Singleton = GetStaticEnum(%s, %s, TEXT(\"%s\"));\r\n"), *SingletonName, *PackageSingletonName, *Enum->GetName());
+		GeneratedPackageCPP.Logf(TEXT("\tif (!Singleton)\r\n"));
+		GeneratedPackageCPP.Logf(TEXT("\t{\r\n"));
+		GeneratedPackageCPP.Logf(TEXT("\t\textern %sclass UEnum* %s;\r\n"), *FriendApiString, *StaticConstructionString);		
+		GeneratedPackageCPP.Logf(TEXT("\t\tSingleton = GetStaticEnum(%s, %s, TEXT(\"%s\"));\r\n"), *SingletonName, *PackageSingletonName, *Enum->GetName());
 			
-		GeneratedPackageCPP.Logf(TEXT("    }\r\n"));
-		GeneratedPackageCPP.Logf(TEXT("    return Singleton;\r\n"));
+		GeneratedPackageCPP.Logf(TEXT("\t}\r\n"));
+		GeneratedPackageCPP.Logf(TEXT("\treturn Singleton;\r\n"));
 		GeneratedPackageCPP.Logf(TEXT("}\r\n"));
 
 		GeneratedPackageCPP.Logf(TEXT("static FCompiledInDeferEnum Z_CompiledInDeferEnum_UEnum_%s(%s_StaticEnum, TEXT(\"%s\"), TEXT(\"%s\"), %s);\r\n"), 
@@ -2995,26 +2914,26 @@ void FNativeClassHeaderGenerator::ExportGeneratedEnumsInitCode(const TArray<UEnu
 
 			FUHTStringBuilder GeneratedEnumRegisterFunctionText;
 
-			GeneratedEnumRegisterFunctionText.Logf(TEXT("    UEnum* %s\r\n"), *EnumSingletonName);
-			GeneratedEnumRegisterFunctionText.Logf(TEXT("    {\r\n"));
+			GeneratedEnumRegisterFunctionText.Logf(TEXT("\tUEnum* %s\r\n"), *EnumSingletonName);
+			GeneratedEnumRegisterFunctionText.Logf(TEXT("\t{\r\n"));
 			// Enums can either have a UClass or UPackage as outer (if declared in non-UClass header).
 			if (Enum->GetOuter()->IsA(UStruct::StaticClass()))
 			{
-				GeneratedEnumRegisterFunctionText.Logf(TEXT("        UClass* Outer=%s;\r\n"), *GetSingletonName(CastChecked<UStruct>(Enum->GetOuter())));
+				GeneratedEnumRegisterFunctionText.Logf(TEXT("\t\tUClass* Outer=%s;\r\n"), *GetSingletonName(CastChecked<UStruct>(Enum->GetOuter())));
 			}
 			else
 			{
-				GeneratedEnumRegisterFunctionText.Logf(TEXT("        UPackage* Outer=%s;\r\n"), *GetPackageSingletonName(CastChecked<UPackage>(Enum->GetOuter())));
+				GeneratedEnumRegisterFunctionText.Logf(TEXT("\t\tUPackage* Outer=%s;\r\n"), *GetPackageSingletonName(CastChecked<UPackage>(Enum->GetOuter())));
 			}
-			GeneratedEnumRegisterFunctionText.Logf(TEXT("        static UEnum* ReturnEnum = NULL;\r\n"));
-			GeneratedEnumRegisterFunctionText.Logf(TEXT("        if (!ReturnEnum)\r\n"));
-			GeneratedEnumRegisterFunctionText.Logf(TEXT("        {\r\n"));
+			GeneratedEnumRegisterFunctionText.Logf(TEXT("\t\tstatic UEnum* ReturnEnum = NULL;\r\n"));
+			GeneratedEnumRegisterFunctionText.Logf(TEXT("\t\tif (!ReturnEnum)\r\n"));
+			GeneratedEnumRegisterFunctionText.Logf(TEXT("\t\t{\r\n"));
 
-			GeneratedEnumRegisterFunctionText.Logf(TEXT("            ReturnEnum = new(EC_InternalUseOnlyConstructor, Outer, TEXT(\"%s\"), RF_Public|RF_Transient|RF_Native) UEnum(FObjectInitializer());\r\n"), *Enum->GetName());
-			GeneratedEnumRegisterFunctionText.Logf(TEXT("            TArray<TPair<FName, uint8>> EnumNames;\r\n"));
+			GeneratedEnumRegisterFunctionText.Logf(TEXT("\t\t\tReturnEnum = new(EC_InternalUseOnlyConstructor, Outer, TEXT(\"%s\"), RF_Public|RF_Transient|RF_Native) UEnum(FObjectInitializer());\r\n"), *Enum->GetName());
+			GeneratedEnumRegisterFunctionText.Logf(TEXT("\t\t\tTArray<TPair<FName, uint8>> EnumNames;\r\n"));
 			for (int32 Index = 0; Index < Enum->NumEnums(); Index++)
 			{
-				GeneratedEnumRegisterFunctionText.Logf(TEXT("            EnumNames.Add(TPairInitializer<FName, uint8>(FName(TEXT(\"%s\")), %d));\r\n"), *Enum->GetNameByIndex(Index).ToString(), Enum->GetValueByIndex(Index));
+				GeneratedEnumRegisterFunctionText.Logf(TEXT("\t\t\tEnumNames.Add(TPairInitializer<FName, uint8>(FName(TEXT(\"%s\")), %d));\r\n"), *Enum->GetNameByIndex(Index).ToString(), Enum->GetValueByIndex(Index));
 			}
 
 			FString EnumTypeStr;
@@ -3024,28 +2943,28 @@ void FNativeClassHeaderGenerator::ExportGeneratedEnumsInitCode(const TArray<UEnu
 			case UEnum::ECppForm::Namespaced: EnumTypeStr = TEXT("UEnum::ECppForm::Namespaced"); break;
 			case UEnum::ECppForm::EnumClass:  EnumTypeStr = TEXT("UEnum::ECppForm::EnumClass");  break;
 			}
-			GeneratedEnumRegisterFunctionText.Logf(TEXT("            ReturnEnum->SetEnums(EnumNames, %s);\r\n"), *EnumTypeStr);
-			GeneratedEnumRegisterFunctionText.Logf(TEXT("            ReturnEnum->CppType = TEXT(\"%s\");\r\n"), *Enum->CppType);
+			GeneratedEnumRegisterFunctionText.Logf(TEXT("\t\t\tReturnEnum->SetEnums(EnumNames, %s);\r\n"), *EnumTypeStr);
+			GeneratedEnumRegisterFunctionText.Logf(TEXT("\t\t\tReturnEnum->CppType = TEXT(\"%s\");\r\n"), *Enum->CppType);
 
-			FString Meta = GetMetaDataCodeForObject(Enum, TEXT("ReturnEnum"), TEXT("            "));
+			FString Meta = GetMetaDataCodeForObject(Enum, TEXT("ReturnEnum"), TEXT("\t\t\t"));
 			if (Meta.Len())
 			{
 				GeneratedEnumRegisterFunctionText.Logf(TEXT("#if WITH_METADATA\r\n"));
-				GeneratedEnumRegisterFunctionText.Logf(TEXT("            UMetaData* MetaData = ReturnEnum->GetOutermost()->GetMetaData();\r\n"));
+				GeneratedEnumRegisterFunctionText.Logf(TEXT("\t\t\tUMetaData* MetaData = ReturnEnum->GetOutermost()->GetMetaData();\r\n"));
 				GeneratedEnumRegisterFunctionText.Log(*Meta);
 				GeneratedEnumRegisterFunctionText.Logf(TEXT("#endif\r\n"));
 			}
 
-			GeneratedEnumRegisterFunctionText.Logf(TEXT("        }\r\n"));
-			GeneratedEnumRegisterFunctionText.Logf(TEXT("        return ReturnEnum;\r\n"));
-			GeneratedEnumRegisterFunctionText.Logf(TEXT("    }\r\n"));
+			GeneratedEnumRegisterFunctionText.Logf(TEXT("\t\t}\r\n"));
+			GeneratedEnumRegisterFunctionText.Logf(TEXT("\t\treturn ReturnEnum;\r\n"));
+			GeneratedEnumRegisterFunctionText.Logf(TEXT("\t}\r\n"));
 
 			auto& GeneratedFunctionText = GetGeneratedFunctionTextDevice();
 			GeneratedFunctionText += GeneratedEnumRegisterFunctionText;
 
 			uint32 EnumCrc = GenerateTextCRC(*GeneratedEnumRegisterFunctionText);
 			GGeneratedCodeCRCs.Add(Enum, EnumCrc);
-			// CallSingletons.Logf(TEXT("                OuterClass->LinkChild(%s); // %u\r\n"), *EnumSingletonName, EnumCrc);
+			// CallSingletons.Logf(TEXT("\t\t\t\tOuterClass->LinkChild(%s); // %u\r\n"), *EnumSingletonName, EnumCrc);
 		}
 	}
 }
@@ -3058,17 +2977,17 @@ void FNativeClassHeaderGenerator::ExportMirrorsForNoexportStructs(const TArray<U
 		UScriptStruct* Struct = NativeStructs[i];
 
 		// Export struct.
-		HeaderOutput.Logf(TEXT("%sstruct %s"), GetTabs(TextIndent), NameLookupCPP.GetNameCPP(Struct));
+		HeaderOutput.Logf(TEXT("%sstruct %s"), FCString::Tab(TextIndent), NameLookupCPP.GetNameCPP(Struct));
 		if (Struct->GetSuperStruct() != NULL)
 		{
 			HeaderOutput.Logf(TEXT(" : public %s"), NameLookupCPP.GetNameCPP(Struct->GetSuperStruct()));
 		}
-		HeaderOutput.Logf(TEXT("\r\n%s{\r\n"), GetTabs(TextIndent));
+		HeaderOutput.Logf(TEXT("\r\n%s{\r\n"), FCString::Tab(TextIndent));
 
 		// Export the struct's CPP properties.
 		ExportProperties(Struct, TextIndent, /*bAccessSpecifiers=*/ false, &HeaderOutput);
 
-		HeaderOutput.Logf(TEXT("%s};\r\n\r\n"), GetTabs(TextIndent));
+		HeaderOutput.Logf(TEXT("%s};\r\n\r\n"), FCString::Tab(TextIndent));
 	}
 }
 
@@ -3081,7 +3000,7 @@ bool FNativeClassHeaderGenerator::WillExportEventParms( UFunction* Function )
 void WriteEventFunctionPrologue(FUHTStringBuilder& Output, int32 Indent, const FParmsAndReturnProperties& Parameters, UObject* FunctionOuter, const TCHAR* FunctionName)
 {
 	// now the body - first we need to declare a struct which will hold the parameters for the event/delegate call
-	Output.Logf(TEXT("\r\n%s{\r\n"), GetTabs(Indent));
+	Output.Logf(TEXT("\r\n%s{\r\n"), FCString::Tab(Indent));
 
 	// declare and zero-initialize the parameters and return value, if applicable
 	if (!Parameters.HasParms())
@@ -3089,7 +3008,7 @@ void WriteEventFunctionPrologue(FUHTStringBuilder& Output, int32 Indent, const F
 
 	FString EventStructName = GetEventStructParamsName(FunctionOuter, FunctionName);
 
-	Output.Logf(TEXT("%s%s Parms;\r\n"), GetTabs(Indent + 1), *EventStructName );
+	Output.Logf(TEXT("%s%s Parms;\r\n"), FCString::Tab(Indent + 1), *EventStructName );
 
 	// Declare a parameter struct for this event/delegate and assign the struct members using the values passed into the event/delegate call.
 	for (auto It = Parameters.Parms.CreateConstIterator(); It; ++It)
@@ -3099,7 +3018,7 @@ void WriteEventFunctionPrologue(FUHTStringBuilder& Output, int32 Indent, const F
 		const FString PropertyName = Prop->GetName();
 		if (Prop->ArrayDim > 1)
 		{
-			Output.Logf(TEXT("%sFMemory::Memcpy(Parms.%s,%s,sizeof(Parms.%s));\r\n"), GetTabs(Indent + 1), *PropertyName, *PropertyName, *PropertyName);
+			Output.Logf(TEXT("%sFMemory::Memcpy(Parms.%s,%s,sizeof(Parms.%s));\r\n"), FCString::Tab(Indent + 1), *PropertyName, *PropertyName, *PropertyName);
 		}
 		else
 		{
@@ -3109,7 +3028,7 @@ void WriteEventFunctionPrologue(FUHTStringBuilder& Output, int32 Indent, const F
 				ValueAssignmentText += TEXT(" ? true : false");
 			}
 
-			Output.Logf(TEXT("%sParms.%s=%s;\r\n"), GetTabs(Indent + 1), *PropertyName, *ValueAssignmentText);
+			Output.Logf(TEXT("%sParms.%s=%s;\r\n"), FCString::Tab(Indent + 1), *PropertyName, *ValueAssignmentText);
 		}
 	}
 }
@@ -3126,11 +3045,11 @@ void WriteEventFunctionEpilogue(FUHTStringBuilder& Output, int32 Indent, const F
 			const FString PropertyName = Prop->GetName();
 			if ( Prop->ArrayDim > 1 )
 			{
-				Output.Logf(TEXT("%sFMemory::Memcpy(&%s,&Parms.%s,sizeof(%s));\r\n"), GetTabs(Indent + 1), *PropertyName, *PropertyName, *PropertyName);
+				Output.Logf(TEXT("%sFMemory::Memcpy(&%s,&Parms.%s,sizeof(%s));\r\n"), FCString::Tab(Indent + 1), *PropertyName, *PropertyName, *PropertyName);
 			}
 			else
 			{
-				Output.Logf(TEXT("%s%s=Parms.%s;\r\n"), GetTabs(Indent + 1), *PropertyName, *PropertyName);
+				Output.Logf(TEXT("%s%s=Parms.%s;\r\n"), FCString::Tab(Indent + 1), *PropertyName, *PropertyName);
 			}
 		}
 	}
@@ -3140,9 +3059,9 @@ void WriteEventFunctionEpilogue(FUHTStringBuilder& Output, int32 Indent, const F
 	{
 		// Make sure uint32 -> bool is supported
 		bool bBoolProperty = Parameters.Return->IsA(UBoolProperty::StaticClass());
-		Output.Logf(TEXT("%sreturn %sParms.%s;\r\n"), GetTabs(Indent + 1), bBoolProperty ? TEXT("!!") : TEXT(""), *Parameters.Return->GetName());
+		Output.Logf(TEXT("%sreturn %sParms.%s;\r\n"), FCString::Tab(Indent + 1), bBoolProperty ? TEXT("!!") : TEXT(""), *Parameters.Return->GetName());
 	}
-	Output.Logf(TEXT("%s}\r\n"), GetTabs(Indent));
+	Output.Logf(TEXT("%s}\r\n"), FCString::Tab(Indent));
 }
 
 void FNativeClassHeaderGenerator::ExportDelegateDefinitions(FUnrealSourceFile& SourceFile, const TArray<UDelegateFunction*>& DelegateFunctions, const bool bWrapperImplementationsOnly)
@@ -3237,8 +3156,8 @@ void FNativeClassHeaderGenerator::ExportDelegateDefinitions(FUnrealSourceFile& S
  */
 void ExportProtoDeclaration(FOutputDevice& Out, const FString& MessageName, TFieldIterator<UProperty>& Properties, uint64 PropertyFlags, int32 Indent)
 {
-	Out.Logf(TEXT("%smessage CMsg%sMessage\r\n"), GetTabs(Indent), *MessageName);
-	Out.Logf(TEXT("%s{\r\n"), GetTabs(Indent));
+	Out.Logf(TEXT("%smessage CMsg%sMessage\r\n"), FCString::Tab(Indent), *MessageName);
+	Out.Logf(TEXT("%s{\r\n"), FCString::Tab(Indent));
 
 	static TMap<FString, FString> ToProtoTypeMappings;
 	static bool bInitMapping = false;
@@ -3315,7 +3234,7 @@ void ExportProtoDeclaration(FOutputDevice& Out, const FString& MessageName, TFie
 				ProtoTypeName = ToProtoTypeMappings.Find(VariableTypeName);
 			}
 
-			Out.Log(GetTabs(Indent + 1));
+			Out.Log(FCString::Tab(Indent + 1));
 			Out.Log(bIsRepeated ? TEXT("repeated ") : TEXT("optional "));
 
 			if (ProtoTypeName != NULL)
@@ -3334,7 +3253,7 @@ void ExportProtoDeclaration(FOutputDevice& Out, const FString& MessageName, TFie
 		}
 	}
 
-	Out.Logf(TEXT("%s}\r\n"), GetTabs(Indent));
+	Out.Logf(TEXT("%s}\r\n"), FCString::Tab(Indent));
 }
 
 /**
@@ -3382,7 +3301,7 @@ void FNativeClassHeaderGenerator::ExportProtoMessage(const TArray<UFunction*>& I
 					ParameterList += ParamName;
 				}
 
-				HeaderOutput.Logf(TEXT("// %s%s(%s)\r\n"), GetTabs(Indent), *FunctionName, *ParameterList);
+				HeaderOutput.Logf(TEXT("// %s%s(%s)\r\n"), FCString::Tab(Indent), *FunctionName, *ParameterList);
 
 				TFieldIterator<UProperty> ParamIt(Function);
 				ExportProtoDeclaration(HeaderOutput, FunctionName, ParamIt, CPF_Parm, Indent);
@@ -3422,8 +3341,8 @@ static FString FixJavaName(const FString &StringIn)
  */
 void ExportMCPDeclaration(FOutputDevice& Out, const FString& MessageName, TFieldIterator<UProperty>& Properties, uint64 PropertyFlags, int32 Indent)
 {
-	Out.Logf(TEXT("%spublic class %sCommand extends ProfileCommand\r\n"), GetTabs(Indent), *MessageName);
-	Out.Logf(TEXT("%s{\r\n"), GetTabs(Indent));
+	Out.Logf(TEXT("%spublic class %sCommand extends ProfileCommand\r\n"), FCString::Tab(Indent), *MessageName);
+	Out.Logf(TEXT("%s{\r\n"), FCString::Tab(Indent));
 
 	static TMap<FString, FString> ToMCPTypeMappings;
 	static TMap<FString, FString> ToAnnotationMappings;
@@ -3535,17 +3454,17 @@ void ExportMCPDeclaration(FOutputDevice& Out, const FString& MessageName, TField
 
 			if (AnnotationName != NULL && !AnnotationName->IsEmpty())
 			{
-				Out.Log(GetTabs(Indent + 1));
+				Out.Log(FCString::Tab(Indent + 1));
 				Out.Logf(TEXT("%s\r\n"), **AnnotationName);
 			}
 
 			if (MCPTypeName != NULL)
 			{
-				Out.Log(GetTabs(Indent + 1));
+				Out.Log(FCString::Tab(Indent + 1));
 				Out.Logf(TEXT("private %s %s;\r\n"), **MCPTypeName, *PropertyName);
 
 				ConstructorParams += FString::Printf(TEXT(", %s %s"), **MCPTypeName, *PropertyName);
-				ConstructorText += FString::Printf(TEXT("%sthis.%s = %s;\r\n"), GetTabs(Indent + 2), *PropertyName, *PropertyName);
+				ConstructorText += FString::Printf(TEXT("%sthis.%s = %s;\r\n"), FCString::Tab(Indent + 2), *PropertyName, *PropertyName);
 			}
 			else
 			{
@@ -3558,18 +3477,18 @@ void ExportMCPDeclaration(FOutputDevice& Out, const FString& MessageName, TField
 		}
 	}
 
-	Out.Logf(TEXT("\r\n%spublic %sCommand(String epicId, String profileId%s)\r\n"), GetTabs(Indent + 1), *MessageName, *ConstructorParams);
-	Out.Logf(TEXT("%s{\r\n"), GetTabs(Indent + 1));
-	Out.Logf(TEXT("%ssuper(epicId, profileId);\r\n"), GetTabs(Indent + 2));
+	Out.Logf(TEXT("\r\n%spublic %sCommand(String epicId, String profileId%s)\r\n"), FCString::Tab(Indent + 1), *MessageName, *ConstructorParams);
+	Out.Logf(TEXT("%s{\r\n"), FCString::Tab(Indent + 1));
+	Out.Logf(TEXT("%ssuper(epicId, profileId);\r\n"), FCString::Tab(Indent + 2));
 	Out.Logf(TEXT("%s"), *ConstructorText);
-	Out.Logf(TEXT("%s}\r\n"), GetTabs(Indent + 1));
+	Out.Logf(TEXT("%s}\r\n"), FCString::Tab(Indent + 1));
 
-	Out.Logf(TEXT("\r\n%s@Override\r\n"), GetTabs(Indent + 1));
-	Out.Logf(TEXT("%sprotected void execute(@Name(\"profile\") @NotNull ProfileEx profile)\r\n"), GetTabs(Indent + 1));
-	Out.Logf(TEXT("%s{\r\n"), GetTabs(Indent + 1));
-	Out.Logf(TEXT("%s}\r\n"), GetTabs(Indent + 1));
+	Out.Logf(TEXT("\r\n%s@Override\r\n"), FCString::Tab(Indent + 1));
+	Out.Logf(TEXT("%sprotected void execute(@Name(\"profile\") @NotNull ProfileEx profile)\r\n"), FCString::Tab(Indent + 1));
+	Out.Logf(TEXT("%s{\r\n"), FCString::Tab(Indent + 1));
+	Out.Logf(TEXT("%s}\r\n"), FCString::Tab(Indent + 1));
 
-	Out.Logf(TEXT("%s}\r\n"), GetTabs(Indent));
+	Out.Logf(TEXT("%s}\r\n"), FCString::Tab(Indent));
 }
 
 /**
@@ -3617,7 +3536,7 @@ void FNativeClassHeaderGenerator::ExportMCPMessage(const TArray<UFunction*>& InC
 					ParameterList += ParamName;
 				}
 
-				HeaderOutput.Logf(TEXT("// %s%s(%s)\r\n"), GetTabs(Indent), *FunctionName, *ParameterList);
+				HeaderOutput.Logf(TEXT("// %s%s(%s)\r\n"), FCString::Tab(Indent), *FunctionName, *ParameterList);
 
 				TFieldIterator<UProperty> ParamIt(Function);
 				ExportMCPDeclaration(HeaderOutput, FunctionName, ParamIt, CPF_Parm, Indent);
@@ -3649,8 +3568,8 @@ void FNativeClassHeaderGenerator::ExportEventParm(UFunction* Function, FUHTStrin
 	}
 
 	FString EventParmStructName = GetEventStructParamsName(Function->GetOuter(), *FunctionName);
-	HeaderOutput.Logf(TEXT("%sstruct %s\r\n"), GetTabs(Indent), *EventParmStructName);
-	HeaderOutput.Logf(TEXT("%s{\r\n"), GetTabs(Indent));
+	HeaderOutput.Logf(TEXT("%sstruct %s\r\n"), FCString::Tab(Indent), *EventParmStructName);
+	HeaderOutput.Logf(TEXT("%s{\r\n"), FCString::Tab(Indent));
 
 	for (TFieldIterator<UProperty> It(Function); It && (It->PropertyFlags&CPF_Parm); ++It)
 	{
@@ -3658,7 +3577,7 @@ void FNativeClassHeaderGenerator::ExportEventParm(UFunction* Function, FUHTStrin
 		ForwardDeclarations.Add(Prop);
 
 		FUHTStringBuilder PropertyText;
-		PropertyText.Log(GetTabs(Indent + 1));
+		PropertyText.Log(FCString::Tab(Indent + 1));
 
 		bool bEmitConst = Prop->HasAnyPropertyFlags(CPF_ConstParm) && Prop->IsA<UObjectProperty>();
 
@@ -3715,14 +3634,14 @@ void FNativeClassHeaderGenerator::ExportEventParm(UFunction* Function, FUHTStrin
 		if (bNeedsOutput)
 		{
 			check(Prop->ArrayDim == 1); // can't return arrays
-			HeaderOutput.Logf(TEXT("\r\n%s/** Constructor, intializes return property only **/\r\n"), GetTabs(Indent + 1));
-			HeaderOutput.Logf(TEXT("%s%s()\r\n"), GetTabs(Indent + 1), *EventParmStructName);
-			HeaderOutput.Logf(TEXT("%s%s %s(%s)\r\n"), GetTabs(Indent + 2), TEXT(":"), *Prop->GetName(), *GetNullParameterValue(Prop, false, true));
-			HeaderOutput.Logf(TEXT("%s{\r\n"), GetTabs(Indent + 1));
-			HeaderOutput.Logf(TEXT("%s}\r\n"), GetTabs(Indent + 1));
+			HeaderOutput.Logf(TEXT("\r\n%s/** Constructor, intializes return property only **/\r\n"), FCString::Tab(Indent + 1));
+			HeaderOutput.Logf(TEXT("%s%s()\r\n"), FCString::Tab(Indent + 1), *EventParmStructName);
+			HeaderOutput.Logf(TEXT("%s%s %s(%s)\r\n"), FCString::Tab(Indent + 2), TEXT(":"), *Prop->GetName(), *GetNullParameterValue(Prop, false, true));
+			HeaderOutput.Logf(TEXT("%s{\r\n"), FCString::Tab(Indent + 1));
+			HeaderOutput.Logf(TEXT("%s}\r\n"), FCString::Tab(Indent + 1));
 		}
 	}
-	HeaderOutput.Logf(TEXT("%s};\r\n"), GetTabs(Indent));
+	HeaderOutput.Logf(TEXT("%s};\r\n"), FCString::Tab(Indent));
 }
 
 void FNativeClassHeaderGenerator::ExportEventParms(FScope& Scope, const TArray<UFunction*>& InCallbackFunctions, FUHTStringBuilder& Output, int32 Indent, bool bOutputConstructor)
@@ -4359,10 +4278,8 @@ void FNativeClassHeaderGenerator::ExportNativeFunctions(FUnrealSourceFile& Sourc
 
 	// export the C++ stubs
 	
-	for (TFieldIterator<UFunction> FunctionIt(Class, EFieldIteratorFlags::ExcludeSuper); FunctionIt; ++FunctionIt)
+	for (UFunction* Function : TFieldRange<UFunction>(Class, EFieldIteratorFlags::ExcludeSuper))
 	{
-		UFunction*     Function = *FunctionIt;
-
 		if (!(Function->FunctionFlags & FUNC_Native))
 		{
 			continue;
@@ -4487,10 +4404,8 @@ TArray<UFunction*> FNativeClassHeaderGenerator::ExportCallbackFunctions(FUnrealS
 {
 	TArray<UFunction*> CallbackFunctions;
 
-	for (TFieldIterator<UFunction> FunctionIt(Class, EFieldIteratorFlags::ExcludeSuper); FunctionIt; ++FunctionIt)
+	for (UFunction* Function : TFieldRange<UFunction>(Class, EFieldIteratorFlags::ExcludeSuper))
 	{
-		UFunction*     Function = *FunctionIt;
-
 		if ((Function->FunctionFlags & FUNC_Event) && Function->GetSuperFunction() == nullptr)
 		{
 			CallbackFunctions.Add(Function);
@@ -4855,9 +4770,9 @@ FNativeClassHeaderGenerator::FNativeClassHeaderGenerator(
 			PreHeaderText.Logf(
 				TEXT("// Copyright 1998-2015 Epic Games, Inc. All Rights Reserved.\r\n")
 				TEXT("/*===========================================================================\r\n")
-				TEXT("    C++ class boilerplate exported from UnrealHeaderTool.\r\n")
-				TEXT("    This is automatically generated by the tools.\r\n")
-				TEXT("    DO NOT modify this manually! Edit the corresponding .h files instead!\r\n")
+				TEXT("\tC++ class boilerplate exported from UnrealHeaderTool.\r\n")
+				TEXT("\tThis is automatically generated by the tools.\r\n")
+				TEXT("\tDO NOT modify this manually! Edit the corresponding .h files instead!\r\n")
 				TEXT("===========================================================================*/\r\n")
 				TEXT("#pragma once\r\n")
 				TEXT("\r\n")
@@ -4965,8 +4880,7 @@ bool FNativeClassHeaderGenerator::SaveHeaderIfChanged(const TCHAR* HeaderPath, c
 		return false;
 	}
 
-	FString Tabified = Tabify(InNewHeaderContents);
-	const TCHAR* NewHeaderContents = *Tabified;
+	const TCHAR* NewHeaderContents = InNewHeaderContents;
 	static bool bTestedCmdLine = false;
 	if (!bTestedCmdLine)
 	{
@@ -5211,9 +5125,9 @@ void FNativeClassHeaderGenerator::ExportGeneratedCPP()
 	GeneratedCPPPreamble.Logf(
 		TEXT("// Copyright 1998-2015 Epic Games, Inc. All Rights Reserved.")					LINE_TERMINATOR
 		TEXT("/*===========================================================================")	LINE_TERMINATOR
-		TEXT("    Boilerplate C++ definitions for a single module.")							LINE_TERMINATOR
-		TEXT("    This is automatically generated by UnrealHeaderTool.")						LINE_TERMINATOR
-		TEXT("    DO NOT modify this manually! Edit the corresponding .h files instead!")		LINE_TERMINATOR
+		TEXT("\tBoilerplate C++ definitions for a single module.")								LINE_TERMINATOR
+		TEXT("\tThis is automatically generated by UnrealHeaderTool.")							LINE_TERMINATOR
+		TEXT("\tDO NOT modify this manually! Edit the corresponding .h files instead!")			LINE_TERMINATOR
 		TEXT("===========================================================================*/")	LINE_TERMINATOR
 																								LINE_TERMINATOR
 		);
@@ -5738,7 +5652,6 @@ ECompilationResult::Type UnrealHeaderTool_Main(const FString& ModuleInfoFilename
 	UE_LOG(LogCompile, Log, TEXT("Code generation took %.2f seconds"), GHeaderCodeGenTime);
 	UE_LOG(LogCompile, Log, TEXT("ScriptPlugin overhead was %.2f seconds"), GPluginOverheadTime);
 	UE_LOG(LogCompile, Log, TEXT("Macroize time was %.2f seconds"), GMacroizeTime);
-	UE_LOG(LogCompile, Log, TEXT("Tabify time was %.2f seconds"), GTabifyTime);
 
 	if (bWriteContents)
 	{
diff --git a/Engine/Source/Programs/UnrealHeaderTool/Private/HeaderProvider.cpp b/Engine/Source/Programs/UnrealHeaderTool/Private/HeaderProvider.cpp
--- a/Engine/Source/Programs/UnrealHeaderTool/Private/HeaderProvider.cpp
+++ b/Engine/Source/Programs/UnrealHeaderTool/Private/HeaderProvider.cpp
@@ -14,7 +14,7 @@ FHeaderProvider::FHeaderProvider(EHeaderProviderSourceType InType, const FString
 template <class Predicate>
 bool TryFindSourceFileWithPredicate(FUnrealSourceFile*& Out, Predicate Pred)
 {
-	for (auto MapPair : GUnrealSourceFilesMap)
+	for (const TPair<FString, TSharedRef<FUnrealSourceFile>>& MapPair : GUnrealSourceFilesMap)
 	{
 		if (Pred(MapPair.Value.Get()))
 		{
@@ -51,17 +51,17 @@ FUnrealSourceFile* FHeaderProvider::Resolve()
 			}
 			else
 			{
-				TryFindSourceFileWithPredicate(Cache,
-					[=](const FUnrealSourceFile& SourceFile)
-					{
-						return SourceFile.GetIncludePath() == Id;
-					})
-				|| TryFindSourceFileWithPredicate(Cache,
-					[=](const FUnrealSourceFile& SourceFile)
-					{
-						return SourceFile.GetFilename().EndsWith(TEXT("/") + Id) || SourceFile.GetFilename().EndsWith(TEXT("\\") + Id);
-					}
-				);
+				if (!TryFindSourceFileWithPredicate(Cache, [this](const FUnrealSourceFile& SourceFile) { return SourceFile.GetIncludePath() == Id; }))
+				{
+					FString SlashId     = TEXT("/") + Id;
+					FString BackslashId = TEXT("\\") + Id;
+					TryFindSourceFileWithPredicate(Cache,
+						[&SlashId, &BackslashId](const FUnrealSourceFile& SourceFile)
+						{
+							return SourceFile.GetFilename().EndsWith(SlashId) || SourceFile.GetFilename().EndsWith(BackslashId);
+						}
+					);
+				}
 			}
 		}
 
diff --git a/Engine/Source/Programs/UnrealHeaderTool/Private/ParserHelper.cpp b/Engine/Source/Programs/UnrealHeaderTool/Private/ParserHelper.cpp
--- a/Engine/Source/Programs/UnrealHeaderTool/Private/ParserHelper.cpp
+++ b/Engine/Source/Programs/UnrealHeaderTool/Private/ParserHelper.cpp
@@ -250,11 +250,11 @@ bool FAdvancedDisplayParameterHandler::CanMarkMore() const
 	return bUseNumber ? (NumberLeaveUnmarked > 0) : (0 != ParametersNames.Num());
 }
 
-TMap<UFunction*, TSharedRef<FFunctionData> > FFunctionData::FunctionDataMap;
+TMap<UFunction*, TUniqueObj<FFunctionData> > FFunctionData::FunctionDataMap;
 
 FFunctionData* FFunctionData::FindForFunction(UFunction* Function)
 {
-	auto* Output = FunctionDataMap.Find(Function);
+	TUniqueObj<FFunctionData>* Output = FunctionDataMap.Find(Function);
 
 	check(Output);
 
@@ -263,21 +263,21 @@ FFunctionData* FFunctionData::FindForFunction(UFunction* Function)
 
 FFunctionData* FFunctionData::Add(UFunction* Function)
 {
-	auto Output = FunctionDataMap.Add(Function, MakeShareable(new FFunctionData()));
+	TUniqueObj<FFunctionData>& Output = FunctionDataMap.Add(Function);
 
 	return &Output.Get();
 }
 
 FFunctionData* FFunctionData::Add(const FFuncInfo& FunctionInfo)
 {
-	auto Output = FunctionDataMap.Add(FunctionInfo.FunctionReference, MakeShareable(new FFunctionData(FunctionInfo)));
+	TUniqueObj<FFunctionData>& Output = FunctionDataMap.Emplace(FunctionInfo.FunctionReference, FunctionInfo);
 
 	return &Output.Get();
 }
 
 bool FFunctionData::TryFindForFunction(UFunction* Function, FFunctionData*& OutData)
 {
-	auto* Output = FunctionDataMap.Find(Function);
+	TUniqueObj<FFunctionData>* Output = FunctionDataMap.Find(Function);
 
 	if (!Output)
 	{
diff --git a/Engine/Source/Programs/UnrealHeaderTool/Private/ParserHelper.h b/Engine/Source/Programs/UnrealHeaderTool/Private/ParserHelper.h
--- a/Engine/Source/Programs/UnrealHeaderTool/Private/ParserHelper.h
+++ b/Engine/Source/Programs/UnrealHeaderTool/Private/ParserHelper.h
@@ -2,6 +2,7 @@
 
 #pragma once
 #include "ClassMaps.h"
+#include "UniqueObj.h"
 
 extern class FCompilerMetadataManager GScriptHelper;
 
@@ -1316,7 +1317,7 @@ public:
 	static bool TryFindForFunction(UFunction* Function, FFunctionData*& OutData);
 
 private:
-	static TMap<UFunction*, TSharedRef<FFunctionData> > FunctionDataMap;
+	static TMap<UFunction*, TUniqueObj<FFunctionData> > FunctionDataMap;
 };
 
 /**
diff --git a/Engine/Source/Runtime/Core/Private/Containers/String.cpp b/Engine/Source/Runtime/Core/Private/Containers/String.cpp
--- a/Engine/Source/Runtime/Core/Private/Containers/String.cpp
+++ b/Engine/Source/Runtime/Core/Private/Containers/String.cpp
@@ -543,7 +543,7 @@ int32 FString::ParseIntoArray( TArray<FString>& OutArray, const TCHAR* pchDelim,
 			{
 				new (OutArray) FString(At-Start,Start);
 			}
-			Start += DelimLength + (At-Start);
+			Start = At + DelimLength;
 		}
 		if (!InCullEmpty || *Start)
 		{
diff --git a/Engine/Source/Runtime/Core/Private/Misc/CString.cpp b/Engine/Source/Runtime/Core/Private/Misc/CString.cpp
--- a/Engine/Source/Runtime/Core/Private/Misc/CString.cpp
+++ b/Engine/Source/Runtime/Core/Private/Misc/CString.cpp
@@ -2,11 +2,34 @@
 
 #include "CorePrivatePCH.h"
 
-static struct FForceInitAtBootTCStringSpcHelperAnsi : public TForceInitAtBoot<TCStringSpcHelper<ANSICHAR>>
-{} FForceInitAtBootTCStringSpcHelperAnsi;
+// 4 lines of 64 chars each, plus a null terminator
+template <>
+const ANSICHAR TCStringSpcHelper<ANSICHAR>::SpcArray[MAX_SPACES + 1] =
+	"                                                                "
+	"                                                                "
+	"                                                                "
+	"                                                               ";
 
-static struct FForceInitAtBootTCStringSpcHelperWide : public TForceInitAtBoot<TCStringSpcHelper<WIDECHAR>>
-{} FForceInitAtBootTCStringSpcHelperWide;
+template <>
+const WIDECHAR TCStringSpcHelper<WIDECHAR>::SpcArray[MAX_SPACES + 1] =
+	L"                                                                "
+	L"                                                                "
+	L"                                                                "
+	L"                                                               ";
+
+template <>
+const ANSICHAR TCStringSpcHelper<ANSICHAR>::TabArray[MAX_TABS + 1] = 
+	"\t\t\t\t\t\t\t\t\t\t\t\t\t\t\t\t\t\t\t\t\t\t\t\t\t\t\t\t\t\t\t\t\t\t\t\t\t\t\t\t\t\t\t\t\t\t\t\t\t\t\t\t\t\t\t\t\t\t\t\t\t\t\t\t"
+	"\t\t\t\t\t\t\t\t\t\t\t\t\t\t\t\t\t\t\t\t\t\t\t\t\t\t\t\t\t\t\t\t\t\t\t\t\t\t\t\t\t\t\t\t\t\t\t\t\t\t\t\t\t\t\t\t\t\t\t\t\t\t\t\t"
+	"\t\t\t\t\t\t\t\t\t\t\t\t\t\t\t\t\t\t\t\t\t\t\t\t\t\t\t\t\t\t\t\t\t\t\t\t\t\t\t\t\t\t\t\t\t\t\t\t\t\t\t\t\t\t\t\t\t\t\t\t\t\t\t\t"
+	"\t\t\t\t\t\t\t\t\t\t\t\t\t\t\t\t\t\t\t\t\t\t\t\t\t\t\t\t\t\t\t\t\t\t\t\t\t\t\t\t\t\t\t\t\t\t\t\t\t\t\t\t\t\t\t\t\t\t\t\t\t\t\t";
+
+template <>
+const WIDECHAR TCStringSpcHelper<WIDECHAR>::TabArray[MAX_TABS + 1] = 
+	L"\t\t\t\t\t\t\t\t\t\t\t\t\t\t\t\t\t\t\t\t\t\t\t\t\t\t\t\t\t\t\t\t\t\t\t\t\t\t\t\t\t\t\t\t\t\t\t\t\t\t\t\t\t\t\t\t\t\t\t\t\t\t\t\t"
+	L"\t\t\t\t\t\t\t\t\t\t\t\t\t\t\t\t\t\t\t\t\t\t\t\t\t\t\t\t\t\t\t\t\t\t\t\t\t\t\t\t\t\t\t\t\t\t\t\t\t\t\t\t\t\t\t\t\t\t\t\t\t\t\t\t"
+	L"\t\t\t\t\t\t\t\t\t\t\t\t\t\t\t\t\t\t\t\t\t\t\t\t\t\t\t\t\t\t\t\t\t\t\t\t\t\t\t\t\t\t\t\t\t\t\t\t\t\t\t\t\t\t\t\t\t\t\t\t\t\t\t\t"
+	L"\t\t\t\t\t\t\t\t\t\t\t\t\t\t\t\t\t\t\t\t\t\t\t\t\t\t\t\t\t\t\t\t\t\t\t\t\t\t\t\t\t\t\t\t\t\t\t\t\t\t\t\t\t\t\t\t\t\t\t\t\t\t\t";
 
 bool FToBoolHelper::FromCStringAnsi( const ANSICHAR* String )
 {
diff --git a/Engine/Source/Runtime/Core/Public/Containers/Array.h b/Engine/Source/Runtime/Core/Public/Containers/Array.h
--- a/Engine/Source/Runtime/Core/Public/Containers/Array.h
+++ b/Engine/Source/Runtime/Core/Public/Containers/Array.h
@@ -404,7 +404,7 @@ public:
 	template <typename OtherElementType, typename OtherAllocator>
 	explicit TArray(const TArray<OtherElementType, OtherAllocator>& Other)
 	{
-		CopyToEmpty(Other);
+		CopyToEmpty(Other, 0, 0);
 	}
 
 	/**
@@ -414,7 +414,7 @@ public:
 	 */
 	TArray(const TArray& Other)
 	{
-		CopyToEmpty(Other);
+		CopyToEmpty(Other, 0, 0);
 	}
 
 	/**
@@ -426,7 +426,7 @@ public:
 	 */
 	TArray(const TArray& Other, int32 ExtraSlack)
 	{
-		CopyToEmpty(Other, ExtraSlack);
+		CopyToEmpty(Other, 0, ExtraSlack);
 	}
 
 	/**
@@ -441,7 +441,7 @@ public:
 	TArray& operator=(const TArray<ElementType, OtherAllocator>& Other)
 	{
 		DestructItems(GetData(), ArrayNum);
-		CopyToEmpty(Other);
+		CopyToEmpty(Other, ArrayMax, 0);
 		return *this;
 	}
 
@@ -456,7 +456,7 @@ public:
 		if (this != &Other)
 		{
 			DestructItems(GetData(), ArrayNum);
-			CopyToEmpty(Other);
+			CopyToEmpty(Other, ArrayMax, 0);
 		}
 		return *this;
 	}
@@ -471,7 +471,7 @@ private:
 	 * @param FromArray Array to move from.
 	 */
 	template <typename FromArrayType, typename ToArrayType>
-	static FORCEINLINE typename TEnableIf<UE4Array_Private::TCanMoveTArrayPointersBetweenArrayTypes<FromArrayType, ToArrayType>::Value>::Type MoveOrCopy(ToArrayType& ToArray, FromArrayType& FromArray)
+	static FORCEINLINE typename TEnableIf<UE4Array_Private::TCanMoveTArrayPointersBetweenArrayTypes<FromArrayType, ToArrayType>::Value>::Type MoveOrCopy(ToArrayType& ToArray, FromArrayType& FromArray, int32 PrevMax)
 	{
 		ToArray.AllocatorInstance.MoveToEmpty(FromArray.AllocatorInstance);
 
@@ -481,6 +481,22 @@ private:
 		FromArray.ArrayMax = 0;
 	}
 
+	/**
+	 * Moves or copies array. Depends on the array type traits.
+	 *
+	 * This override copies.
+	 *
+	 * @param ToArray Array to move into.
+	 * @param FromArray Array to move from.
+	 * @param ExtraSlack Tells how much extra memory should be preallocated
+	 *                   at the end of the array in the number of elements.
+	 */
+	template <typename FromArrayType, typename ToArrayType>
+	static FORCEINLINE typename TEnableIf<!UE4Array_Private::TCanMoveTArrayPointersBetweenArrayTypes<FromArrayType, ToArrayType>::Value>::Type MoveOrCopy(ToArrayType& ToArray, FromArrayType& FromArray, int32 PrevMax)
+	{
+		ToArray.CopyToEmpty(FromArray, PrevMax, 0);
+	}
+
 	/**
 	 * Moves or copies array. Depends on the array type traits.
 	 *
@@ -492,9 +508,9 @@ private:
 	 *                   at the end of the array in the number of elements.
 	 */
 	template <typename FromArrayType, typename ToArrayType>
-	static FORCEINLINE typename TEnableIf<UE4Array_Private::TCanMoveTArrayPointersBetweenArrayTypes<FromArrayType, ToArrayType>::Value>::Type MoveOrCopy(ToArrayType& ToArray, FromArrayType& FromArray, int32 ExtraSlack)
+	static FORCEINLINE typename TEnableIf<UE4Array_Private::TCanMoveTArrayPointersBetweenArrayTypes<FromArrayType, ToArrayType>::Value>::Type MoveOrCopyWithSlack(ToArrayType& ToArray, FromArrayType& FromArray, int32 PrevMax, int32 ExtraSlack)
 	{
-		MoveOrCopy(ToArray, FromArray);
+		MoveOrCopy(ToArray, FromArray, PrevMax);
 
 		ToArray.Reserve(ToArray.ArrayNum + ExtraSlack);
 	}
@@ -510,9 +526,9 @@ private:
 	 *                   at the end of the array in the number of elements.
 	 */
 	template <typename FromArrayType, typename ToArrayType>
-	static FORCEINLINE typename TEnableIf<!UE4Array_Private::TCanMoveTArrayPointersBetweenArrayTypes<FromArrayType, ToArrayType>::Value>::Type MoveOrCopy(ToArrayType& ToArray, FromArrayType& FromArray, int32 ExtraSlack = 0)
+	static FORCEINLINE typename TEnableIf<!UE4Array_Private::TCanMoveTArrayPointersBetweenArrayTypes<FromArrayType, ToArrayType>::Value>::Type MoveOrCopyWithSlack(ToArrayType& ToArray, FromArrayType& FromArray, int32 PrevMax, int32 ExtraSlack)
 	{
-		ToArray.CopyToEmpty(FromArray, ExtraSlack);
+		ToArray.CopyToEmpty(FromArray, PrevMax, ExtraSlack);
 	}
 
 public:
@@ -523,7 +539,7 @@ public:
 	 */
 	FORCEINLINE TArray(TArray&& Other)
 	{
-		MoveOrCopy(*this, Other);
+		MoveOrCopy(*this, Other, 0);
 	}
 
 	/**
@@ -534,7 +550,7 @@ public:
 	template <typename OtherElementType, typename OtherAllocator>
 	FORCEINLINE explicit TArray(TArray<OtherElementType, OtherAllocator>&& Other)
 	{
-		MoveOrCopy(*this, Other);
+		MoveOrCopy(*this, Other, 0);
 	}
 
 	/**
@@ -551,7 +567,7 @@ public:
 		// to tell if they're compatible with the current one.  Probably going to be a pretty
 		// rare requirement anyway.
 
-		MoveOrCopy(*this, Other, ExtraSlack);
+		MoveOrCopyWithSlack(*this, Other, 0, ExtraSlack);
 	}
 
 	/**
@@ -564,7 +580,7 @@ public:
 		if (this != &Other)
 		{
 			DestructItems(GetData(), ArrayNum);
-			MoveOrCopy(*this, Other);
+			MoveOrCopy(*this, Other, ArrayMax);
 		}
 		return *this;
 	}
@@ -2275,17 +2291,21 @@ private:
 	 * data in question does not need a constructor.
 	 *
 	 * @param Source The source array to copy
-	 * @param ExtraSlack (Optional) Additional amount of memory to allocate at
+	 * @param PrevMax The previous allocated size
+	 * @param ExtraSlack Additional amount of memory to allocate at
 	 *                   the end of the buffer. Counted in elements. Zero by
 	 *                   default.
 	 */
 	template <typename OtherElementType, typename OtherAllocator>
-	void CopyToEmpty(const TArray<OtherElementType, OtherAllocator>& Source, int32 ExtraSlack = 0)
+	void CopyToEmpty(const TArray<OtherElementType, OtherAllocator>& Source, int32 PrevMax, int32 ExtraSlack)
 	{
 		check(ExtraSlack >= 0);
 
 		int32 SourceCount = Source.Num();
-		AllocatorInstance.ResizeAllocation(0, SourceCount + ExtraSlack, sizeof(ElementType));
+		if (SourceCount + ExtraSlack != PrevMax)
+		{
+			AllocatorInstance.ResizeAllocation(0, SourceCount + ExtraSlack, sizeof(ElementType));
+		}
 
 		ConstructItems<ElementType>(GetData(), Source.GetData(), SourceCount);
 
diff --git a/Engine/Source/Runtime/Core/Public/Containers/Set.h b/Engine/Source/Runtime/Core/Public/Containers/Set.h
--- a/Engine/Source/Runtime/Core/Public/Containers/Set.h
+++ b/Engine/Source/Runtime/Core/Public/Containers/Set.h
@@ -328,8 +328,10 @@ public:
 	/** Compacts the allocated elements into a contiguous range. */
 	FORCEINLINE void Compact()
 	{
-		Elements.Compact();
-		Rehash();
+		if (Elements.Compact())
+		{
+			Rehash();
+		}
 	}
 
 	/** Preallocates enough memory to contain Number elements */
diff --git a/Engine/Source/Runtime/Core/Public/Containers/SparseArray.h b/Engine/Source/Runtime/Core/Public/Containers/SparseArray.h
--- a/Engine/Source/Runtime/Core/Public/Containers/SparseArray.h
+++ b/Engine/Source/Runtime/Core/Public/Containers/SparseArray.h
@@ -340,14 +340,17 @@ public:
 	}
 
 	/** Compacts the allocated elements into a contiguous index range. */
-	void Compact()
+	/** Returns true if any elements were relocated, false otherwise. */
+	bool Compact()
 	{
 		int32 NumFree = NumFreeIndices;
 		if (NumFree == 0)
 		{
-			return;
+			return false;
 		}
 
+		bool bResult = false;
+
 		FElementOrFreeListLink* ElementData = Data.GetData();
 
 		int32 EndIndex    = Data.Num();
@@ -367,6 +370,8 @@ public:
 
 				RelocateConstructItems<FElementOrFreeListLink>(ElementData + FreeIndex, ElementData + EndIndex, 1);
 				AllocationFlags[FreeIndex] = true;
+
+				bResult = true;
 			}
 
 			FreeIndex = NextFreeIndex;
@@ -377,6 +382,8 @@ public:
 
 		NumFreeIndices = 0;
 		FirstFreeIndex = -1;
+
+		return bResult;
 	}
 
 	/** Sorts the elements using the provided comparison class. */
diff --git a/Engine/Source/Runtime/Core/Public/Misc/CString.h b/Engine/Source/Runtime/Core/Public/Misc/CString.h
--- a/Engine/Source/Runtime/Core/Public/Misc/CString.h
+++ b/Engine/Source/Runtime/Core/Public/Misc/CString.h
@@ -193,6 +193,15 @@ struct TCString
 	 */
 	static const CharType* Spc( int32 NumSpaces );
 
+	/**
+	 * Returns a static string that is filled with a variable number of tabs
+	 *
+	 * @param NumTabs Number of tabs to put into the string, max of 255
+	 * 
+	 * @return The string of NumTabs tabs.
+	 */
+	static const CharType* Tab( int32 NumTabs );
+
 	/**
 	 * Find string in string, case insensitive, requires non-alphanumeric lead-in.
 	 */
@@ -319,7 +328,7 @@ struct TCString
 
 typedef TCString<TCHAR>    FCString;
 typedef TCString<ANSICHAR> FCStringAnsi;
-typedef TCString<WIDECHAR>  FCStringWide;
+typedef TCString<WIDECHAR> FCStringWide;
 
 /*-----------------------------------------------------------------------------
 	generic TCString implementations
@@ -328,61 +337,28 @@ typedef TCString<WIDECHAR>  FCStringWide;
 template <typename CharType>
 struct TCStringSpcHelper
 {
-	friend struct TCString<CharType>;
-	friend struct TForceInitAtBoot<TCStringSpcHelper<CharType>>;
-
-private:
 	/** Number of characters to be stored in string. */
 	static const int32 MAX_SPACES = 255;
 
-	/** String storage. */
-	CharType StaticString[ MAX_SPACES + 1 ];
-
-	/** Constructor, fills the string storage with spaces. */
-	TCStringSpcHelper()
-	{
-		CharType SpcChar = LITERAL(CharType, ' ');
+	/** Number of tabs to be stored in string. */
+	static const int32 MAX_TABS = 255;
 
-		for( int32 Cx = 0; Cx < MAX_SPACES; ++Cx )
-		{
-			StaticString[Cx] = SpcChar;
-		}
-		StaticString[MAX_SPACES] = 0;
-	}
-	
-	/**
-	 * @param NumSpaces - Number of characters to put into the string, max of 255
-	 *
-	 * @return The string filled with spaces
-	 */
-	const CharType* Spc( int32 NumSpaces ) const
-	{
-		check( NumSpaces >= 0 );
-		check( NumSpaces <= MAX_SPACES );
-
-		const uint32 SpcPos = MAX_SPACES - NumSpaces;
-		return &StaticString[SpcPos];
-	}
-
-	/** Gets TCStringSpcHelper singleton instance. */
-	static TCStringSpcHelper& Get()
-	{
-		static TCStringSpcHelper StringHelper;
-		return StringHelper;
-	}
+	static CORE_API const CharType SpcArray[MAX_SPACES + 1];
+	static CORE_API const CharType TabArray[MAX_TABS + 1];
 };
 
-/**
- * Returns a static string that is filled with a variable number of spaces
- *
- * @param NumSpaces Number of spaces to put into the string, max of 255
- * 
- * @return The string of NumSpaces spaces.
- */
 template <typename T>
 const typename TCString<T>::CharType* TCString<T>::Spc( int32 NumSpaces )
 {
-	return TCStringSpcHelper<T>::Get().Spc( NumSpaces );
+	check(NumSpaces >= 0 && NumSpaces <= TCStringSpcHelper<T>::MAX_SPACES);
+	return TCStringSpcHelper<T>::SpcArray + TCStringSpcHelper<T>::MAX_SPACES - NumSpaces;
+}
+
+template <typename T>
+const typename TCString<T>::CharType* TCString<T>::Tab( int32 NumTabs )
+{
+	check(NumTabs >= 0 && NumTabs <= TCStringSpcHelper<T>::MAX_TABS);
+	return TCStringSpcHelper<T>::TabArray + TCStringSpcHelper<T>::MAX_TABS - NumTabs;
 }
 
 //
