Fixed hot reload bug - was not using the normal cpp name for native types. Dynamic types still need to use their overridden name. Also fixed a typo and a comment that waas out of date.

diff --git a/Engine/Source/Programs/UnrealHeaderTool/Private/CodeGenerator.cpp b/Engine/Source/Programs/UnrealHeaderTool/Private/CodeGenerator.cpp
--- a/Engine/Source/Programs/UnrealHeaderTool/Private/CodeGenerator.cpp
+++ b/Engine/Source/Programs/UnrealHeaderTool/Private/CodeGenerator.cpp
@@ -1442,11 +1442,11 @@ void FNativeClassHeaderGenerator::ExportNativeGeneratedInitCode(FClass* Class, F
 	SingletonName.ReplaceInline(TEXT("()"), TEXT(""), ESearchCase::CaseSensitive); // function address
 
 	{	
-		FString OverridenClassName = *FNativeClassHeaderGenerator::GetOverriddenName(Class);
+		FString OverriddenClassName = *FNativeClassHeaderGenerator::GetOverriddenName(Class);
 		const TCHAR* ClassNameCPP = NameLookupCPP.GetNameCPP(Class);
 
 		GeneratedFunctionText.Logf(TEXT("\tstatic FCompiledInDefer Z_CompiledInDefer_UClass_%s(%s, &%s::StaticClass, TEXT(\"%s\"), %s);\r\n"),
-			ClassNameCPP, *SingletonName, ClassNameCPP, *OverridenClassName,
+			ClassNameCPP, *SingletonName, ClassNameCPP, bIsDynamic ? ClassNameCPP : *OverriddenClassName,
 			bIsDynamic ? TEXT("true") : TEXT("false"));
 		
 		// Append base class' CRC at the end of the generated code, this will force update derived classes
@@ -1468,7 +1468,7 @@ void FNativeClassHeaderGenerator::ExportNativeGeneratedInitCode(FClass* Class, F
 		}
 		else
 		{
-			GeneratedPackageCPP.Logf(TEXT("\tIMPLEMENT_DYNAMIC_CLASS(%s, TEXT(\"%s\"), %u);\r\n"), ClassNameCPP, *OverridenClassName, ClassCrc);
+			GeneratedPackageCPP.Logf(TEXT("\tIMPLEMENT_DYNAMIC_CLASS(%s, TEXT(\"%s\"), %u);\r\n"), ClassNameCPP, *OverriddenClassName, ClassCrc);
 		}
 	}
 }
diff --git a/Engine/Source/Runtime/CoreUObject/Private/UObject/UObjectBase.cpp b/Engine/Source/Runtime/CoreUObject/Private/UObject/UObjectBase.cpp
--- a/Engine/Source/Runtime/CoreUObject/Private/UObject/UObjectBase.cpp
+++ b/Engine/Source/Runtime/CoreUObject/Private/UObject/UObjectBase.cpp
@@ -686,7 +686,6 @@ void UObjectCompiledInDefer(UClass *(*InRegister)(), UClass *(*InStaticClass)(),
 	}
 	else
 	{
-		// Name is with the CPP prefix, skip it when adding to the map
 		FClassConstructFunctions ClassFunctions;
 		ClassFunctions.ZConstructFn = InRegister;
 		ClassFunctions.StaticClassFn = InStaticClass;
