Fix for my fix...

diff --git a/Engine/Source/Programs/UnrealHeaderTool/Private/CodeGenerator.cpp b/Engine/Source/Programs/UnrealHeaderTool/Private/CodeGenerator.cpp
--- a/Engine/Source/Programs/UnrealHeaderTool/Private/CodeGenerator.cpp
+++ b/Engine/Source/Programs/UnrealHeaderTool/Private/CodeGenerator.cpp
@@ -1446,7 +1446,7 @@ void FNativeClassHeaderGenerator::ExportNativeGeneratedInitCode(FClass* Class, F
 		const TCHAR* ClassNameCPP = NameLookupCPP.GetNameCPP(Class);
 
 		GeneratedFunctionText.Logf(TEXT("\tstatic FCompiledInDefer Z_CompiledInDefer_UClass_%s(%s, &%s::StaticClass, TEXT(\"%s\"), %s);\r\n"),
-			ClassNameCPP, *SingletonName, ClassNameCPP, bIsDynamic ? ClassNameCPP : *OverriddenClassName,
+			ClassNameCPP, *SingletonName, ClassNameCPP, bIsDynamic ? *OverriddenClassName : ClassNameCPP,
 			bIsDynamic ? TEXT("true") : TEXT("false"));
 		
 		// Append base class' CRC at the end of the generated code, this will force update derived classes
