Fixing UnrealHeaderTool Mac compile errors post integration

diff --git a/Engine/Source/Programs/UnrealHeaderTool/Private/HeaderParser.cpp b/Engine/Source/Programs/UnrealHeaderTool/Private/HeaderParser.cpp
--- a/Engine/Source/Programs/UnrealHeaderTool/Private/HeaderParser.cpp
+++ b/Engine/Source/Programs/UnrealHeaderTool/Private/HeaderParser.cpp
@@ -1111,8 +1111,8 @@ void AddModuleRelativePathToMetadata(UField* Type, TMap<FName, FString> &MetaDat
 //
 UEnum* FHeaderParser::CompileEnum()
 {
-	FUnrealSourceFile* CurrentSourceFile = GetCurrentSourceFile();
-	auto Scope = CurrentSourceFile->GetScope();
+	FUnrealSourceFile* CurrentSrcFile = GetCurrentSourceFile();
+	auto Scope = CurrentSrcFile->GetScope();
 
 	CheckAllow( TEXT("'Enum'"), ALLOW_TypeDecl );
 
@@ -1180,11 +1180,11 @@ UEnum* FHeaderParser::CompileEnum()
 
 	ParseFieldMetaData(EnumToken.MetaData, EnumToken.Identifier);
 	// Create enum definition.
-	UEnum* Enum = new(EC_InternalUseOnlyConstructor, CurrentSourceFile->GetPackage(), EnumToken.Identifier, RF_Public) UEnum(FObjectInitializer());
+	UEnum* Enum = new(EC_InternalUseOnlyConstructor, CurrentSrcFile->GetPackage(), EnumToken.Identifier, RF_Public) UEnum(FObjectInitializer());
 	Scope->AddType(Enum);
 
-	AddTypeDefinition(UHTMakefile, CurrentSourceFile, Enum, InputLine);
-	UHTMakefile.AddEnum(CurrentSourceFile, Enum);
+	AddTypeDefinition(UHTMakefile, CurrentSrcFile, Enum, InputLine);
+	UHTMakefile.AddEnum(CurrentSrcFile, Enum);
 	// Validate the metadata for the enum
 	ValidateMetaDataFormat(Enum, EnumToken.MetaData);
 
@@ -1204,7 +1204,7 @@ UEnum* FHeaderParser::CompileEnum()
 		}
 
 		GEnumUnderlyingTypes.Add(Enum, CPT_Byte);
-		UHTMakefile.AddGEnumUnderlyingType(CurrentSourceFile, Enum, CPT_Byte);
+		UHTMakefile.AddGEnumUnderlyingType(CurrentSrcFile, Enum, CPT_Byte);
 	}
 
 	// Get opening brace.
@@ -1692,8 +1692,8 @@ EAccessSpecifier FHeaderParser::ParseAccessProtectionSpecifier(FToken& Token)
  */
 UScriptStruct* FHeaderParser::CompileStructDeclaration(FClasses& AllClasses)
 {
-	FUnrealSourceFile* CurrentSourceFile = GetCurrentSourceFile();
-	auto Scope = CurrentSourceFile->GetScope();
+	FUnrealSourceFile* CurrentSrcFile = GetCurrentSourceFile();
+	auto Scope = CurrentSrcFile->GetScope();
 
 	// Make sure structs can be declared here.
 	CheckAllow( TEXT("'struct'"), ALLOW_TypeDecl );//@TODO: UCREMOVAL: After the switch: Make this require global scope
@@ -1908,13 +1908,13 @@ UScriptStruct* FHeaderParser::CompileStructDeclaration(FClasses& AllClasses)
 		StructFlags |= (BaseStruct->StructFlags&STRUCT_Inherit);
 	}
 	// Create.
-	UScriptStruct* Struct = new(EC_InternalUseOnlyConstructor, CurrentSourceFile->GetPackage(), *EffectiveStructName, RF_Public) UScriptStruct(FObjectInitializer(), BaseStruct);
-	UHTMakefile.AddScriptStruct(CurrentSourceFile, Struct);
+	UScriptStruct* Struct = new(EC_InternalUseOnlyConstructor, CurrentSrcFile->GetPackage(), *EffectiveStructName, RF_Public) UScriptStruct(FObjectInitializer(), BaseStruct);
+	UHTMakefile.AddScriptStruct(CurrentSrcFile, Struct);
 
 	Scope->AddType(Struct);
-	FScope::AddTypeScope(Struct, &CurrentSourceFile->GetScope().Get(), CurrentSourceFile, UHTMakefile);
+	FScope::AddTypeScope(Struct, &CurrentSrcFile->GetScope().Get(), CurrentSrcFile, UHTMakefile);
 
-	AddTypeDefinition(UHTMakefile, CurrentSourceFile, Struct, InputLine);
+	AddTypeDefinition(UHTMakefile, CurrentSrcFile, Struct, InputLine);
 
 	AddModuleRelativePathToMetadata(Struct, MetaData);
 
@@ -1957,8 +1957,8 @@ UScriptStruct* FHeaderParser::CompileStructDeclaration(FClasses& AllClasses)
 		StructToken.Struct = Struct;
 
 		// add this struct to the compiler's persistent tracking system
-		FClassMetaData* ClassMetaData = GScriptHelper.AddClassData(StructToken.Struct, UHTMakefile, CurrentSourceFile);
-		UHTMakefile.AddGScriptHelperEntry(CurrentSourceFile, Struct, ClassMetaData);
+		FClassMetaData* ClassMetaData = GScriptHelper.AddClassData(StructToken.Struct, UHTMakefile, CurrentSrcFile);
+		UHTMakefile.AddGScriptHelperEntry(CurrentSrcFile, Struct, ClassMetaData);
 	}
 
 	int32 SavedLineNumber = InputLine;
@@ -2557,7 +2557,7 @@ void FHeaderParser::VerifyRepNotifyCallbacks( UClass* TargetClass )
 void FHeaderParser::CompileDirective(FClasses& AllClasses)
 {
 	FUnrealSourceFile* CurrentSourceFilePtr = GetCurrentSourceFile();
-	TSharedRef<FUnrealSourceFile> CurrentSourceFile = CurrentSourceFilePtr->AsShared();
+	TSharedRef<FUnrealSourceFile> CurrentSrcFile = CurrentSourceFilePtr->AsShared();
 	FToken Directive;
 
 	int32 LineAtStartOfDirective = InputLine;
@@ -2592,7 +2592,7 @@ void FHeaderParser::CompileDirective(FClasses& AllClasses)
 	}
 	else if (Directive.Matches(TEXT("include")))
 	{
-		FString ExpectedHeaderName = CurrentSourceFile->GetGeneratedHeaderFilename();
+		FString ExpectedHeaderName = CurrentSrcFile->GetGeneratedHeaderFilename();
 		FToken IncludeName;
 		if (GetToken(IncludeName) && (IncludeName.TokenType == TOKEN_Const) && (IncludeName.Type == CPT_String))
 		{
@@ -3955,7 +3955,7 @@ UProperty* FHeaderParser::GetVarNameAndDim
 {
 	check(Scope);
 
-	FUnrealSourceFile* CurrentSourceFile = GetCurrentSourceFile();
+	FUnrealSourceFile* CurrentSrcFile = GetCurrentSourceFile();
 	EObjectFlags ObjectFlags = RF_Public;
 	if (VariableCategory == EVariableCategory::Member && CurrentAccessSpecifier == ACCESS_Private)
 	{
@@ -4186,7 +4186,7 @@ UProperty* FHeaderParser::GetVarNameAndDim
 		if (VarProperty.ArrayType == EArrayType::Dynamic)
 		{
 			Array       = new (EC_InternalUseOnlyConstructor, Scope, PropertyName, ObjectFlags) UArrayProperty(FObjectInitializer());
-			UHTMakefile.AddArrayProperty(CurrentSourceFile, Array);
+			UHTMakefile.AddArrayProperty(CurrentSrcFile, Array);
 			NewScope    = Array;
 			ObjectFlags = RF_Public;
 		}
@@ -4197,13 +4197,13 @@ UProperty* FHeaderParser::GetVarNameAndDim
 		else if (VarProperty.MapKeyProp.IsValid())
 		{
 			Map               = new (EC_InternalUseOnlyConstructor, Scope, PropertyName, ObjectFlags) UMapProperty(FObjectInitializer());
-			UHTMakefile.AddMapProperty(CurrentSourceFile, Map);
+			UHTMakefile.AddMapProperty(CurrentSrcFile, Map);
 			NewScope          = Map;
 			ObjectFlags       = RF_Public;
-			NewMapKeyProperty = CreateVariableProperty(*VarProperty.MapKeyProp, NewScope, *(PropertyName.ToString() + TEXT("_Key")), ObjectFlags, VariableCategory, UHTMakefile, CurrentSourceFile);
+			NewMapKeyProperty = CreateVariableProperty(*VarProperty.MapKeyProp, NewScope, *(PropertyName.ToString() + TEXT("_Key")), ObjectFlags, VariableCategory, UHTMakefile, CurrentSrcFile);
 		}
 
-		NewProperty = CreateVariableProperty(VarProperty, NewScope, PropertyName, ObjectFlags, VariableCategory, UHTMakefile, CurrentSourceFile);
+		NewProperty = CreateVariableProperty(VarProperty, NewScope, PropertyName, ObjectFlags, VariableCategory, UHTMakefile, CurrentSrcFile);
 
 		auto PropagateFlags = [](uint64 FlagsToPropagate, FPropertyBase& From, UProperty* To) {
 			// Copy some of the property flags to the inner property.
@@ -4265,7 +4265,7 @@ UProperty* FHeaderParser::GetVarNameAndDim
 	VarProperty.TokenProperty = NewProperty;
 	FClassMetaData* ScopeData = GScriptHelper.FindClassData(Scope);
 	check(ScopeData);
-	ScopeData->AddProperty(VarProperty, UHTMakefile, CurrentSourceFile);
+	ScopeData->AddProperty(VarProperty, UHTMakefile, CurrentSrcFile);
 
 	// if we had any metadata, add it to the class
 	AddMetaDataToClassData(VarProperty.TokenProperty, VarProperty.MetaData);
@@ -4694,14 +4694,14 @@ bool FHeaderParser::SafeMatchSymbol( const TCHAR* Match )
 
 FClass* FHeaderParser::ParseClassNameDeclaration(FClasses& AllClasses, FString& DeclaredClassName, FString& RequiredAPIMacroIfPresent)
 {
-	FUnrealSourceFile* CurrentSourceFile = GetCurrentSourceFile();
+	FUnrealSourceFile* CurrentSrcFile = GetCurrentSourceFile();
 	ParseNameWithPotentialAPIMacroPrefix(/*out*/ DeclaredClassName, /*out*/ RequiredAPIMacroIfPresent, TEXT("class"));
 
 	FClass* FoundClass = AllClasses.FindClass(*GetClassNameWithPrefixRemoved(*DeclaredClassName));
 	check(FoundClass);
 
-	FClassMetaData* ClassMetaData = GScriptHelper.AddClassData(FoundClass, UHTMakefile, CurrentSourceFile);
-	UHTMakefile.AddGScriptHelperEntry(CurrentSourceFile, FoundClass, ClassMetaData);
+	FClassMetaData* ClassMetaData = GScriptHelper.AddClassData(FoundClass, UHTMakefile, CurrentSrcFile);
+	UHTMakefile.AddGScriptHelperEntry(CurrentSrcFile, FoundClass, ClassMetaData);
 
 	// Get parent class.
 	bool bSpecifiesParentClass = false;
@@ -4788,7 +4788,7 @@ FClass* FHeaderParser::ParseClassNameDeclaration(FClasses& AllClasses, FString&
 
 void FHeaderParser::HandleOneInheritedClass(FClasses& AllClasses, UClass* Class, FString InterfaceName)
 {
-	FUnrealSourceFile* CurrentSourceFile = GetCurrentSourceFile();
+	FUnrealSourceFile* CurrentSrcFile = GetCurrentSourceFile();
 	// Check for UInterface derived interface inheritance
 	if (UClass* Interface = AllClasses.FindScriptClass(InterfaceName))
 	{
@@ -4806,7 +4806,7 @@ void FHeaderParser::HandleOneInheritedClass(FClasses& AllClasses, UClass* Class,
 		{
 			FClassMetaData* ClassData = GScriptHelper.FindClassData(Class);
 			check(ClassData);
-			ClassData->AddInheritanceParent(Interface, UHTMakefile, CurrentSourceFile);
+			ClassData->AddInheritanceParent(Interface, UHTMakefile, CurrentSrcFile);
 		}
 	}
 	else
@@ -4814,7 +4814,7 @@ void FHeaderParser::HandleOneInheritedClass(FClasses& AllClasses, UClass* Class,
 		// Non-UObject inheritance
 		FClassMetaData* ClassData = GScriptHelper.FindClassData(Class);
 		check(ClassData);
-		ClassData->AddInheritanceParent(InterfaceName, UHTMakefile, CurrentSourceFile);
+		ClassData->AddInheritanceParent(InterfaceName, UHTMakefile, CurrentSrcFile);
 	}
 }
 
@@ -5059,7 +5059,7 @@ bool FHeaderParser::TryParseIInterfaceClass(FClasses& AllClasses)
  */
 void FHeaderParser::CompileInterfaceDeclaration(FClasses& AllClasses)
 {
-	FUnrealSourceFile* CurrentSourceFile = GetCurrentSourceFile();
+	FUnrealSourceFile* CurrentSrcFile = GetCurrentSourceFile();
 	// Start of an interface block. Since Interfaces and Classes are always at the same nesting level,
 	// whereever a class declaration is allowed, an interface declaration is also allowed.
 	CheckAllow( TEXT("'interface'"), ALLOW_Class );
@@ -5142,8 +5142,8 @@ void FHeaderParser::CompileInterfaceDeclaration(FClasses& AllClasses)
 	}
 
 	// Try parsing metadata for the interface
-	FClassMetaData* ClassData = GScriptHelper.AddClassData(InterfaceClass, UHTMakefile, CurrentSourceFile);
-	UHTMakefile.AddGScriptHelperEntry(CurrentSourceFile, InterfaceClass, ClassData);
+	FClassMetaData* ClassData = GScriptHelper.AddClassData(InterfaceClass, UHTMakefile, CurrentSrcFile);
+	UHTMakefile.AddGScriptHelperEntry(CurrentSrcFile, InterfaceClass, ClassData);
 	check(ClassData);
 
 	ClassData->SetPrologLine(PrologFinishLine);
@@ -5299,9 +5299,9 @@ void FHeaderParser::ParseParameterList(FClasses& AllClasses, UFunction* Function
 }
 void FHeaderParser::CompileDelegateDeclaration(FClasses& AllClasses, const TCHAR* DelegateIdentifier, EDelegateSpecifierAction::Type SpecifierAction)
 {
-	FUnrealSourceFile* CurrentSourceFile = GetCurrentSourceFile();
+	FUnrealSourceFile* CurrentSrcFile = GetCurrentSourceFile();
 	TMap<FName, FString> MetaData;
-	AddModuleRelativePathToMetadata(*CurrentSourceFile, MetaData);
+	AddModuleRelativePathToMetadata(*CurrentSrcFile, MetaData);
 
 	FFuncInfo            FuncInfo;
 
@@ -5412,10 +5412,10 @@ void FHeaderParser::CompileDelegateDeclaration(FClasses& AllClasses, const TCHAR
 
 	FuncInfo.MacroLine = InputLine;
 	auto* DelegateSignatureFunction = CreateDelegateFunction(FuncInfo);
-	UHTMakefile.AddDelegateFunction(CurrentSourceFile, DelegateSignatureFunction);
+	UHTMakefile.AddDelegateFunction(CurrentSrcFile, DelegateSignatureFunction);
 
-	FClassMetaData* ClassMetaData = GScriptHelper.AddClassData(DelegateSignatureFunction, UHTMakefile, CurrentSourceFile);
-	UHTMakefile.AddGScriptHelperEntry(CurrentSourceFile, DelegateSignatureFunction, ClassMetaData);
+	FClassMetaData* ClassMetaData = GScriptHelper.AddClassData(DelegateSignatureFunction, UHTMakefile, CurrentSrcFile);
+	UHTMakefile.AddGScriptHelperEntry(CurrentSrcFile, DelegateSignatureFunction, ClassMetaData);
 
 	DelegateSignatureFunction->FunctionFlags |= FuncInfo.FunctionFlags;
 
@@ -5499,9 +5499,9 @@ void FHeaderParser::CompileFunctionDeclaration(FClasses& AllClasses)
 {
 	CheckAllow(TEXT("'Function'"), ALLOW_Function);
 
-	FUnrealSourceFile* CurrentSourceFile = GetCurrentSourceFile();
+	FUnrealSourceFile* CurrentSrcFile = GetCurrentSourceFile();
 	TMap<FName, FString> MetaData;
-	AddModuleRelativePathToMetadata(*CurrentSourceFile, MetaData);
+	AddModuleRelativePathToMetadata(*CurrentSrcFile, MetaData);
 
 	// New-style UFUNCTION() syntax 
 	TArray<FPropertySpecifier> SpecifiersFound;
@@ -5779,10 +5779,10 @@ void FHeaderParser::CompileFunctionDeclaration(FClasses& AllClasses)
 	}
 
 	auto* TopFunction = CreateFunction(FuncInfo);
-	UHTMakefile.AddFunction(CurrentSourceFile, TopFunction);
+	UHTMakefile.AddFunction(CurrentSrcFile, TopFunction);
 
-	FClassMetaData* ClassMetaData = GScriptHelper.AddClassData(TopFunction, UHTMakefile, CurrentSourceFile);
-	UHTMakefile.AddGScriptHelperEntry(CurrentSourceFile, TopFunction, ClassMetaData);
+	FClassMetaData* ClassMetaData = GScriptHelper.AddClassData(TopFunction, UHTMakefile, CurrentSrcFile);
+	UHTMakefile.AddGScriptHelperEntry(CurrentSrcFile, TopFunction, ClassMetaData);
 
 	TopFunction->FunctionFlags |= FuncInfo.FunctionFlags;
 
@@ -6527,16 +6527,16 @@ ECompilationResult::Type FHeaderParser::ParseHeader(FClasses& AllClasses, FUnrea
 	SetCurrentSourceFile(SourceFile);
 	UHTMakefile.AddToHeaderOrder(SourceFile);
 	NameLookupCPP.SetCurrentSourceFile(SourceFile);
-	FUnrealSourceFile* CurrentSourceFile = SourceFile;
-	if (CurrentSourceFile->IsParsed())
+	FUnrealSourceFile* CurrentSrcFile = SourceFile;
+	if (CurrentSrcFile->IsParsed())
 	{
 		return ECompilationResult::Succeeded;
 	}
 
-	CurrentSourceFile->MarkAsParsed();
+	CurrentSrcFile->MarkAsParsed();
 
 	// Early-out if this class has previously failed some aspect of parsing
-	if (FailedFilesAnnotation.Get(CurrentSourceFile))
+	if (FailedFilesAnnotation.Get(CurrentSrcFile))
 	{
 		return ECompilationResult::OtherCompilationError;
 	}
@@ -6555,16 +6555,16 @@ ECompilationResult::Type FHeaderParser::ParseHeader(FClasses& AllClasses, FUnrea
 	if (FParse::Param(FCommandLine::Get(), TEXT("VERBOSE")))
 	{
 		// Message.
-		Warn->Logf(TEXT("Parsing %s"), *CurrentSourceFile->GetFilename());
+		Warn->Logf(TEXT("Parsing %s"), *CurrentSrcFile->GetFilename());
 	}
 
 	// Init compiler variables.
-	ResetParser(*CurrentSourceFile->GetContent());
+	ResetParser(*CurrentSrcFile->GetContent());
 
 	// Init nesting.
 	NestLevel = 0;
 	TopNest = NULL;
-	PushNest(NEST_GlobalScope, nullptr, CurrentSourceFile);
+	PushNest(NEST_GlobalScope, nullptr, CurrentSrcFile);
 
 	// C++ classes default to private access level
 	CurrentAccessSpecifier = ACCESS_Private; 
@@ -6591,7 +6591,7 @@ ECompilationResult::Type FHeaderParser::ParseHeader(FClasses& AllClasses, FUnrea
 
 		PopNest(NEST_GlobalScope, TEXT("Global scope"));
 
-		auto ScopeTypeIterator = CurrentSourceFile->GetScope()->GetTypeIterator();
+		auto ScopeTypeIterator = CurrentSrcFile->GetScope()->GetTypeIterator();
 		while (ScopeTypeIterator.MoveNext())
 		{
 			auto* Type = *ScopeTypeIterator;
@@ -6641,7 +6641,7 @@ ECompilationResult::Type FHeaderParser::ParseHeader(FClasses& AllClasses, FUnrea
 		// First-pass success.
 		Result = ECompilationResult::Succeeded;
 
-		for (auto* Class : CurrentSourceFile->GetDefinedClasses())
+		for (auto* Class : CurrentSrcFile->GetDefinedClasses())
 		{
 			PostParsingClassSetup(Class);
 
@@ -6657,11 +6657,11 @@ ECompilationResult::Type FHeaderParser::ParseHeader(FClasses& AllClasses, FUnrea
 			bNoExportClassesOnly = bNoExportClassesOnly && Class->HasAnyClassFlags(CLASS_NoExport);
 		}
 
-		check(CurrentSourceFile->IsParsed());
+		check(CurrentSrcFile->IsParsed());
 
 		if (!bSpottedAutogeneratedHeaderInclude && !bEmptyFile && !bNoExportClassesOnly)
 		{
-			const FString ExpectedHeaderName = CurrentSourceFile->GetGeneratedHeaderFilename();
+			const FString ExpectedHeaderName = CurrentSrcFile->GetGeneratedHeaderFilename();
 			FError::Throwf(TEXT("Expected an include at the top of the header: '#include \"%s\"'"), *ExpectedHeaderName);
 		}
 	}
@@ -6671,7 +6671,7 @@ ECompilationResult::Type FHeaderParser::ParseHeader(FClasses& AllClasses, FUnrea
 		if (NestLevel == 0)
 		{
 			// Pushing nest so there is a file context for this error.
-			PushNest(NEST_GlobalScope, nullptr, CurrentSourceFile);
+			PushNest(NEST_GlobalScope, nullptr, CurrentSrcFile);
 		}
 
 		// Handle compiler error.
@@ -6683,7 +6683,7 @@ ECompilationResult::Type FHeaderParser::ParseHeader(FClasses& AllClasses, FUnrea
 			Warn->Log(ELogVerbosity::Error, ErrorMsg);
 		}
 
-		FailedFilesAnnotation.Set(CurrentSourceFile);
+		FailedFilesAnnotation.Set(CurrentSrcFile);
 		Result = GCompilationResult;
 	}
 #endif
@@ -6930,9 +6930,9 @@ void FHeaderParser::ExportNativeHeaders(
 FHeaderParser::FHeaderParser(FFeedbackContext* InWarn, FUHTMakefile& InUHTMakefile)
 : FBaseParser                       ()
 , Warn                              (InWarn)
+, UHTMakefile(InUHTMakefile)
 , bSpottedAutogeneratedHeaderInclude(false)
 , TopNest                           (NULL)
-, UHTMakefile(InUHTMakefile)
 {
 	FScriptLocation::Compiler = this;
 
@@ -8031,7 +8031,7 @@ void FHeaderParser::SkipDeprecatedMacroIfNecessary()
 void FHeaderParser::CompileVersionDeclaration(UStruct* Struct)
 {
 	FUnrealSourceFile* CurrentSourceFilePtr = GetCurrentSourceFile();
-	TSharedRef<FUnrealSourceFile> CurrentSourceFile = CurrentSourceFilePtr->AsShared();
+	TSharedRef<FUnrealSourceFile> CurrentSrcFile = CurrentSourceFilePtr->AsShared();
 	// Do nothing if we're at the end of file.
 	FToken Token;
 	if (!GetToken(Token, true, ESymbolParseOption::Normal))
@@ -8051,7 +8051,7 @@ void FHeaderParser::CompileVersionDeclaration(UStruct* Struct)
 	if (Token.TokenType == ETokenType::TOKEN_Symbol
 		&& !FCString::Stricmp(Token.Identifier, TEXT(")")))
 	{
-		CurrentSourceFile->GetGeneratedCodeVersions().FindOrAdd(Struct) = Version;
+		CurrentSrcFile->GetGeneratedCodeVersions().FindOrAdd(Struct) = Version;
 		UngetToken(Token);
 		return;
 	}
@@ -8059,7 +8059,7 @@ void FHeaderParser::CompileVersionDeclaration(UStruct* Struct)
 	// Overwrite with version specified by macro.
 	Version = ToGeneratedCodeVersion(Token.Identifier);
 
-	CurrentSourceFile->GetGeneratedCodeVersions().FindOrAdd(Struct) = Version;
+	CurrentSrcFile->GetGeneratedCodeVersions().FindOrAdd(Struct) = Version;
 }
 
 void FHeaderParser::ResetClassData()
diff --git a/Engine/Source/Programs/UnrealHeaderTool/Private/Scope.h b/Engine/Source/Programs/UnrealHeaderTool/Private/Scope.h
--- a/Engine/Source/Programs/UnrealHeaderTool/Private/Scope.h
+++ b/Engine/Source/Programs/UnrealHeaderTool/Private/Scope.h
@@ -267,8 +267,8 @@ class FFileScope : public FScope
 {
 public:
 	FFileScope()
-		: Name(NAME_None)
-		, SourceFile(nullptr)
+		: SourceFile(nullptr)
+		, Name(NAME_None)
 	{ }
 	// Constructor.
 	FFileScope(FName Name, FUnrealSourceFile* SourceFile);
diff --git a/Engine/Source/Programs/UnrealHeaderTool/Private/UHTMakefile/NameLookupCPPArchiveProxy.cpp b/Engine/Source/Programs/UnrealHeaderTool/Private/UHTMakefile/NameLookupCPPArchiveProxy.cpp
--- a/Engine/Source/Programs/UnrealHeaderTool/Private/UHTMakefile/NameLookupCPPArchiveProxy.cpp
+++ b/Engine/Source/Programs/UnrealHeaderTool/Private/UHTMakefile/NameLookupCPPArchiveProxy.cpp
@@ -3,19 +3,21 @@
 #include "UHTMakefile/UHTMakefile.h"
 #include "UHTMakefile/NameLookupCPPArchiveProxy.h"
 
-FNameLookupCPPArchiveProxy::FNameLookupCPPArchiveProxy(const FUHTMakefile& UHTMakefile, const FNameLookupCPP* NameLookupCPP)
+FNameLookupCPPArchiveProxy::FNameLookupCPPArchiveProxy(const FUHTMakefile& UHTMakefile, const FNameLookupCPP* InNameLookupCPP)
 {
-	TArray<TPair<FSerializeIndex, FString>> StructNameMap;
-	StructNameMap.Empty(NameLookupCPP->StructNameMap.Num());
-	for (auto& Kvp : NameLookupCPP->StructNameMap)
+	// TODO: This really probably shouldn't be a local variable and should be populating the member variable instead
+	TArray<TPair<FSerializeIndex, FString>> StructNameMapLocal;
+	StructNameMapLocal.Empty(InNameLookupCPP->StructNameMap.Num());
+	for (auto& Kvp : InNameLookupCPP->StructNameMap)
 	{
-		StructNameMap.Add(TPairInitializer<FSerializeIndex, FString>(UHTMakefile.GetStructIndex(Kvp.Key), Kvp.Value));
+		StructNameMapLocal.Add(TPairInitializer<FSerializeIndex, FString>(UHTMakefile.GetStructIndex(Kvp.Key), Kvp.Value));
 	}
-	TArray<FString> InterfaceAllocations;
-	InterfaceAllocations.Empty(NameLookupCPP->InterfaceAllocations.Num());
-	for (TCHAR* InterfaceAllocation : NameLookupCPP->InterfaceAllocations)
+	// TODO: This really probably shouldn't be a local variable and should be populating the member variable instead
+	TArray<FString> InterfaceAllocationsLocal;
+	InterfaceAllocationsLocal.Empty(InNameLookupCPP->InterfaceAllocations.Num());
+	for (TCHAR* InterfaceAllocation : InNameLookupCPP->InterfaceAllocations)
 	{
-		InterfaceAllocations.Add(InterfaceAllocation);
+		InterfaceAllocationsLocal.Add(InterfaceAllocation);
 	}
 }
 
@@ -27,7 +29,7 @@ FArchive& operator<<(FArchive& Ar, FNameLookupCPPArchiveProxy& NameLookupCPPArch
 	return Ar;
 }
 
-void FNameLookupCPPArchiveProxy::Resolve(FNameLookupCPP* NameLookupCPP, FUHTMakefile& UHTMakefile)
+void FNameLookupCPPArchiveProxy::Resolve(FNameLookupCPP* InNameLookupCPP, FUHTMakefile& UHTMakefile)
 {
 
 }
diff --git a/Engine/Source/Programs/UnrealHeaderTool/Private/UHTMakefile/PropertyArchiveProxy.h b/Engine/Source/Programs/UnrealHeaderTool/Private/UHTMakefile/PropertyArchiveProxy.h
--- a/Engine/Source/Programs/UnrealHeaderTool/Private/UHTMakefile/PropertyArchiveProxy.h
+++ b/Engine/Source/Programs/UnrealHeaderTool/Private/UHTMakefile/PropertyArchiveProxy.h
@@ -191,7 +191,7 @@ struct FNamePropertyArchiveProxy : public FPropertyArchiveProxy
 
 struct FStrPropertyArchiveProxy : public FPropertyArchiveProxy
 {
-	FStrPropertyArchiveProxy::FStrPropertyArchiveProxy() { }
+	FStrPropertyArchiveProxy() { }
 	FStrPropertyArchiveProxy(FUHTMakefile& UHTMakefile, const UStrProperty* StrProperty);
 	static void AddReferencedNames(const UStrProperty* Property, FUHTMakefile& UHTMakefile)
 	{
@@ -203,7 +203,7 @@ struct FStrPropertyArchiveProxy : public FPropertyArchiveProxy
 
 struct FTextPropertyArchiveProxy : public FPropertyArchiveProxy
 {
-	FTextPropertyArchiveProxy::FTextPropertyArchiveProxy() { }
+	FTextPropertyArchiveProxy() { }
 	FTextPropertyArchiveProxy(FUHTMakefile& UHTMakefile, const UTextProperty* TextProperty);
 
 	UTextProperty* CreateTextProperty(const FUHTMakefile& UHTMakefile) const;
@@ -215,7 +215,7 @@ struct FTextPropertyArchiveProxy : public FPropertyArchiveProxy
 
 struct FDelegatePropertyArchiveProxy : public FPropertyArchiveProxy
 {
-	FDelegatePropertyArchiveProxy::FDelegatePropertyArchiveProxy() { }
+	FDelegatePropertyArchiveProxy() { }
 	FDelegatePropertyArchiveProxy(FUHTMakefile& UHTMakefile, const UDelegateProperty* DelegateProperty);
 
 	UDelegateProperty* CreateDelegateProperty(const FUHTMakefile& UHTMakefile) const;
@@ -227,7 +227,7 @@ struct FDelegatePropertyArchiveProxy : public FPropertyArchiveProxy
 
 struct FMulticastDelegatePropertyArchiveProxy : public FPropertyArchiveProxy
 {
-	FMulticastDelegatePropertyArchiveProxy::FMulticastDelegatePropertyArchiveProxy() { }
+	FMulticastDelegatePropertyArchiveProxy() { }
 	FMulticastDelegatePropertyArchiveProxy(FUHTMakefile& UHTMakefile, const UMulticastDelegateProperty* MulticastDelegateProperty);
 
 	UMulticastDelegateProperty* CreateMulticastDelegateProperty(const FUHTMakefile& UHTMakefile) const;
@@ -239,7 +239,7 @@ struct FMulticastDelegatePropertyArchiveProxy : public FPropertyArchiveProxy
 
 struct FObjectPropertyBaseArchiveProxy : public FPropertyArchiveProxy
 {
-	FObjectPropertyBaseArchiveProxy::FObjectPropertyBaseArchiveProxy() { }
+	FObjectPropertyBaseArchiveProxy() { }
 	FObjectPropertyBaseArchiveProxy(FUHTMakefile& UHTMakefile, const UObjectPropertyBase* ObjectPropertyBase);
 
 	UObjectPropertyBase* CreateObjectPropertyBase(const FUHTMakefile& UHTMakefile) const;
@@ -257,7 +257,7 @@ struct FObjectPropertyBaseArchiveProxy : public FPropertyArchiveProxy
 
 struct FClassPropertyArchiveProxy : public FObjectPropertyBaseArchiveProxy
 {
-	FClassPropertyArchiveProxy::FClassPropertyArchiveProxy() { }
+	FClassPropertyArchiveProxy() { }
 	FClassPropertyArchiveProxy(FUHTMakefile& UHTMakefile, const UClassProperty* ClassProperty);
 
 	UClassProperty* CreateClassProperty(const FUHTMakefile& UHTMakefile) const;
@@ -275,7 +275,7 @@ struct FClassPropertyArchiveProxy : public FObjectPropertyBaseArchiveProxy
 
 struct FObjectPropertyArchiveProxy : public FObjectPropertyBaseArchiveProxy
 {
-	FObjectPropertyArchiveProxy::FObjectPropertyArchiveProxy() { }
+	FObjectPropertyArchiveProxy() { }
 	FObjectPropertyArchiveProxy(FUHTMakefile& UHTMakefile, const UObjectProperty* ObjectProperty);
 	static void AddReferencedNames(const UObjectProperty* Property, FUHTMakefile& UHTMakefile)
 	{
@@ -287,7 +287,7 @@ struct FObjectPropertyArchiveProxy : public FObjectPropertyBaseArchiveProxy
 
 struct FWeakObjectPropertyArchiveProxy : public FObjectPropertyBaseArchiveProxy
 {
-	FWeakObjectPropertyArchiveProxy::FWeakObjectPropertyArchiveProxy() { }
+	FWeakObjectPropertyArchiveProxy() { }
 	FWeakObjectPropertyArchiveProxy(FUHTMakefile& UHTMakefile, const UWeakObjectProperty* WeakObjectProperty);
 
 	UWeakObjectProperty* CreateWeakObjectProperty(const FUHTMakefile& UHTMakefile) const;
@@ -299,7 +299,7 @@ struct FWeakObjectPropertyArchiveProxy : public FObjectPropertyBaseArchiveProxy
 
 struct FLazyObjectPropertyArchiveProxy : public FObjectPropertyBaseArchiveProxy
 {
-	FLazyObjectPropertyArchiveProxy::FLazyObjectPropertyArchiveProxy() { }
+	FLazyObjectPropertyArchiveProxy() { }
 	FLazyObjectPropertyArchiveProxy(FUHTMakefile& UHTMakefile, const ULazyObjectProperty* LazyObjectProperty);
 
 	ULazyObjectProperty* CreateLazyObjectProperty(const FUHTMakefile& UHTMakefile) const;
@@ -311,7 +311,7 @@ struct FLazyObjectPropertyArchiveProxy : public FObjectPropertyBaseArchiveProxy
 
 struct FAssetObjectPropertyArchiveProxy : public FObjectPropertyBaseArchiveProxy
 {
-	FAssetObjectPropertyArchiveProxy::FAssetObjectPropertyArchiveProxy() { }
+	FAssetObjectPropertyArchiveProxy() { }
 	FAssetObjectPropertyArchiveProxy(FUHTMakefile& UHTMakefile, const UAssetObjectProperty* AssetObjectProperty);
 
 	UAssetObjectProperty* CreateAssetObjectProperty(const FUHTMakefile& UHTMakefile) const;
@@ -323,7 +323,7 @@ struct FAssetObjectPropertyArchiveProxy : public FObjectPropertyBaseArchiveProxy
 
 struct FAssetClassPropertyArchiveProxy : public FObjectPropertyBaseArchiveProxy
 {
-	FAssetClassPropertyArchiveProxy::FAssetClassPropertyArchiveProxy() { }
+	FAssetClassPropertyArchiveProxy() { }
 	FAssetClassPropertyArchiveProxy(FUHTMakefile& UHTMakefile, const UAssetClassProperty* AssetClassProperty);
 
 	UAssetClassProperty* CreateAssetClassProperty(const FUHTMakefile& UHTMakefile) const;
@@ -341,7 +341,7 @@ struct FAssetClassPropertyArchiveProxy : public FObjectPropertyBaseArchiveProxy
 
 struct FInterfacePropertyArchiveProxy : public FPropertyArchiveProxy
 {
-	FInterfacePropertyArchiveProxy::FInterfacePropertyArchiveProxy() { }
+	FInterfacePropertyArchiveProxy() { }
 	FInterfacePropertyArchiveProxy(FUHTMakefile& UHTMakefile, const UInterfaceProperty* InterfaceProperty);
 
 	UInterfaceProperty* CreateInterfaceProperty(const FUHTMakefile& UHTMakefile) const;
@@ -359,7 +359,7 @@ struct FInterfacePropertyArchiveProxy : public FPropertyArchiveProxy
 
 struct FStructPropertyArchiveProxy : public FPropertyArchiveProxy
 {
-	FStructPropertyArchiveProxy::FStructPropertyArchiveProxy() { }
+	FStructPropertyArchiveProxy() { }
 	FStructPropertyArchiveProxy(FUHTMakefile& UHTMakefile, const UStructProperty* StructProperty);
 
 	UStructProperty* CreateStructProperty(const FUHTMakefile& UHTMakefile) const;
@@ -383,7 +383,7 @@ FArchive& operator<<(FArchive& Ar, FScriptMapLayout& ScriptMapLayout);
 
 struct FMapPropertyArchiveProxy : public FPropertyArchiveProxy
 {
-	FMapPropertyArchiveProxy::FMapPropertyArchiveProxy() { }
+	FMapPropertyArchiveProxy() { }
 	FMapPropertyArchiveProxy(FUHTMakefile& UHTMakefile, const UMapProperty* MapProperty);
 
 	UMapProperty* CreateMapProperty(const FUHTMakefile& UHTMakefile) const;
@@ -403,7 +403,7 @@ struct FMapPropertyArchiveProxy : public FPropertyArchiveProxy
 
 struct FArrayPropertyArchiveProxy : public FPropertyArchiveProxy
 {
-	FArrayPropertyArchiveProxy::FArrayPropertyArchiveProxy() { }
+	FArrayPropertyArchiveProxy() { }
 	FArrayPropertyArchiveProxy(FUHTMakefile& UHTMakefile, const UArrayProperty* ArrayProperty);
 
 	UArrayProperty* CreateArrayProperty(const FUHTMakefile& UHTMakefile) const;
diff --git a/Engine/Source/Programs/UnrealHeaderTool/Private/UHTMakefile/UHTMakefile.cpp b/Engine/Source/Programs/UnrealHeaderTool/Private/UHTMakefile/UHTMakefile.cpp
--- a/Engine/Source/Programs/UnrealHeaderTool/Private/UHTMakefile/UHTMakefile.cpp
+++ b/Engine/Source/Programs/UnrealHeaderTool/Private/UHTMakefile/UHTMakefile.cpp
@@ -517,7 +517,7 @@ void FUHTMakefile::SaveToFile(const TCHAR* MakefilePath)
 	UE_LOG(LogCompile, Log, TEXT("Saving UHT makefile took %f seconds"), UHTMakefileLoadTime);
 }
 
-void FUHTMakefile::LoadFromFile(const TCHAR* MakefilePath, FManifest* Manifest)
+void FUHTMakefile::LoadFromFile(const TCHAR* MakefilePath, FManifest* InManifest)
 {
 	FArchive* UHTMakefileArchive = IFileManager::Get().CreateFileReader(MakefilePath);
 	if (!UHTMakefileArchive)
@@ -537,7 +537,7 @@ void FUHTMakefile::LoadFromFile(const TCHAR* MakefilePath, FManifest* Manifest)
 	UHTMakefileLoadTimer.Stop();
 	UE_LOG(LogCompile, Log, TEXT("Loading UHT makefile took %f seconds"), UHTMakefileLoadTime);
 
-	UpdateModulesCompatibility(Manifest);
+	UpdateModulesCompatibility(InManifest);
 }
 
 bool FUHTMakefile::CanLoadModule(const FManifestModule& ManifestModule)
diff --git a/Engine/Source/Programs/UnrealHeaderTool/Private/UHTMakefile/UHTMakefile.h b/Engine/Source/Programs/UnrealHeaderTool/Private/UHTMakefile/UHTMakefile.h
--- a/Engine/Source/Programs/UnrealHeaderTool/Private/UHTMakefile/UHTMakefile.h
+++ b/Engine/Source/Programs/UnrealHeaderTool/Private/UHTMakefile/UHTMakefile.h
@@ -104,10 +104,10 @@ class FUHTMakefile
 {
 public:
 	FUHTMakefile()
-		: Version(CurrentVersion)
-		, LoadingPhase(EUHTMakefileLoadingPhase::Max)
+		: bShouldForceRegeneration(false)
 		, bShouldMoveNewObjects(false)
-		, bShouldForceRegeneration(false)
+		, Version(CurrentVersion)
+		, LoadingPhase(EUHTMakefileLoadingPhase::Max)
 	{
 		InitializationOrder.AddDefaulted(+EUHTMakefileLoadingPhase::Max);
 	}
diff --git a/Engine/Source/Programs/UnrealHeaderTool/Private/UHTMakefile/UHTMakefileHeaderDescriptor.cpp b/Engine/Source/Programs/UnrealHeaderTool/Private/UHTMakefile/UHTMakefileHeaderDescriptor.cpp
--- a/Engine/Source/Programs/UnrealHeaderTool/Private/UHTMakefile/UHTMakefileHeaderDescriptor.cpp
+++ b/Engine/Source/Programs/UnrealHeaderTool/Private/UHTMakefile/UHTMakefileHeaderDescriptor.cpp
@@ -6,9 +6,9 @@
 #include "ClassMaps.h"
 
 FUHTMakefileHeaderDescriptor::FUHTMakefileHeaderDescriptor(FUHTMakefile* InUHTMakefile /*= nullptr*/)
-	: LoadingPhase(EUHTMakefileLoadingPhase::Preload)
-	, UHTMakefile(InUHTMakefile)
+	: UHTMakefile(InUHTMakefile)
 	, SourceFile(nullptr)
+	, LoadingPhase(EUHTMakefileLoadingPhase::Preload)
 {
 	InitializeObjectIndexes.AddDefaulted(+EUHTMakefileLoadingPhase::Max);
 	for (uint8 i = 0; i < +EUHTMakefileLoadingPhase::Max; ++i)
diff --git a/Engine/Source/Programs/UnrealHeaderTool/Private/UHTMakefile/UHTMakefileModuleDescriptor.h b/Engine/Source/Programs/UnrealHeaderTool/Private/UHTMakefile/UHTMakefileModuleDescriptor.h
--- a/Engine/Source/Programs/UnrealHeaderTool/Private/UHTMakefile/UHTMakefileModuleDescriptor.h
+++ b/Engine/Source/Programs/UnrealHeaderTool/Private/UHTMakefile/UHTMakefileModuleDescriptor.h
@@ -11,9 +11,10 @@ class FUHTMakefileModuleDescriptor
 {
 public:
 	FUHTMakefileModuleDescriptor(FUHTMakefile* InUHTMakefile = nullptr)
-		: UHTMakefile(InUHTMakefile)
+		: LoadingPhase(EUHTMakefileLoadingPhase::Preload)
+		, UHTMakefile(InUHTMakefile)
 		, Package(nullptr)
-		, LoadingPhase(EUHTMakefileLoadingPhase::Preload)
+		
 	{
 		HeadersOrder.AddDefaulted(+EUHTMakefileLoadingPhase::Max);
 	}
