#pragma once removed from .cpp files.

diff --git a/Engine/Source/Programs/UnrealHeaderTool/Private/UHTMakefile/ClassArchiveProxy.cpp b/Engine/Source/Programs/UnrealHeaderTool/Private/UHTMakefile/ClassArchiveProxy.cpp
--- a/Engine/Source/Programs/UnrealHeaderTool/Private/UHTMakefile/ClassArchiveProxy.cpp
+++ b/Engine/Source/Programs/UnrealHeaderTool/Private/UHTMakefile/ClassArchiveProxy.cpp
@@ -1,6 +1,5 @@
 // Copyright 1998-2015 Epic Games, Inc. All Rights Reserved.
 
-#pragma once
 #include "UnrealHeaderTool.h"
 #include "UHTMakefile/UHTMakefile.h"
 #include "UHTMakefile/ClassArchiveProxy.h"
diff --git a/Engine/Source/Programs/UnrealHeaderTool/Private/UHTMakefile/ImplementedInterfaceArchiveProxy.cpp b/Engine/Source/Programs/UnrealHeaderTool/Private/UHTMakefile/ImplementedInterfaceArchiveProxy.cpp
--- a/Engine/Source/Programs/UnrealHeaderTool/Private/UHTMakefile/ImplementedInterfaceArchiveProxy.cpp
+++ b/Engine/Source/Programs/UnrealHeaderTool/Private/UHTMakefile/ImplementedInterfaceArchiveProxy.cpp
@@ -1,6 +1,5 @@
 // Copyright 1998-2015 Epic Games, Inc. All Rights Reserved.
 
-#pragma once
 #include "UnrealHeaderTool.h"
 #include "UHTMakefile/UHTMakefile.h"
 #include "UHTMakefile/ImplementedInterfaceArchiveProxy.h"
diff --git a/Engine/Source/Programs/UnrealHeaderTool/Private/UHTMakefile/PropertyArchiveProxy.cpp b/Engine/Source/Programs/UnrealHeaderTool/Private/UHTMakefile/PropertyArchiveProxy.cpp
--- a/Engine/Source/Programs/UnrealHeaderTool/Private/UHTMakefile/PropertyArchiveProxy.cpp
+++ b/Engine/Source/Programs/UnrealHeaderTool/Private/UHTMakefile/PropertyArchiveProxy.cpp
@@ -1,6 +1,5 @@
 // Copyright 1998-2015 Epic Games, Inc. All Rights Reserved.
 
-#pragma once
 #include "UnrealHeaderTool.h"
 #include "UHTMakefile/UHTMakefile.h"
 #include "UHTMakefile/PropertyArchiveProxy.h"
diff --git a/Engine/Source/Programs/UnrealHeaderTool/Private/UHTMakefile/PropertyBaseArchiveProxy.cpp b/Engine/Source/Programs/UnrealHeaderTool/Private/UHTMakefile/PropertyBaseArchiveProxy.cpp
--- a/Engine/Source/Programs/UnrealHeaderTool/Private/UHTMakefile/PropertyBaseArchiveProxy.cpp
+++ b/Engine/Source/Programs/UnrealHeaderTool/Private/UHTMakefile/PropertyBaseArchiveProxy.cpp
@@ -1,6 +1,5 @@
 // Copyright 1998-2015 Epic Games, Inc. All Rights Reserved.
 
-#pragma once
 #include "UnrealHeaderTool.h"
 #include "UHTMakefile/UHTMakefile.h"
 #include "UHTMakefile/PropertyBaseArchiveProxy.h"
diff --git a/Engine/Source/Programs/UnrealHeaderTool/Private/UHTMakefile/RepRecordArchiveProxy.cpp b/Engine/Source/Programs/UnrealHeaderTool/Private/UHTMakefile/RepRecordArchiveProxy.cpp
--- a/Engine/Source/Programs/UnrealHeaderTool/Private/UHTMakefile/RepRecordArchiveProxy.cpp
+++ b/Engine/Source/Programs/UnrealHeaderTool/Private/UHTMakefile/RepRecordArchiveProxy.cpp
@@ -1,6 +1,5 @@
 // Copyright 1998-2015 Epic Games, Inc. All Rights Reserved.
 
-#pragma once
 #include "UnrealHeaderTool.h"
 #include "UHTMakefile/UHTMakefile.h"
 #include "UHTMakefile/RepRecordArchiveProxy.h"
diff --git a/Engine/Source/Programs/UnrealHeaderTool/Private/UHTMakefile/ScriptStructArchiveProxy.cpp b/Engine/Source/Programs/UnrealHeaderTool/Private/UHTMakefile/ScriptStructArchiveProxy.cpp
--- a/Engine/Source/Programs/UnrealHeaderTool/Private/UHTMakefile/ScriptStructArchiveProxy.cpp
+++ b/Engine/Source/Programs/UnrealHeaderTool/Private/UHTMakefile/ScriptStructArchiveProxy.cpp
@@ -1,6 +1,5 @@
 // Copyright 1998-2015 Epic Games, Inc. All Rights Reserved.
 
-#pragma once
 #include "UnrealHeaderTool.h"
 #include "UHTMakefile/UHTMakefile.h"
 #include "UHTMakefile/StructArchiveProxy.h"
