UHT: Fix build system not being able to parse certain messages out of UHT logs, due to errors not containing the string "error:". Was previously outputting messages like "blah.cpp(123): message".

diff --git a/Engine/Source/Programs/UnrealHeaderTool/Private/HeaderParser.cpp b/Engine/Source/Programs/UnrealHeaderTool/Private/HeaderParser.cpp
--- a/Engine/Source/Programs/UnrealHeaderTool/Private/HeaderParser.cpp
+++ b/Engine/Source/Programs/UnrealHeaderTool/Private/HeaderParser.cpp
@@ -6753,7 +6753,7 @@ ECompilationResult::Type FHeaderParser::ParseHeader(FClasses& AllClasses, FUnrea
 		// Handle compiler error.
 		{
 			TGuardValue<ELogTimes::Type> DisableLogTimes(GPrintLogTimes, ELogTimes::None);
-			FString FormattedErrorMessageWithContext = FString::Printf(TEXT("%s: %s"), *GetContext(), ErrorMsg);
+			FString FormattedErrorMessageWithContext = FString::Printf(TEXT("%s: Error: %s"), *GetContext(), ErrorMsg);
 
 			UE_LOG(LogCompile, Log,  TEXT("%s"), *FormattedErrorMessageWithContext );
 			Warn->Log(ELogVerbosity::Error, ErrorMsg);
