Copying //UE4/Dev-Core to //UE4/Dev-Main (Source: //UE4/Dev-Core @ 3049602)

diff --git a/Engine/Source/Programs/CrashReporter/CrashReportWebSite/ViewModels/FormHelper.cs b/Engine/Source/Programs/CrashReporter/CrashReportWebSite/ViewModels/FormHelper.cs
new file mode 100644
--- /dev/null
+++ b/Engine/Source/Programs/CrashReporter/CrashReportWebSite/ViewModels/FormHelper.cs
@@ -0,0 +1,248 @@
+﻿// Copyright 1998-2016 Epic Games, Inc. All Rights Reserved.
+
+using System;
+using System.Collections.Generic;
+using System.Linq;
+using System.Web;
+using System.Web.Mvc;
+using System.Web.Helpers;
+
+namespace Tools.CrashReporter.CrashReportWebSite.ViewModels
+{
+	/// <summary>
+	/// A helper class to extract parameters passed from web pages
+	/// </summary>
+	public class FormHelper
+	{
+		/// <summary>The user search query.</summary>
+		public string SearchQuery = "";
+
+		/// <summary>User name as query for filtering.</summary>
+		public string UsernameQuery = "";
+
+		/// <summary>Epic ID or Machine ID as query for filtering.</summary>
+		public string EpicIdOrMachineQuery = "";
+
+		/// <summary>Jira as query for crash filtering.</summary>
+		public string JiraQuery = "";
+
+        /// <summary> Jira Id for crash/bugg filtering </summary>
+	    public string JiraId = "";
+
+		/// <summary>Message/Summary or Description as query for filtering.</summary>
+		public string MessageQuery = "";
+
+		/// <summary>BuiltFromCL as query for filtering.</summary>
+		public string BuiltFromCL = "";
+
+		/// <summary>BuggId as query for filtering.</summary>
+		public string BuggId = "";
+
+		/// <summary>The page to display from the list.</summary>
+		public int Page = 1;
+
+		/// <summary>The number of items to display per page.</summary>
+		public int PageSize = 100;
+
+		/// <summary>The term to sort by. e.g. Jira.</summary>
+		public string SortTerm = "";
+
+		/// <summary>The types of crashes we wish to see.</summary>
+		public string CrashType = "CrashesAsserts";
+
+		/// <summary>Whether to sort ascending or descending.</summary>
+		public string SortOrder = "Descending";
+
+		/// <summary>The previous sort order.</summary>
+		public string PreviousOrder = "";
+
+		/// <summary>The user group to filter by.</summary>
+		public string UserGroup = "General";
+
+		/// <summary>The earliest date we are interested in.</summary>
+		public DateTime DateFrom = DateTime.UtcNow;
+
+		/// <summary>The most recent date we are interested in.</summary>
+		public DateTime DateTo = DateTime.UtcNow;
+
+		/// <summary>The branch name to filter by.</summary>
+		public string BranchName = "";
+
+		/// <summary>The version name to filter by.</summary>
+		public string VersionName = "";
+
+		/// <summary>The platform name to filter by.</summary>
+		public string PlatformName = "";
+
+		/// <summary>The game to filter by.</summary>
+		public string GameName = "";
+
+		private string PreviousTerm = "";
+
+	    public string EngineMode = "";
+
+		/// <summary>
+		/// Intelligently extract parameters from a web request that could be in a form or a query string (GET vs. POST).
+		/// </summary>
+		/// <param name="Request">The owning HttpRequest.</param>
+		/// <param name="Form">The form passed to the controller.</param>
+		/// <param name="Key">The parameter we wish to retrieve a value for.</param>
+		/// <param name="DefaultValue">The default value to set if the parameter is null or empty.</param>
+		/// <param name="Result">The value of the requested parameter.</param>
+		/// <returns>true if the parameter was found to be not null or empty.</returns>
+		private bool GetFormParameter( HttpRequestBase Request, FormCollection Form, string Key, string DefaultValue, out string Result )
+		{
+			Result = DefaultValue;
+
+			string Value = "";
+			if (Form.Count == 0)
+			{
+				Value = Request.QueryString[Key];
+			}
+			else
+			{
+				Value = Form[Key];
+			}
+
+			if (!string.IsNullOrEmpty( Value ))
+			{
+				Result = Value.Trim();
+				return true;
+			}
+
+			return false;
+		}
+
+		/// <summary>
+		/// Same as GetFormParameter, but ignores validation. Only used for search query.
+		/// </summary>
+		/// <param name="Request">The owning HttpRequest.</param>
+		/// <param name="Form">The form passed to the controller.</param>
+		/// <param name="Key">The parameter we wish to retrieve a value for.</param>
+		/// <param name="DefaultValue">The default value to set if the parameter is null or empty.</param>
+		/// <param name="Result">The value of the requested parameter.</param>
+		/// <returns>true if the parameter was found to be not null or empty.</returns>
+		private bool GetFormParameterUnvalidated( HttpRequestBase Request, FormCollection Form, string Key, string DefaultValue, out string Result )
+		{
+			var Collection = Request.Unvalidated;
+
+			Result = DefaultValue;
+
+			string Value = "";
+			if (Form.Count == 0)
+			{
+				Value = Collection.QueryString[Key];
+			}
+			else
+			{
+				Value = Form[Key];
+			}
+
+			if (!string.IsNullOrEmpty( Value ))
+			{
+				Result = Value.Trim();
+				return true;
+			}
+
+			return false;
+		}
+
+		/// <summary>
+		/// Attempt to read a form item as a date in milliseconds since 1970
+		/// </summary>
+		/// <param name="Request">The request that contains parameters if the form does not.</param>
+		/// <param name="Form">The form that contains parameters if the request does not.</param>
+		/// <param name="Key">Name of item in form/request</param>
+		/// <param name="Date">Result, only modified on succes</param>
+		void TryParseDate( HttpRequestBase Request, FormCollection Form, string Key, ref DateTime Date )
+		{
+			string MillisecondsString;
+			if (!GetFormParameter( Request, Form, Key, "", out MillisecondsString ))
+			{
+				return;
+			}
+
+			try
+			{
+				Date = CrashesViewModel.Epoch.AddMilliseconds( long.Parse( MillisecondsString ) );
+			}
+			catch (FormatException)
+			{
+			}
+		}
+
+		/// <summary>
+		/// Extract all the parameters from the client user input.
+		/// </summary>
+		/// <param name="Request">The request that contains parameters if the form does not.</param>
+		/// <param name="Form">The form that contains parameters if the request does not.</param>
+		/// <param name="DefaultSortTerm">The default sort term to use.</param>
+		public FormHelper( HttpRequestBase Request, FormCollection Form, string DefaultSortTerm )
+		{
+			// Set up Default values if there is no QueryString and set values to the Query string if it is there.
+			GetFormParameterUnvalidated( Request, Form, "SearchQuery", SearchQuery, out SearchQuery );
+			GetFormParameter( Request, Form, "UsernameQuery", UsernameQuery, out UsernameQuery );
+			GetFormParameter( Request, Form, "EpicIdOrMachineQuery", EpicIdOrMachineQuery, out EpicIdOrMachineQuery );
+			GetFormParameter( Request, Form, "JiraQuery", JiraQuery, out JiraQuery );
+            GetFormParameter( Request, Form, "JiraId", JiraId, out JiraId);
+			GetFormParameter( Request, Form, "MessageQuery", MessageQuery, out MessageQuery );
+			GetFormParameter( Request, Form, "BuggId", BuggId, out BuggId );
+			GetFormParameter( Request, Form, "BuiltFromCL", BuiltFromCL, out BuiltFromCL );
+			GetFormParameter( Request, Form, "SortTerm", DefaultSortTerm, out SortTerm );
+			GetFormParameter( Request, Form, "CrashType", CrashType, out CrashType );
+			GetFormParameter( Request, Form, "UserGroup", UserGroup, out UserGroup );
+			GetFormParameter( Request, Form, "BranchName", BranchName, out BranchName );
+			GetFormParameter( Request, Form, "VersionName", VersionName, out VersionName );
+			GetFormParameter( Request, Form, "PlatformName", PlatformName, out PlatformName );
+			GetFormParameter( Request, Form, "GameName", GameName, out GameName );
+            GetFormParameter(Request, Form, "EngineMode", EngineMode, out EngineMode);
+
+
+			string PageString = Page.ToString();
+			if (GetFormParameter( Request, Form, "Page", PageString, out PageString ))
+			{
+				if (!int.TryParse( PageString, out Page ) || Page < 1)
+				{
+					Page = 1;
+				}
+			}
+
+			string PageSizeString = PageSize.ToString();
+			if (GetFormParameter( Request, Form, "PageSize", PageSizeString, out PageSizeString ))
+			{
+				if (!int.TryParse( PageSizeString, out PageSize ) || PageSize < 1)
+				{
+					PageSize = 100;
+				}
+			}
+
+			GetFormParameter( Request, Form, "SortOrder", SortOrder, out SortOrder );
+			GetFormParameter( Request, Form, "PreviousOrder", PreviousOrder, out PreviousOrder );
+			GetFormParameter( Request, Form, "PreviousTerm", PreviousTerm, out PreviousTerm );
+			
+			DateFrom = DateTime.Today.AddDays( -7 ).ToUniversalTime();
+			TryParseDate( Request, Form, "DateFrom", ref DateFrom );
+
+			DateTo = DateTime.Today.ToUniversalTime();
+			TryParseDate( Request, Form, "DateTo", ref DateTo );
+
+			// Set the sort order 
+			if (PreviousOrder == "Descending" && PreviousTerm == SortTerm)
+			{
+				SortOrder = "Ascending";
+			}
+			else if (PreviousOrder == "Ascending" && PreviousTerm == SortTerm)
+			{
+				SortOrder = "Descending";
+			}
+			else if (string.IsNullOrEmpty( PreviousOrder ))
+			{
+				// keep SortOrder Where it's at.
+			}
+			else
+			{
+				SortOrder = "Descending";
+			}
+		}
+	}
+}
\ No newline at end of file
diff --git a/Engine/Source/Programs/CrashReporter/CrashReportWebSite/ViewModels/PagingInfo.cs b/Engine/Source/Programs/CrashReporter/CrashReportWebSite/ViewModels/PagingInfo.cs
new file mode 100644
--- /dev/null
+++ b/Engine/Source/Programs/CrashReporter/CrashReportWebSite/ViewModels/PagingInfo.cs
@@ -0,0 +1,128 @@
+﻿// Copyright 1998-2016 Epic Games, Inc. All Rights Reserved.
+
+using System;
+
+namespace Tools.CrashReporter.CrashReportWebSite.ViewModels
+{
+	/// <summary>
+	/// A class to handle pagination of large lists.
+	/// </summary>
+	public class PagingInfo
+	{
+		/// <summary>The current page of elements to display.</summary>
+		public int CurrentPage { get; set; }
+		/// <summary>The total number of elements.</summary>
+		public int TotalResults { get; set; }
+		/// <summary>The number of elements per page.</summary>
+		public int PageSize { get; set; }
+
+		/// <summary>
+		/// The total number of pages.
+		/// </summary>
+		public int PageCount
+		{
+			// return total number of pages
+			get 
+			{ 
+				return ( int )Math.Ceiling( ( float )TotalResults / ( float )PageSize ); 
+			}
+		}
+
+		/// <summary>
+		/// Return the Index for the first result in the paging list to display.
+		/// </summary>
+		public int FirstPageIndex
+		{
+			get 
+			{
+				if( CurrentPage >= PagingListSize )
+				{
+					return CurrentPage;
+				}
+
+				return 1;
+			}
+		}
+
+		/// <summary>
+		/// Return the Index for the last result in the paging list to display.
+		/// </summary>
+		public int LastPageIndex
+		{
+			get 
+			{
+				if( ( ( FirstPageIndex - 1 ) + PagingListSize ) < LastPage )
+				{
+					return ( FirstPageIndex - 1 ) + PagingListSize;
+				}
+
+				return LastPage;
+			}
+		}
+
+		/// <summary>
+		/// Return the number of pages to display on the screen
+		/// </summary>
+		public int PagingListSize
+		{
+			get 
+			{ 
+				return 10; 
+			}
+		}
+
+		/// <summary>
+		/// Return the Index for the previous result in the paging list to display.
+		/// </summary>
+		public int PreviousPageIndex
+		{
+			get 
+			{
+				if( CurrentPage - 1 > 1 )
+				{
+					return CurrentPage - 1;
+				}
+
+				return 1;
+			}
+		}
+
+		/// <summary>
+		/// Return the Index for the previous result in the paging list to display.
+		/// </summary>
+		public int NextPageIndex
+		{
+			get 
+			{
+				if( CurrentPage + 1 <= LastPage )
+				{
+					return CurrentPage + 1;
+				}
+
+				return LastPage;
+			}
+		}
+
+		/// <summary>
+		/// Return the Index for the first result in the paging list to display.
+		/// </summary>
+		public int FirstPage
+		{
+			get 
+			{ 
+				return 1; 
+			}
+		}
+
+		/// <summary>
+		/// Return the Index for the last result.
+		/// </summary>
+		public int LastPage
+		{
+			get 
+			{
+				return PageCount; 
+			}
+		}
+	}
+}
\ No newline at end of file
diff --git a/Engine/Source/Programs/CrashReporter/CrashReportWebSite/ViewModels/UsersViewModel.cs b/Engine/Source/Programs/CrashReporter/CrashReportWebSite/ViewModels/UsersViewModel.cs
new file mode 100644
--- /dev/null
+++ b/Engine/Source/Programs/CrashReporter/CrashReportWebSite/ViewModels/UsersViewModel.cs
@@ -0,0 +1,24 @@
+﻿// Copyright 1998-2016 Epic Games, Inc. All Rights Reserved.
+
+using System;
+using System.Collections.Generic;
+using System.Linq;
+using System.Web.Mvc;
+
+namespace Tools.CrashReporter.CrashReportWebSite.Models
+{
+	/// <summary>
+	/// The view model for the user administrative page.
+	/// </summary>
+	public class UsersViewModel
+	{
+		/// <summary>The currently selected user group.</summary>
+		public string UserGroup { get; set; }
+
+		/// <summary>The currently selected user group.</summary>
+		public HashSet<string> Users { get; set; }
+
+		/// <summary>The number of users in each group.</summary>
+		public Dictionary<string, int> GroupCounts { get; set; }
+	}
+}
\ No newline at end of file
diff --git a/Engine/Source/Programs/CrashReporter/CrashReportWebSite/Views/Buggs/Index.aspx b/Engine/Source/Programs/CrashReporter/CrashReportWebSite/Views/Buggs/Index.aspx
--- a/Engine/Source/Programs/CrashReporter/CrashReportWebSite/Views/Buggs/Index.aspx
+++ b/Engine/Source/Programs/CrashReporter/CrashReportWebSite/Views/Buggs/Index.aspx
@@ -1,7 +1,7 @@
 ﻿<%-- // Copyright 1998-2016 Epic Games, Inc. All Rights Reserved. --%>
 
 <%@ Page Title="" Language="C#" MasterPageFile="~/Views/Shared/Site.Master" Inherits="System.Web.Mvc.ViewPage<BuggsViewModel>" %>
-<%@ Import Namespace="Tools.CrashReporter.CrashReportWebSite.Models" %>
+<%@ Import Namespace="Tools.CrashReporter.CrashReportWebSite.ViewModels" %>
 
 <asp:Content ID="StyleSheet" ContentPlaceHolderID="CssContent" runat="server">
 	<link href="../../Content/Site.css" rel="stylesheet" type="text/css" />
diff --git a/Engine/Source/Programs/CrashReporter/CrashReportWebSite/Views/Buggs/Show.aspx b/Engine/Source/Programs/CrashReporter/CrashReportWebSite/Views/Buggs/Show.aspx
--- a/Engine/Source/Programs/CrashReporter/CrashReportWebSite/Views/Buggs/Show.aspx
+++ b/Engine/Source/Programs/CrashReporter/CrashReportWebSite/Views/Buggs/Show.aspx
@@ -1,7 +1,8 @@
 ﻿<%-- // Copyright 1998-2016 Epic Games, Inc. All Rights Reserved. --%>
 
 <%@ Page Title="" Language="C#" MasterPageFile="~/Views/Shared/Site.Master" Inherits="System.Web.Mvc.ViewPage<BuggViewModel>" %>
-<%@ Import Namespace="Tools.CrashReporter.CrashReportWebSite.Models" %>
+<%@ Import Namespace="Tools.CrashReporter.CrashReportWebSite.DataModels" %>
+<%@ Import Namespace="Tools.CrashReporter.CrashReportWebSite.ViewModels" %>
 
 <asp:Content ID="StyleSheet" ContentPlaceHolderID="CssContent" runat="server">
 	<link href="../../Content/Site.css" rel="stylesheet" type="text/css" />
@@ -118,14 +119,14 @@
 		<%using( Html.BeginForm( "Show", "Buggs" ) )
 		{	%>
 
-			<%if (string.IsNullOrEmpty( Model.Bugg.Jira ))
+			<%if (string.IsNullOrEmpty( Model.Bugg.TTPID ))
 			{ %>
 				<dd class='even'>
 					<input type="submit" value="CopyToJira" title="<%=Model.Bugg.ToTooltip()%>" id="id <%=Model.Bugg.Id%>" name="CopyToJira-<%=Model.Bugg.Id%>" class="SearchButton CopyToJiraStyle" />
 				</dd>
 			<% } else { %>
 				<dd class='even'>
-					<a href="https://jira.ol.epicgames.net/browse/<%=Model.Bugg.Jira%>" target="_blank"><%=Model.Bugg.Jira%></a>
+					<a href="https://jira.ol.epicgames.net/browse/<%=Model.Bugg.TTPID%>" target="_blank"><%=Model.Bugg.TTPID%></a>
 				</dd>
 			<% } %>
 
@@ -262,7 +263,7 @@
 							<td class="TimeOfCrash">
 								<%=CrashInstance.GetTimeOfCrash()[0]%> &nbsp;
 								<%=CrashInstance.GetTimeOfCrash()[1]%><br />
-								Change: <%= CrashInstance.BuiltFromCL %><br />
+								Change: <%= CrashInstance.ChangeListVersion %><br />
 							</td>
 							<td class="Summary">
 								<%=CrashInstance.UserName%><br />
diff --git a/Engine/Source/Programs/CrashReporter/CrashReportWebSite/Views/Crashes/Index.aspx b/Engine/Source/Programs/CrashReporter/CrashReportWebSite/Views/Crashes/Index.aspx
--- a/Engine/Source/Programs/CrashReporter/CrashReportWebSite/Views/Crashes/Index.aspx
+++ b/Engine/Source/Programs/CrashReporter/CrashReportWebSite/Views/Crashes/Index.aspx
@@ -1,8 +1,8 @@
 ﻿<%-- // Copyright 1998-2016 Epic Games, Inc. All Rights Reserved. --%>
 
-<%@ Page Title="" Language="C#" MasterPageFile="~/Views/Shared/Site.Master" Inherits="System.Web.Mvc.ViewPage<CrashesViewModel>" %>
+<%@ Import Namespace="Tools.CrashReporter.CrashReportWebSite.ViewModels" %>
 
-<%@ Import Namespace="Tools.CrashReporter.CrashReportWebSite.Models" %>
+<%@ Page Title="" Language="C#" MasterPageFile="~/Views/Shared/Site.Master" Inherits="System.Web.Mvc.ViewPage<CrashesViewModel>" %>
 
 <asp:Content ID="StyleSheet" ContentPlaceHolderID="CssContent" runat="server">
 	<link href="../../Content/Site.css" rel="stylesheet" type="text/css" />
@@ -86,12 +86,11 @@
 			});
 		});
 	</script>
-
 </asp:Content>
 
 <asp:Content ID="AboveMainContent"  ContentPlaceHolderID="AboveMainContent" runat="server" >
 	<div style="clear: both;">
-		<small style="color: lightgray;">Generated in <%=Model.GenerationTime%> second(s)</small>
+		<small style="color: lightgray;">Generated in <%=Model.GenerationTime %> second(s)</small>
 		<br />
 
 		<div id='SearchForm' style="clear: both">
@@ -225,7 +224,16 @@
 							<%=Html.DropDownListFor( m=>m.PlatformName, Model.PlatformNames )%>
 						</td>
 					</tr>
-
+                    <tr>
+						<td rowspan="4" >
+						</td>
+                        <td>
+							<p class="SearchTextTitle">Filter by Engine Mode</p>
+						</td>
+						<td>
+							<%=Html.DropDownListFor( m=>m.EngineMode, Model.EngineModes )%>
+						</td>
+					</tr>
 					<tr>
 						<td colspan="7" >
 							<input type="submit" value="Search" class='SearchButton' />
@@ -311,10 +319,11 @@
 					}
 					, 
 					new { style = "color:black; text-decoration:none;" } )%></span>
-			</div>
-			<%Html.RenderPartial( "/Views/Crashes/ViewCrash.ascx" );%>
-		</form>
-	</div>
+			    </div>
+			    <%Html.RenderPartial( "/Views/Crashes/ViewCrash.ascx" );%>
+            </div>
+	    </form>
+    </div>
 </div>
 
 <div class="PaginationBox">
diff --git a/Engine/Source/Programs/CrashReporter/CrashReportWebSite/Views/Crashes/Show.aspx b/Engine/Source/Programs/CrashReporter/CrashReportWebSite/Views/Crashes/Show.aspx
--- a/Engine/Source/Programs/CrashReporter/CrashReportWebSite/Views/Crashes/Show.aspx
+++ b/Engine/Source/Programs/CrashReporter/CrashReportWebSite/Views/Crashes/Show.aspx
@@ -1,7 +1,8 @@
 ﻿<%-- // Copyright 1998-2016 Epic Games, Inc. All Rights Reserved. --%>
 
 <%@ Page Title="" Language="C#" MasterPageFile="~/Views/Shared/Site.Master" Inherits="System.Web.Mvc.ViewPage<CrashViewModel>" %>
-<%@ Import Namespace="Tools.CrashReporter.CrashReportWebSite.Models" %>
+<%@ Import Namespace="Tools.CrashReporter.CrashReportWebSite.DataModels" %>
+<%@ Import Namespace="Tools.CrashReporter.CrashReportWebSite.ViewModels" %>
 
 <asp:Content ID="StyleSheet" ContentPlaceHolderID="CssContent" runat="server">
 	<link href="../../Content/Site.css" rel="stylesheet" type="text/css" />
@@ -116,7 +117,7 @@
 				
 			<dt>Saved Files</dt>
 				<dd class='even'>
-					<%if( Model.Crash.HasLogFile() ) 
+					<%if( Model.Crash.GetHasLogFile() ) 
 					{ %>
 						<a style='text-decoration:none; vertical-align:middle;' href='<%=Model.Crash.GetLogUrl() %>'><img src="../../Content/Images/Icons/log.png" style="height:20px;border:none;" />&nbsp;Log</a>
 						<%
@@ -125,7 +126,7 @@
 				</dd>
 			<dt>&nbsp;</dt>
 				<dd>	
-					<%if( Model.Crash.HasMiniDumpFile() ) 
+					<%if( Model.Crash.GetHasMiniDumpFile() ) 
 					{ %>
 						<a style='text-decoration:none;' href='<%=Model.Crash.GetMiniDumpUrl() %>' title="<%=Model.Crash.GetDumpTitle()%>"><img src="../../Content/Images/Icons/miniDump.png" style="height:20px;border:none;" />&nbsp;<%=Model.Crash.GetDumpName() %></a>
 						<%
@@ -133,7 +134,7 @@
 				</dd>
 			<dt>&nbsp;</dt>
 				<dd>
-					<%if( Model.Crash.HasVideoFile() ) 
+					<%if( Model.Crash.GetHasVideoFile() ) 
 					{ %>
 						<a id="VideoLink" style="text-decoration:none;" href='<%=Model.Crash.GetVideoUrl() %>'><img src="../../Content/Images/Icons/video.png" style="height:20px;border:none;" />&nbsp;Video</a>
 						<%
@@ -149,9 +150,9 @@
 				</dd>
 			<dt>&nbsp;</dt>
 				<dd>
-					<%if (Model.Crash.HasDiagnosticsFile() ) 
+					<%if (Model.Crash.GetHasDiagnosticsFile() ) 
 					{ %>
-						<a id="A1" style="text-decoration:none;" href='<%=Model.Crash.GetDiagnosticsUrl() %>'><img src="../../Content/Images/Icons/log.png" style="height:20px;border:none;" />&nbsp;Diags</a>
+						<a id="A2" style="text-decoration:none;" href='<%=Model.Crash.GetDiagnosticsUrl() %>'><img src="../../Content/Images/Icons/log.png" style="height:20px;border:none;" />&nbsp;Diags</a>
 						<%
 					} %>
 				</dd>
@@ -159,7 +160,7 @@
 				<dd>
 					<%if( Model.Crash.HasMetaDataFile() ) 
 					{ %>
-						<a id="A2" style="text-decoration:none;" href='<%=Model.Crash.GetMetaDataUrl() %>'><img src="../../Content/Images/Icons/log.png" style="height:20px;border:none;" />&nbsp;WERInfo</a>
+						<a id="A3" style="text-decoration:none;" href='<%=Model.Crash.GetMetaDataUrl() %>'><img src="../../Content/Images/Icons/log.png" style="height:20px;border:none;" />&nbsp;WERInfo</a>
 						<%
 					} %>
 				</dd>
@@ -182,7 +183,7 @@
 				</dd>
 
 			<dt>User Group</dt>
-				<dd><%=Html.DisplayFor(m => Model.Crash.UserGroupName) %></dd>
+				<dd><%=Html.DisplayFor(m => Model.Crash.User.UserGroup.Name) %></dd>
 
 			<dt>Game Name</dt>
 				<dd class='even'><%=Html.DisplayFor(m => Model.Crash.GameName) %></dd>
@@ -197,13 +198,13 @@
 				<dd ><small><%=Html.DisplayFor(m => Model.Crash.PlatformName) %></small></dd>
 
 			<dt>Machine Id</dt> 
-				<dd class='even'><small><%=Html.DisplayFor(m => Model.Crash.MachineId) %></small></dd>
+				<dd class='even'><small><%=Html.DisplayFor(m => Model.Crash.ComputerName) %></small></dd>
 
 			<dt>Epic Account Id</dt> 
 				<dd class='even'><small><%=Html.DisplayFor(m => Model.Crash.EpicAccountId) %></small></dd>
 
 			<dt>Allowed to contact</dt> 
-				<dd class='even'><%=Html.DisplayFor(m => Model.Crash.AllowToBeContacted) %></dd>
+				<dd class='even'><%=Html.DisplayFor(m => Model.Crash.Processed) %></dd>
 
 			<dt>Build Version</dt>
 				<dd ><%=Html.DisplayFor(m => Model.Crash.BuildVersion) %></dd>
@@ -212,7 +213,7 @@
 				<dd class='even'><%=Html.DisplayFor(m => Model.Crash.FixedChangeList) %></dd>
 
 			<dt>Changelist #</dt>
-				<dd class='even'><%=Html.DisplayFor(m => Model.Crash.BuiltFromCL) %></dd>
+				<dd class='even'><%=Html.DisplayFor(m => Model.Crash.ChangeListVersion) %></dd>
 
 			<dt>JIRA</dt>
 				<%--<dd ><%=Html.DisplayFor(m => Model.Crash.Jira) %></dd>--%>
diff --git a/Engine/Source/Programs/CrashReporter/CrashReportWebSite/Views/Helpers/PagingHelper.cs b/Engine/Source/Programs/CrashReporter/CrashReportWebSite/Views/Helpers/PagingHelper.cs
--- a/Engine/Source/Programs/CrashReporter/CrashReportWebSite/Views/Helpers/PagingHelper.cs
+++ b/Engine/Source/Programs/CrashReporter/CrashReportWebSite/Views/Helpers/PagingHelper.cs
@@ -5,6 +5,7 @@ using System.Web.Mvc;
 using System.Text;
 
 using Tools.CrashReporter.CrashReportWebSite.Models;
+using Tools.CrashReporter.CrashReportWebSite.ViewModels;
 
 namespace Tools.CrashReporter.CrashReportWebSite.Views.Helpers
 {
diff --git a/Engine/Source/Programs/CrashReporter/CrashReportWebSite/Views/Helpers/UrlHelperExtension.cs b/Engine/Source/Programs/CrashReporter/CrashReportWebSite/Views/Helpers/UrlHelperExtension.cs
--- a/Engine/Source/Programs/CrashReporter/CrashReportWebSite/Views/Helpers/UrlHelperExtension.cs
+++ b/Engine/Source/Programs/CrashReporter/CrashReportWebSite/Views/Helpers/UrlHelperExtension.cs
@@ -4,6 +4,7 @@ using System.Web.Mvc;
 using System.Text;
 
 using Tools.CrashReporter.CrashReportWebSite.Models;
+using Tools.CrashReporter.CrashReportWebSite.ViewModels;
 
 namespace Tools.CrashReporter.CrashReportWebSite.Views.Helpers
 {
diff --git a/Engine/Source/Programs/CrashReporter/CrashReportWebSite/Views/Home/Index.aspx b/Engine/Source/Programs/CrashReporter/CrashReportWebSite/Views/Home/Index.aspx
--- a/Engine/Source/Programs/CrashReporter/CrashReportWebSite/Views/Home/Index.aspx
+++ b/Engine/Source/Programs/CrashReporter/CrashReportWebSite/Views/Home/Index.aspx
@@ -1,8 +1,8 @@
 ﻿<%-- // Copyright 1998-2016 Epic Games, Inc. All Rights Reserved. --%>
 
-<%@ Page Title="" Language="C#" MasterPageFile="~/Views/Shared/Site.Master" Inherits="System.Web.Mvc.ViewPage<CrashesViewModel>" %>
+<%@ Import Namespace="Tools.CrashReporter.CrashReportWebSite.ViewModels" %>
 
-<%@ Import Namespace="Tools.CrashReporter.CrashReportWebSite.Models" %>
+<%@ Page Title="" Language="C#" MasterPageFile="~/Views/Shared/Site.Master" Inherits="System.Web.Mvc.ViewPage<CrashesViewModel>" %>
 
 <asp:Content ID="StyleSheet" ContentPlaceHolderID="CssContent" runat="server">
 	<link href="../../Content/Site.css" rel="stylesheet" type="text/css" />
@@ -78,14 +78,12 @@
 			});
 		});
 	</script>
-
 </asp:Content>
 
 <asp:Content ID="AboveMainContent"  ContentPlaceHolderID="AboveMainContent" runat="server" >
 	<div style="clear: both;">
 		<small style="color: lightgray;">Generated in <%=Model.GenerationTime%> second(s)</small>
 		<br />
-
 		<div id='SearchForm' style="clear: both">
 			<div id="TitleBar" class="TitleBar">
 				SEARCH OPTIONS
@@ -217,7 +215,17 @@
 							<%=Html.DropDownListFor( m=>m.PlatformName, Model.PlatformNames )%>
 						</td>
 					</tr>
-
+                    
+                    <tr>
+						<td rowspan="1" >
+						</td>
+                        <td>
+							<p class="SearchTextTitle">Filter by Engine Mode</p>
+						</td>
+						<td>
+							<%=Html.DropDownListFor( m=>m.EngineMode, Model.EngineModes )%>
+						</td>
+					</tr>
 					<tr>
 						<td colspan="7" >
 							<input type="submit" value="Search" class='SearchButton' />
diff --git a/Engine/Source/Programs/CrashReporter/CrashReportWebSite/Views/Reports/Index.aspx b/Engine/Source/Programs/CrashReporter/CrashReportWebSite/Views/Reports/Index.aspx
--- a/Engine/Source/Programs/CrashReporter/CrashReportWebSite/Views/Reports/Index.aspx
+++ b/Engine/Source/Programs/CrashReporter/CrashReportWebSite/Views/Reports/Index.aspx
@@ -1,7 +1,7 @@
 ﻿<%-- // Copyright 1998-2016 Epic Games, Inc. All Rights Reserved. --%>
 
 <%@ Page Title="" Language="C#" MasterPageFile="~/Views/Shared/Site.Master" Inherits="System.Web.Mvc.ViewPage<ReportsViewModel>" %>
-<%@ Import Namespace="Tools.CrashReporter.CrashReportWebSite.Models" %>
+<%@ Import Namespace="Tools.CrashReporter.CrashReportWebSite.ViewModels" %>
 
 <asp:Content ID="StyleSheet" ContentPlaceHolderID="CssContent" runat="server">
 	<link href="../../Content/Site.css" rel="stylesheet" type="text/css" />
diff --git a/Engine/Source/Programs/CrashReporter/CrashReportWebSite/Views/Reports/Unsubscribe.aspx b/Engine/Source/Programs/CrashReporter/CrashReportWebSite/Views/Reports/Unsubscribe.aspx
--- a/Engine/Source/Programs/CrashReporter/CrashReportWebSite/Views/Reports/Unsubscribe.aspx
+++ b/Engine/Source/Programs/CrashReporter/CrashReportWebSite/Views/Reports/Unsubscribe.aspx
@@ -1,6 +1,6 @@
 ﻿<%-- // Copyright 1998-2016 Epic Games, Inc. All Rights Reserved. --%>
 
-<%@ Page Title="" Language="C#" MasterPageFile="~/Views/Shared/Site.Master" Inherits="System.Web.Mvc.ViewPage<ReportsViewModel>" %>
+<%@ Page Title="" Language="C#" MasterPageFile="~/Views/Shared/Site.Master" Inherits="System.Web.Mvc.ViewPage<Tools.CrashReporter.CrashReportWebSite.ViewModels.ReportsViewModel>" %>
 <%@ Import Namespace="Tools.CrashReporter.CrashReportWebSite.Models" %>
 
 <asp:Content ID="StyleSheet" ContentPlaceHolderID="CssContent" runat="server">
diff --git a/Engine/Source/Programs/CrashReporter/CrashReportWebSite/Web.Development.config b/Engine/Source/Programs/CrashReporter/CrashReportWebSite/Web.Development.config
new file mode 100644
--- /dev/null
+++ b/Engine/Source/Programs/CrashReporter/CrashReportWebSite/Web.Development.config
@@ -0,0 +1,37 @@
+<?xml version="1.0" encoding="utf-8"?>
+
+<!-- For more information on using web.config transformation visit http://go.microsoft.com/fwlink/?LinkId=125889 -->
+
+<configuration xmlns:xdt="http://schemas.microsoft.com/XML-Document-Transform">
+  <!--
+    In the example below, the "SetAttributes" transform will change the value of 
+    "connectionString" to use "ReleaseSQLServer" only when the "Match" locator 
+    finds an attribute "name" that has a value of "MyDB".
+    
+    <connectionStrings>
+      <add name="MyDB" 
+        connectionString="Data Source=ReleaseSQLServer;Initial Catalog=MyReleaseDB;Integrated Security=True" 
+        xdt:Transform="SetAttributes" xdt:Locator="Match(name)"/>
+    </connectionStrings>
+  -->
+  <connectionStrings  xdt:Transform="Replace">
+    <add name="CrashReportConnectionString"
+      connectionString="Data Source=db-09;Initial Catalog=CrashReport;Server=db-09;Database=CrashReport;Integrated Security=true"
+      providerName="System.Data.SqlClient"/>
+    <add name="Hangfire" connectionString="Data Source=db-09;Initial Catalog=CrashReporterSchduledTask;Server=db-09;Database=CrashReport;Integrated Security=true" providerName="System.Data.SqlClient" />
+  </connectionStrings>
+  <system.web>
+    <compilation xdt:Transform="RemoveAttributes(debug)" />
+    <!--
+      In the example below, the "Replace" transform will replace the entire 
+      <customErrors> section of your web.config file.
+      Note that because there is only one customErrors section under the 
+      <system.web> node, there is no need to use the "xdt:Locator" attribute.
+      
+      <customErrors defaultRedirect="GenericError.htm"
+        mode="RemoteOnly" xdt:Transform="Replace">
+        <error statusCode="500" redirect="InternalError.htm"/>
+      </customErrors>
+    -->
+  </system.web>
+</configuration>
\ No newline at end of file
diff --git a/Engine/Source/Programs/CrashReporter/CrashReportWebSite/Web.Release.config b/Engine/Source/Programs/CrashReporter/CrashReportWebSite/Web.Release.config
--- a/Engine/Source/Programs/CrashReporter/CrashReportWebSite/Web.Release.config
+++ b/Engine/Source/Programs/CrashReporter/CrashReportWebSite/Web.Release.config
@@ -18,7 +18,8 @@
     <add name="CrashReportConnectionString"
       connectionString="Data Source=db-09;Initial Catalog=CrashReport;Server=db-09;Database=CrashReport;Integrated Security=true"
       providerName="System.Data.SqlClient"/>
-    <add name="Hangfire" connectionString="Data Source=db-09;Initial Catalog=CrashReporterSchduledTask;Server=db-09;Database=CrashReport;Integrated Security=true" providerName="System.Data.SqlClient" />
+    <add name="CrashReportEntities" connectionString="metadata=res://*/DataModels.CrashReports.csdl|res://*/DataModels.CrashReports.ssdl|res://*/DataModels.CrashReports.msl;provider=System.Data.SqlClient;provider connection string=&quot;data source=db-09;initial catalog=CrashReport;integrated security=True;MultipleActiveResultSets=True;App=EntityFramework&quot;" providerName="System.Data.EntityClient" />
+  <add name="Hangfire" connectionString="Data Source=db-09;Initial Catalog=CrashReporterSchduledTask;Server=db-09;Database=CrashReport;Integrated Security=true" providerName="System.Data.SqlClient" />
   </connectionStrings>
   <system.web>
     <compilation xdt:Transform="RemoveAttributes(debug)" />
diff --git a/Engine/Source/Programs/CrashReporter/CrashReportWebSite/Web.config b/Engine/Source/Programs/CrashReporter/CrashReportWebSite/Web.config
--- a/Engine/Source/Programs/CrashReporter/CrashReportWebSite/Web.config
+++ b/Engine/Source/Programs/CrashReporter/CrashReportWebSite/Web.config
@@ -1,4 +1,4 @@
-<?xml version="1.0"?>
+﻿<?xml version="1.0" encoding="utf-8"?>
 <!-- 
     Note: As an alternative to hand editing this file you can use the 
     web admin tool to configure settings for your application. Use
@@ -9,22 +9,18 @@
 -->
 <configuration>
   <configSections>
-    <sectionGroup name="applicationSettings"
-      type="System.Configuration.ApplicationSettingsGroup, System, Version=4.0.0.0, Culture=neutral, PublicKeyToken=b77a5c561934e089">
-      <section name="Tools.CrashReporter.CrashReportWebSite.Properties.Settings"
-        type="System.Configuration.ClientSettingsSection, System, Version=4.0.0.0, Culture=neutral, PublicKeyToken=b77a5c561934e089" requirePermission="false"/>
+    <sectionGroup name="applicationSettings" type="System.Configuration.ApplicationSettingsGroup, System, Version=4.0.0.0, Culture=neutral, PublicKeyToken=b77a5c561934e089">
+      <section name="Tools.CrashReporter.CrashReportWebSite.Properties.Settings" type="System.Configuration.ClientSettingsSection, System, Version=4.0.0.0, Culture=neutral, PublicKeyToken=b77a5c561934e089" requirePermission="false" />
     </sectionGroup>
-  </configSections>
-  <appSettings/>
+    
+    <section name="entityFramework" type="System.Data.Entity.Internal.ConfigFile.EntityFrameworkSection, EntityFramework, Version=6.0.0.0, Culture=neutral, PublicKeyToken=b77a5c561934e089" requirePermission="false" />
+  <!-- For more information on Entity Framework configuration, visit http://go.microsoft.com/fwlink/?LinkID=237468 --></configSections>
+  <appSettings />
   <connectionStrings>
-    <add name="CrashReportConnectionString"
-      connectionString="Data Source=db-09;Initial Catalog=CrashReport;Server=db-09;Database=CrashReport;Integrated Security=true"
-      providerName="System.Data.SqlClient"/>
-    <add name="CrashReportTestConnectionString"
-      connectionString="Data Source=CrashReporterDev;Initial Catalog=CrashReport;Server=NCLRNDEV-02;Database=CrashReporterDev;Integrated Security=true"
-      providerName="System.Data.SqlClient"/>
+    <add name="CrashReportConnectionString" connectionString="Data Source=db-09;Initial Catalog=CrashReport;Server=db-09;Database=CrashReport;Connection Timeout=300;Integrated Security=true" providerName="System.Data.SqlClient" />
+    <add name="CrashReportTestConnectionString" connectionString="Data Source=CrashReporterDev;Initial Catalog=CrashReport;Server=NCLRNDEV-02;Database=CrashReporterDev;Integrated Security=true" providerName="System.Data.SqlClient" />
     <add name="Hangfire" connectionString="Data Source=db-09;Initial Catalog=CrashReporterSchduledTask;Server=db-09;Database=CrashReport;Integrated Security=true" providerName="System.Data.SqlClient" />
-  </connectionStrings>
+  <add name="CrashReportEntities" connectionString="metadata=res://*/DataModels.CrashReports.csdl|res://*/DataModels.CrashReports.ssdl|res://*/DataModels.CrashReports.msl;provider=System.Data.SqlClient;provider connection string=&quot;data source=db-09;initial catalog=CrashReport;integrated security=True;MultipleActiveResultSets=True;App=EntityFramework&quot;" providerName="System.Data.EntityClient" /></connectionStrings>
   <!--
     For a description of web.config changes see http://go.microsoft.com/fwlink/?LinkId=235367.
 
@@ -34,8 +30,7 @@
       </system.Web>
   -->
   <system.web>
-    <compilation debug="true" targetFramework="4.5">
-    </compilation>
+    <compilation debug="true" targetFramework="4.5"></compilation>
     <!--
             The <customErrors> section enables configuration 
             of what to do if/when an unhandled error occurs 
@@ -48,17 +43,17 @@
             <error statusCode="404" redirect="FileNotFound.htm"/>
         </customErrors>
     -->
-    <customErrors mode="Off"/>
+    <customErrors mode="Off" />
     <pages enableSessionState="false" clientIDMode="AutoID" controlRenderingCompatibilityVersion="4.0">
       <namespaces>
-        <add namespace="System.Web.Mvc"/>
-        <add namespace="System.Web.Mvc.Ajax"/>
-        <add namespace="System.Web.Mvc.Html"/>
-        <add namespace="System.Web.Routing"/>
-        <add namespace="Tools.CrashReporter.CrashReportWebSite.Views.Helpers"/>
+        <add namespace="System.Web.Mvc" />
+        <add namespace="System.Web.Mvc.Ajax" />
+        <add namespace="System.Web.Mvc.Html" />
+        <add namespace="System.Web.Routing" />
+        <add namespace="Tools.CrashReporter.CrashReportWebSite.Views.Helpers" />
       </namespaces>
     </pages>
-    <sessionState mode="Off"/>
+    <sessionState mode="Off" />
   </system.web>
   <!-- 
         The system.webServer section is required for running ASP.NET AJAX under Internet
@@ -66,17 +61,17 @@
   -->
   <system.webServer>
     <modules runAllManagedModulesForAllRequests="false">
-      <add name="UriCacheModule"/>
+      <add name="UriCacheModule" />
       <!--<add name="FileCacheModule"/>-->
-      <add name="TokenCacheModule"/>
-      <add name="RequestMonitorModule"/>
+      <add name="TokenCacheModule" />
+      <add name="RequestMonitorModule" />
     </modules>
     <handlers>
-      <remove name="UrlRoutingHandler"/>
+      <remove name="UrlRoutingHandler" />
     </handlers>
     <caching>
       <profiles>
-        <add extension=".aspx" policy="CacheForTimePeriod" kernelCachePolicy="CacheForTimePeriod" duration="00:02:30" varyByHeaders="Post"/>
+        <add extension=".aspx" policy="CacheForTimePeriod" kernelCachePolicy="CacheForTimePeriod" duration="00:02:30" varyByHeaders="Post" />
       </profiles>
     </caching>
   </system.webServer>
@@ -91,33 +86,57 @@
       <setting name="CrashReporterCSV" serializeAs="String">
         <value>/CrashReporterCSV/</value>
       </setting>
+      <setting name="JiraDeploymentAddress" serializeAs="String">
+        <value>https://jira.ol.epicgames.net/rest/api/2</value>
+      </setting>
+      <setting name="SitePath" serializeAs="String">
+        <value>\\devweb-02\Sites</value>
+      </setting>
+      <setting name="SlackWebhookUrl" serializeAs="String">
+        <value>https://hooks.slack.com/services/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</value>
+      </setting>
+      <setting name="SlackChannel" serializeAs="String">
+        <value>#crashreporterspam</value>
+      </setting>
+      <setting name="SlackUsername" serializeAs="String">
+        <value>CrashReportWebsite</value>
+      </setting>
+      <setting name="SlackEmoji" serializeAs="String">
+        <value>:rngesus:</value>
+      </setting>
     </Tools.CrashReporter.CrashReportWebSite.Properties.Settings>
   </applicationSettings>
   <runtime>
     <assemblyBinding xmlns="urn:schemas-microsoft-com:asm.v1">
       <dependentAssembly>
-        <assemblyIdentity name="System.Web.Helpers" publicKeyToken="31bf3856ad364e35"/>
-        <bindingRedirect oldVersion="1.0.0.0-3.0.0.0" newVersion="3.0.0.0"/>
+        <assemblyIdentity name="System.Web.Helpers" publicKeyToken="31bf3856ad364e35" />
+        <bindingRedirect oldVersion="1.0.0.0-3.0.0.0" newVersion="3.0.0.0" />
       </dependentAssembly>
       <dependentAssembly>
-        <assemblyIdentity name="System.Data.Linq" publicKeyToken="b77a5c561934e089"/>
-        <bindingRedirect oldVersion="1.0.0.0-4.0.0.0" newVersion="4.0.0.0"/>
+        <assemblyIdentity name="System.Data.Linq" publicKeyToken="b77a5c561934e089" />
+        <bindingRedirect oldVersion="1.0.0.0-4.0.0.0" newVersion="4.0.0.0" />
       </dependentAssembly>
       <dependentAssembly>
-        <assemblyIdentity name="System.Web.WebPages" publicKeyToken="31bf3856ad364e35"/>
-        <bindingRedirect oldVersion="1.0.0.0-3.0.0.0" newVersion="3.0.0.0"/>
+        <assemblyIdentity name="System.Web.WebPages" publicKeyToken="31bf3856ad364e35" />
+        <bindingRedirect oldVersion="1.0.0.0-3.0.0.0" newVersion="3.0.0.0" />
       </dependentAssembly>
       <dependentAssembly>
-        <assemblyIdentity name="System.Web.Mvc" publicKeyToken="31bf3856ad364e35"/>
-        <bindingRedirect oldVersion="1.0.0.0-5.2.3.0" newVersion="5.2.3.0"/>
+        <assemblyIdentity name="System.Web.Mvc" publicKeyToken="31bf3856ad364e35" />
+        <bindingRedirect oldVersion="0.0.0.0-5.2.3.0" newVersion="5.2.3.0" />
       </dependentAssembly>
     </assemblyBinding>
   </runtime>
   <system.net>
     <mailSettings>
       <smtp deliveryMethod="Network">
-        <network host="smtp.epicgames.net" defaultCredentials="false"/>
+        <network host="smtp.epicgames.net" defaultCredentials="false" />
       </smtp>
     </mailSettings>
   </system.net>
+  <entityFramework>
+    <defaultConnectionFactory type="System.Data.Entity.Infrastructure.SqlConnectionFactory, EntityFramework" />
+    <providers>
+      <provider invariantName="System.Data.SqlClient" type="System.Data.Entity.SqlServer.SqlProviderServices, EntityFramework.SqlServer" />
+    </providers>
+  </entityFramework>
 </configuration>
\ No newline at end of file
diff --git a/Engine/Source/Programs/CrashReporter/CrashReportWebSite/packages.config b/Engine/Source/Programs/CrashReporter/CrashReportWebSite/packages.config
--- a/Engine/Source/Programs/CrashReporter/CrashReportWebSite/packages.config
+++ b/Engine/Source/Programs/CrashReporter/CrashReportWebSite/packages.config
@@ -3,6 +3,7 @@
   <package id="Common.Logging" version="3.3.1" targetFramework="net45" />
   <package id="Common.Logging.Core" version="3.3.1" targetFramework="net45" />
   <package id="CommonServiceLocator" version="1.3" targetFramework="net45" />
+  <package id="EntityFramework" version="6.1.3" targetFramework="net45" />
   <package id="HangFire" version="1.5.3" targetFramework="net45" />
   <package id="Hangfire.Core" version="1.5.3" targetFramework="net45" />
   <package id="Hangfire.SqlServer" version="1.5.3" targetFramework="net45" />
diff --git a/Engine/Source/Programs/UnrealDocTool/APIDocTool/APIDocTool/Program.cs b/Engine/Source/Programs/UnrealDocTool/APIDocTool/APIDocTool/Program.cs
--- a/Engine/Source/Programs/UnrealDocTool/APIDocTool/APIDocTool/Program.cs
+++ b/Engine/Source/Programs/UnrealDocTool/APIDocTool/APIDocTool/Program.cs
@@ -260,7 +260,6 @@ namespace APIDocTool
 			"MSVC_PRAGMA(X)=",
 			"MS_ALIGN(X)= ",
 			"GCC_ALIGN(X)= ",
-			"SAFE_BOOL_OPERATORS(X)= ",
 			"VARARGS=",
 			"VARARG_DECL(FuncRet,StaticFuncRet,Return,FuncName,Pure,FmtType,ExtraDecl,ExtraCall)=FuncRet FuncName(ExtraDecl FmtType Fmt, ...)",
 			"VARARG_BODY(FuncRet,FuncName,FmtType,ExtraDecl)=FuncRet FuncName(ExtraDecl FmtType Fmt, ...)",
diff --git a/Engine/Source/Programs/UnrealHeaderTool/Private/BaseParser.cpp b/Engine/Source/Programs/UnrealHeaderTool/Private/BaseParser.cpp
--- a/Engine/Source/Programs/UnrealHeaderTool/Private/BaseParser.cpp
+++ b/Engine/Source/Programs/UnrealHeaderTool/Private/BaseParser.cpp
@@ -793,6 +793,24 @@ bool FBaseParser::MatchIdentifier( const TCHAR* Match )
 	return false;
 }
 
+bool FBaseParser::MatchConstInt( const TCHAR* Match )
+{
+	FToken Token;
+	if (GetToken(Token))
+	{
+		if( Token.TokenType==TOKEN_Const && Token.Type == CPT_Int && FCString::Stricmp(Token.Identifier,Match)==0 )
+		{
+			return true;
+		}
+		else
+		{
+			UngetToken(Token);
+		}
+	}
+	
+	return false;
+}
+
 
 void FBaseParser::MatchSemi()
 {
diff --git a/Engine/Source/Programs/UnrealHeaderTool/Private/BaseParser.h b/Engine/Source/Programs/UnrealHeaderTool/Private/BaseParser.h
--- a/Engine/Source/Programs/UnrealHeaderTool/Private/BaseParser.h
+++ b/Engine/Source/Programs/UnrealHeaderTool/Private/BaseParser.h
@@ -136,6 +136,7 @@ protected:
 	// Matching predefined text.
 	bool MatchIdentifier( FName Match );
 	bool MatchIdentifier( const TCHAR* Match );
+	bool MatchConstInt( const TCHAR* Match );
 	bool PeekIdentifier( FName Match );
 	bool PeekIdentifier( const TCHAR* Match );
 	bool MatchSymbol( const TCHAR* Match, ESymbolParseOption bParseTemplateClosingBracket = ESymbolParseOption::Normal );
diff --git a/Engine/Source/Programs/UnrealHeaderTool/Private/CodeGenerator.cpp b/Engine/Source/Programs/UnrealHeaderTool/Private/CodeGenerator.cpp
--- a/Engine/Source/Programs/UnrealHeaderTool/Private/CodeGenerator.cpp
+++ b/Engine/Source/Programs/UnrealHeaderTool/Private/CodeGenerator.cpp
@@ -14,6 +14,7 @@
 #include "Runtime/Core/Public/Features/IModularFeatures.h"
 #include "UHTMakefile/UHTMakefile.h"
 #include "ScopeExit.h"
+#include "FileLineException.h"
 
 /////////////////////////////////////////////////////
 // Globals
@@ -5755,6 +5756,20 @@ ECompilationResult::Type PreparseModules(FUHTMakefile& UHTMakefile, const FStrin
 					}
 				}
 			#if !PLATFORM_EXCEPTIONS_DISABLED
+				catch (const FFileLineException& Ex)
+				{
+					TGuardValue<ELogTimes::Type> DisableLogTimes(GPrintLogTimes, ELogTimes::None);
+
+					FString AbsFilename           = IFileManager::Get().ConvertToAbsolutePathForExternalAppForRead(*Ex.Filename);
+					FString Prefix                = FString::Printf(TEXT("%s(%d): "), *AbsFilename, Ex.Line);
+					FString FormattedErrorMessage = FString::Printf(TEXT("%sError: %s\r\n"), *Prefix, *Ex.Message);
+					Result = GCompilationResult;
+
+					UE_LOG(LogCompile, Log, TEXT("%s"), *FormattedErrorMessage);
+					GWarn->Log(ELogVerbosity::Error, FormattedErrorMessage);
+
+					++NumFailures;
+				}
 				catch (TCHAR* ErrorMsg)
 				{
 					TGuardValue<ELogTimes::Type> DisableLogTimes(GPrintLogTimes, ELogTimes::None);
@@ -5815,7 +5830,7 @@ ECompilationResult::Type UnrealHeaderTool_Main(const FString& ModuleInfoFilename
 {
 	check(GIsUCCMakeStandaloneHeaderGenerator);
 	ECompilationResult::Type Result = ECompilationResult::Succeeded;
-	
+
 	FString ModuleInfoPath = FPaths::GetPath(ModuleInfoFilename);
 
 	// Load the manifest file, giving a list of all modules to be processed, pre-sorted by dependency ordering
@@ -6118,7 +6133,7 @@ TSharedRef<FUnrealSourceFile> PerformInitialParseOnHeader(UPackage* InParent, co
 	// Parse the header to extract the information needed
 	FUHTStringBuilder ClassHeaderTextStrippedOfCppText;
 	TArray<FSimplifiedParsingClassInfo> ParsedClassArray;
-	FHeaderParser::SimplifiedClassParse(Buffer, /*out*/ ParsedClassArray, /*out*/ DependsOn, ClassHeaderTextStrippedOfCppText);
+	FHeaderParser::SimplifiedClassParse(FileName, Buffer, /*out*/ ParsedClassArray, /*out*/ DependsOn, ClassHeaderTextStrippedOfCppText);
 
 	FUnrealSourceFile* UnrealSourceFilePtr = new FUnrealSourceFile(InParent, FileName, MoveTemp(ClassHeaderTextStrippedOfCppText));
 	TSharedRef<FUnrealSourceFile> UnrealSourceFile = MakeShareable(UnrealSourceFilePtr);
diff --git a/Engine/Source/Programs/UnrealHeaderTool/Private/FileLineException.cpp b/Engine/Source/Programs/UnrealHeaderTool/Private/FileLineException.cpp
new file mode 100644
--- /dev/null
+++ b/Engine/Source/Programs/UnrealHeaderTool/Private/FileLineException.cpp
@@ -0,0 +1,45 @@
+// Copyright 1998-2016 Epic Games, Inc. All Rights Reserved.
+
+#include "UnrealHeaderTool.h"
+#include "FileLineException.h"
+
+void VARARGS FFileLineException::ThrowfImpl(FString&& InFilename, int32 InLine, const TCHAR* Fmt, ...)
+{
+	int32	BufferSize	= 512;
+	TCHAR	StartingBuffer[512];
+	TCHAR*	Buffer		= StartingBuffer;
+	int32	Result		= -1;
+
+	// First try to print to a stack allocated location 
+	GET_VARARGS_RESULT( Buffer, BufferSize, BufferSize-1, Fmt, Fmt, Result );
+
+	// If that fails, start allocating regular memory
+	if( Result == -1 )
+	{
+		Buffer = nullptr;
+		while(Result == -1)
+		{
+			BufferSize *= 2;
+			Buffer = (TCHAR*) FMemory::Realloc( Buffer, BufferSize * sizeof(TCHAR) );
+			GET_VARARGS_RESULT( Buffer, BufferSize, BufferSize-1, Fmt, Fmt, Result );
+		}
+	}
+
+	Buffer[Result] = 0;
+
+	FString ResultString(Buffer);
+
+	if( BufferSize != 512 )
+	{
+		FMemory::Free( Buffer );
+	}
+
+	throw FFileLineException(MoveTemp(ResultString), MoveTemp(InFilename), InLine);
+}
+
+FFileLineException::FFileLineException(FString&& InMessage, FString&& InFilename, int32 InLine)
+	: Message (MoveTemp(InMessage))
+	, Filename(MoveTemp(InFilename))
+	, Line    (InLine)
+{
+}
diff --git a/Engine/Source/Programs/UnrealHeaderTool/Private/FileLineException.h b/Engine/Source/Programs/UnrealHeaderTool/Private/FileLineException.h
new file mode 100644
--- /dev/null
+++ b/Engine/Source/Programs/UnrealHeaderTool/Private/FileLineException.h
@@ -0,0 +1,27 @@
+// Copyright 1998-2016 Epic Games, Inc. All Rights Reserved.
+
+#pragma once
+
+struct FFileLineException
+{
+	FString Message;
+	FString Filename;
+	int32   Line;
+
+	template <typename... Types>
+	FUNCTION_NO_RETURN_START
+		static void VARARGS Throwf(FString&& Filename, int32 Line, const TCHAR* Fmt, Types... Args)
+	FUNCTION_NO_RETURN_END
+	{
+		static_assert(TAnd<TIsValidVariadicFunctionArg<Types>...>::Value, "Invalid argument(s) passed to FError::Throwf");
+
+		ThrowfImpl(MoveTemp(Filename), Line, Fmt, Args...);
+	}
+
+private:
+	FUNCTION_NO_RETURN_START
+	static void VARARGS ThrowfImpl(FString&& Filename, int32 Line, const TCHAR* Fmt, ...)
+	FUNCTION_NO_RETURN_END;
+
+	FFileLineException(FString&& InMessage, FString&& InFilename, int32 InLine);
+};
diff --git a/Engine/Source/Programs/UnrealHeaderTool/Private/HeaderParser.cpp b/Engine/Source/Programs/UnrealHeaderTool/Private/HeaderParser.cpp
--- a/Engine/Source/Programs/UnrealHeaderTool/Private/HeaderParser.cpp
+++ b/Engine/Source/Programs/UnrealHeaderTool/Private/HeaderParser.cpp
@@ -13,6 +13,7 @@
 #include "Manifest.h"
 #include "UnitConversion.h"
 #include "GeneratedCodeVersion.h"
+#include "FileLineException.h"
 
 #include "Algo/FindSortedStringCaseInsensitive.h"
 
@@ -2053,12 +2054,10 @@ UScriptStruct* FHeaderParser::CompileStructDeclaration(FClasses& AllClasses)
 				}
 				PushCompilerDirective(ECompilerDirective::WithEditor);
 			}
-			else if (MatchIdentifier(TEXT("CPP")))
+			else if (MatchIdentifier(TEXT("CPP")) || MatchConstInt(TEXT("0")) || MatchConstInt(TEXT("1")) || MatchIdentifier(TEXT("WITH_HOT_RELOAD")) || MatchIdentifier(TEXT("WITH_HOT_RELOAD_CTORS")))
 			{
 				bConsumeAsCppText = !bInvertConditional;
 				PushCompilerDirective(ECompilerDirective::Insignificant);
-				//@todo: UCREMOVAL, !CPP should be interpreted as noexport and you should not need the no export.
-				// this applies to structs, enums, and everything else
 			}
 			else
 			{
@@ -2629,27 +2628,39 @@ void FHeaderParser::CompileDirective(FClasses& AllClasses)
 		// Eat the ! if present
 		bool bNotDefined = MatchSymbol(TEXT("!"));
 
-		FToken Define;
-		if (!GetIdentifier(Define))
-		{
-			FError::Throwf(TEXT("Missing define name '#if'") );
-		}
-
-		if ( Define.Matches(TEXT("WITH_EDITORONLY_DATA")) )
-		{
-			PushCompilerDirective(ECompilerDirective::WithEditorOnlyData);
-		}
-		else if ( Define.Matches(TEXT("WITH_EDITOR")) )
-		{
-			PushCompilerDirective(ECompilerDirective::WithEditor);
-		}
-		else if (Define.Matches(TEXT("CPP")) && bNotDefined)
+		int32 TempInt;
+		if (GetConstInt(TempInt) && TempInt == 0 || TempInt == 1)
 		{
 			PushCompilerDirective(ECompilerDirective::Insignificant);
 		}
 		else
 		{
-			FError::Throwf(TEXT("Unknown define '#if %s' in class or global scope"), Define.Identifier);
+			FToken Define;
+			if (!GetIdentifier(Define))
+			{
+				FError::Throwf(TEXT("Missing define name '#if'") );
+			}
+
+			if ( Define.Matches(TEXT("WITH_EDITORONLY_DATA")) )
+			{
+				PushCompilerDirective(ECompilerDirective::WithEditorOnlyData);
+			}
+			else if ( Define.Matches(TEXT("WITH_EDITOR")) )
+			{
+				PushCompilerDirective(ECompilerDirective::WithEditor);
+			}
+			else if (Define.Matches(TEXT("WITH_HOT_RELOAD")) || Define.Matches(TEXT("WITH_HOT_RELOAD_CTORS")) || Define.Matches(TEXT("1")))
+			{
+				PushCompilerDirective(ECompilerDirective::Insignificant);
+			}
+			else if ( Define.Matches(TEXT("CPP")) && bNotDefined)
+			{
+				PushCompilerDirective(ECompilerDirective::Insignificant);
+			}
+			else
+			{
+				FError::Throwf(TEXT("Unknown define '#if %s' in class or global scope"), Define.Identifier);
+			}
 		}
 	}
 	else if (Directive.Matches(TEXT("endif")))
@@ -4044,16 +4055,7 @@ EFindName FHeaderParser::GetFindFlagForPropertyName(const TCHAR* PropertyName)
 	}
 	// Add it to the list for future look ups
 	PreviousNames.Add(UpperPropertyStr,1);
-	// Check for a mismatch between the INI file and the config property name
-	FName CurrentText(PropertyName,FNAME_Find);
-	if (CurrentText != NAME_None &&
-		FCString::Strcmp(PropertyName,*CurrentText.ToString()) != 0)
-	{
-		FError::Throwf(
-			TEXT("INI file contains an incorrect case for (%s) should be (%s)"),
-			*CurrentText.ToString(),
-			PropertyName);
-	}
+	FName CurrentText(PropertyName,FNAME_Find); // keep generating this FName in case it has been affecting the case of future FNames.
 	return FNAME_Replace_Not_Safe_For_Threading;
 }
 
@@ -7366,8 +7368,98 @@ void FHeaderParser::ParseClassName(const TCHAR* Temp, FString& ClassName)
 	}
 }
 
+enum class EBlockDirectiveType
+{
+	// We're in a CPP block
+	CPPBlock,
+
+	// We're in a !CPP block
+	NotCPPBlock,
+
+	// We're in a 0 block
+	ZeroBlock,
+
+	// We're in a 1 block
+	OneBlock,
+
+	// We're in a WITH_HOT_RELOAD block
+	WithHotReload,
+
+	// We're in a WITH_EDITOR block
+	WithEditor,
+
+	// We're in a WITH_EDITORONLY_DATA block
+	WithEditorOnlyData,
+
+	// We're in a block with an unrecognized directive
+	UnrecognizedBlock
+};
+
+bool ShouldKeepBlockContents(EBlockDirectiveType DirectiveType)
+{
+	switch (DirectiveType)
+	{
+		case EBlockDirectiveType::NotCPPBlock:
+		case EBlockDirectiveType::OneBlock:
+		case EBlockDirectiveType::WithHotReload:
+		case EBlockDirectiveType::WithEditor:
+		case EBlockDirectiveType::WithEditorOnlyData:
+			return true;
+
+		case EBlockDirectiveType::CPPBlock:
+		case EBlockDirectiveType::ZeroBlock:
+		case EBlockDirectiveType::UnrecognizedBlock:
+			return false;
+	}
+
+	check(false);
+	ASSUME(false);
+}
+
+EBlockDirectiveType ParseCommandToBlockDirectiveType(const TCHAR** Str)
+{
+	if (FParse::Command(Str, TEXT("0")))
+	{
+		return EBlockDirectiveType::ZeroBlock;
+	}
+
+	if (FParse::Command(Str, TEXT("1")))
+	{
+		return EBlockDirectiveType::OneBlock;
+	}
+
+	if (FParse::Command(Str, TEXT("CPP")))
+	{
+		return EBlockDirectiveType::CPPBlock;
+	}
+
+	if (FParse::Command(Str, TEXT("!CPP")))
+	{
+		return EBlockDirectiveType::NotCPPBlock;
+	}
+
+#if WITH_HOT_RELOAD_CTORS
+	if (FParse::Command(Str, TEXT("WITH_HOT_RELOAD")))
+	{
+		return EBlockDirectiveType::WithHotReload;
+	}
+#endif
+
+	if (FParse::Command(Str, TEXT("WITH_EDITOR")))
+	{
+		return EBlockDirectiveType::WithEditor;
+	}
+
+	if (FParse::Command(Str, TEXT("WITH_EDITORONLY_DATA")))
+	{
+		return EBlockDirectiveType::WithEditorOnlyData;
+	}
+
+	return EBlockDirectiveType::UnrecognizedBlock;
+}
+
 // Performs a preliminary parse of the text in the specified buffer, pulling out useful information for the header generation process
-void FHeaderParser::SimplifiedClassParse(const TCHAR* InBuffer, TArray<FSimplifiedParsingClassInfo>& OutParsedClassArray, TArray<FHeaderProvider>& DependentOn, FStringOutputDevice& ClassHeaderTextStrippedOfCppText)
+void FHeaderParser::SimplifiedClassParse(const TCHAR* Filename, const TCHAR* InBuffer, TArray<FSimplifiedParsingClassInfo>& OutParsedClassArray, TArray<FHeaderProvider>& DependentOn, FStringOutputDevice& ClassHeaderTextStrippedOfCppText)
 {
 	FHeaderPreParser Parser;
 	FString StrLine;
@@ -7398,127 +7490,121 @@ void FHeaderParser::SimplifiedClassParse(const TCHAR* InBuffer, TArray<FSimplifi
 		bool bIf = FParse::Command(&Str,TEXT("#if"));
 		if( bIf || FParse::Command(&Str,TEXT("#ifdef")) || FParse::Command(&Str,TEXT("#ifndef")) )
 		{
-			FStringOutputDevice TextDumpDummy;
-			FStringOutputDevice* Target = NULL;
-			FStringOutputDevice* SpacerTarget = NULL;
-			TArray<bool, TInlineAllocator<8>> KeepPreprocessorDirectiveStack;
-			KeepPreprocessorDirectiveStack.Push(true);
-			bool bNotCPP = false;
-			bool bCPP = false;
-			bool bUnknownDirective = false;
-
-			if( bIf && FParse::Command(&Str,TEXT("CPP")) )
-			{
-				Target = &TextDumpDummy;
-				SpacerTarget = &ClassHeaderTextStrippedOfCppText;
-				bCPP = true;
-			}
-			else if( bIf && FParse::Command(&Str,TEXT("!CPP")) )
-			{
-				Target = &ClassHeaderTextStrippedOfCppText;
-				KeepPreprocessorDirectiveStack.Top() = false;
-				bNotCPP = true;
-			}
-#if WITH_HOT_RELOAD_CTORS
-			else if (bIf && FParse::Command(&Str, TEXT("WITH_HOT_RELOAD")))
+			EBlockDirectiveType RootDirective;
+			if (bIf)
 			{
-				Target = &ClassHeaderTextStrippedOfCppText;
-				KeepPreprocessorDirectiveStack.Top() = false;
-			}
-#endif // WITH_HOT_RELOAD_CTORS
-			else if (bIf && (FParse::Command(&Str,TEXT("WITH_EDITORONLY_DATA")) || FParse::Command(&Str,TEXT("WITH_EDITOR"))))
-			{
-				Target = &ClassHeaderTextStrippedOfCppText;
-				bUnknownDirective = true;
+				RootDirective = ParseCommandToBlockDirectiveType(&Str);
 			}
 			else
 			{
-				// Unknown directives or #ifdef or #ifndef are always treated as CPP
-				bUnknownDirective = true;
-				Target = &TextDumpDummy;
-				SpacerTarget = &ClassHeaderTextStrippedOfCppText;
+				// #ifdef or #ifndef are always treated as CPP
+				RootDirective = EBlockDirectiveType::UnrecognizedBlock;
 			}
 
-			if (KeepPreprocessorDirectiveStack.Top())
-			{
-				Target->Logf( TEXT("%s\r\n"), *StrLine );
-			}
-			else
-			{
-				Target->Logf( TEXT("\r\n") );
-			}
+			TArray<EBlockDirectiveType, TInlineAllocator<8>> DirectiveStack;
+			DirectiveStack.Push(RootDirective);
 
-			if (SpacerTarget != NULL)
-			{
-				// Make sure script line numbers don't get out of whack if there is an inline CPP block in there
-				SpacerTarget->Logf( TEXT("\r\n") );
-			}
+			bool bShouldKeepBlockContents = ShouldKeepBlockContents(RootDirective);
+			bool bIsZeroBlock = RootDirective == EBlockDirectiveType::ZeroBlock;
 
-			while ((KeepPreprocessorDirectiveStack.Num() > 0) && FParse::Line(&Buffer, StrLine, 1))
-			{
-				if (SpacerTarget != NULL)
-				{
-					// Make sure script line numbers don't get out of whack if there is an inline CPP block in there
-					SpacerTarget->Logf( TEXT("\r\n") );
-				}
+			ClassHeaderTextStrippedOfCppText.Logf(TEXT("%s\r\n"), bShouldKeepBlockContents ? *StrLine : TEXT(""));
 
+			while ((DirectiveStack.Num() > 0) && FParse::Line(&Buffer, StrLine, 1))
+			{
 				CurrentLine++;
 				Str = *StrLine;
-				bool bKeepPreprocessorDirective = KeepPreprocessorDirectiveStack.Top();
-				bool bIsPrep = false;
+
+				bool bIsDirective = false;
 				if( FParse::Command(&Str,TEXT("#endif")) )
 				{
-					KeepPreprocessorDirectiveStack.Pop();
-					bIsPrep = true;
+					DirectiveStack.Pop();
+
+					bIsDirective = true;
 				}
 				else if( FParse::Command(&Str,TEXT("#if")) || FParse::Command(&Str,TEXT("#ifdef")) || FParse::Command(&Str,TEXT("#ifndef")) )
 				{
-					bKeepPreprocessorDirective = FParse::Command(&Str, TEXT("WITH_EDITORONLY_DATA")) || FParse::Command(&Str, TEXT("WITH_EDITOR"));
-					KeepPreprocessorDirectiveStack.Push(bKeepPreprocessorDirective);
-					bIsPrep = true;
+					EBlockDirectiveType Directive = ParseCommandToBlockDirectiveType(&Str);
+					DirectiveStack.Push(Directive);
+
+					bIsDirective = true;
 				}
-				else if (FParse::Command(&Str, TEXT("#ifndef")))
+				else if (FParse::Command(&Str,TEXT("#elif")))
 				{
-					KeepPreprocessorDirectiveStack.Push(false);
-					bIsPrep = true;
+					EBlockDirectiveType Directive = ParseCommandToBlockDirectiveType(&Str);
+					DirectiveStack.Top() = Directive;
+
+					bIsDirective = true;
 				}
-				else if( FParse::Command(&Str,TEXT("#elif")) )
+				else if (FParse::Command(&Str, TEXT("#else")))
 				{
-					bIsPrep = true;
+					switch (DirectiveStack[0])
+					{
+						case EBlockDirectiveType::ZeroBlock:
+							DirectiveStack.Top() = EBlockDirectiveType::OneBlock;
+							break;
+
+						case EBlockDirectiveType::OneBlock:
+							DirectiveStack.Top() = EBlockDirectiveType::ZeroBlock;
+							break;
+
+						case EBlockDirectiveType::CPPBlock:
+							DirectiveStack.Top() = EBlockDirectiveType::NotCPPBlock;
+							break;
+
+						case EBlockDirectiveType::NotCPPBlock:
+							DirectiveStack.Top() = EBlockDirectiveType::CPPBlock;
+							break;
+
+						case EBlockDirectiveType::WithHotReload:
+							FFileLineException::Throwf(Filename, CurrentLine, TEXT("Bad preprocessor directive in metadata declaration: %s; Only 'CPP', '1' and '0' can have #else directives"), *ClassName);
+
+						case EBlockDirectiveType::UnrecognizedBlock:
+						case EBlockDirectiveType::WithEditor:
+						case EBlockDirectiveType::WithEditorOnlyData:
+							// We allow unrecognized directives, WITH_EDITOR and WITH_EDITORONLY_DATA to have #else blocks.
+							// However, we don't actually change how UHT processes these #else blocks.
+							break;
+					}
+
+					bIsDirective = true;
 				}
-				else if (KeepPreprocessorDirectiveStack.Num() == 1 && FParse::Command(&Str, TEXT("#else")))
+
+				// Check for UHT identifiers inside skipped blocks, unless it's a zero block, because the compiler is going to skip those anyway.
+				if (!bShouldKeepBlockContents && !bIsZeroBlock)
 				{
-					if (!bUnknownDirective)
+					auto FindInitialStr = [](const TCHAR*& FoundSubstr, const FString& StrToSearch, const TCHAR* ConstructName) -> bool
 					{
-						if (!bNotCPP && !bCPP)
-						{
-							FError::Throwf(TEXT("Bad preprocessor directive in metadata declaration: %s; Only 'CPP' can have ! or #else directives"),*ClassName);
-						}
-						Swap(bNotCPP,bCPP);
-						if( bCPP)
+						if (StrToSearch.StartsWith(ConstructName, ESearchCase::CaseSensitive))
 						{
-							Target = &TextDumpDummy;
-							SpacerTarget = &ClassHeaderTextStrippedOfCppText;
-							KeepPreprocessorDirectiveStack.Top() = true;
-							StrLine = TEXT("#if CPP\r\n");
-						}
-						else
-						{
-							KeepPreprocessorDirectiveStack.Top() = false;
-							Target = &ClassHeaderTextStrippedOfCppText;
-							SpacerTarget = NULL;
+							FoundSubstr = ConstructName;
+							return true;
 						}
+
+						return false;
+					};
+
+					FString TrimmedStrLine = StrLine;
+					TrimmedStrLine.Trim();
+
+					const TCHAR* FoundSubstr = nullptr;
+					if (FindInitialStr(FoundSubstr, TrimmedStrLine, TEXT("UPROPERTY"))
+						|| FindInitialStr(FoundSubstr, TrimmedStrLine, TEXT("UCLASS"))
+						|| FindInitialStr(FoundSubstr, TrimmedStrLine, TEXT("USTRUCT"))
+						|| FindInitialStr(FoundSubstr, TrimmedStrLine, TEXT("UENUM"))
+						|| FindInitialStr(FoundSubstr, TrimmedStrLine, TEXT("UINTERFACE"))
+						|| FindInitialStr(FoundSubstr, TrimmedStrLine, TEXT("UDELEGATE"))
+						|| FindInitialStr(FoundSubstr, TrimmedStrLine, TEXT("UFUNCTION")))
+					{
+						FFileLineException::Throwf(Filename, CurrentLine, TEXT("%s inside this preprocessor block will be skipped"), FoundSubstr);
 					}
-					bIsPrep = true;
 				}
 
-				if (bKeepPreprocessorDirective || !bIsPrep)
-				{
-					Target->Logf( TEXT("%s\r\n"), *StrLine );
-				}
-				else
+				ClassHeaderTextStrippedOfCppText.Logf(TEXT("%s\r\n"), bShouldKeepBlockContents ? *StrLine : TEXT(""));
+
+				if (bIsDirective)
 				{
-					Target->Logf( TEXT("\r\n") );
+					bShouldKeepBlockContents = !DirectiveStack.ContainsByPredicate([](EBlockDirectiveType Directive) { return !ShouldKeepBlockContents(Directive); });
+					bIsZeroBlock = DirectiveStack.Contains(EBlockDirectiveType::ZeroBlock);
 				}
 			}
 		}
diff --git a/Engine/Source/Programs/UnrealHeaderTool/Private/HeaderParser.h b/Engine/Source/Programs/UnrealHeaderTool/Private/HeaderParser.h
--- a/Engine/Source/Programs/UnrealHeaderTool/Private/HeaderParser.h
+++ b/Engine/Source/Programs/UnrealHeaderTool/Private/HeaderParser.h
@@ -171,7 +171,7 @@ public:
 	//   
 	//  It also splits the buffer up into:
 	//   ScriptText (text outside of #if CPP and #if DEFAULTS blocks)
-	static void SimplifiedClassParse(const TCHAR* Buffer, TArray<FSimplifiedParsingClassInfo>& OutParsedClassArray, TArray<FHeaderProvider>& DependentOn, FStringOutputDevice& ScriptText);
+	static void SimplifiedClassParse(const TCHAR* Filename, const TCHAR* Buffer, TArray<FSimplifiedParsingClassInfo>& OutParsedClassArray, TArray<FHeaderProvider>& DependentOn, FStringOutputDevice& ScriptText);
 
 	/**
 	 * Returns True if the given class name includes a valid Unreal prefix and matches up with the given original class Name.
diff --git a/Engine/Source/Runtime/AIModule/Classes/AISystem.h b/Engine/Source/Runtime/AIModule/Classes/AISystem.h
--- a/Engine/Source/Runtime/AIModule/Classes/AISystem.h
+++ b/Engine/Source/Runtime/AIModule/Classes/AISystem.h
@@ -176,8 +176,7 @@ public:
 			return Tmp;
 		}
 
-		SAFE_BOOL_OPERATORS(FBlackboardDataToComponentsIterator);
-		FORCEINLINE_EXPLICIT_OPERATOR_BOOL() const { return CurrentIteratorIndex < Iterators.Num() && (bool)GetCurrentIteratorRef(); }
+		FORCEINLINE explicit operator bool() const { return CurrentIteratorIndex < Iterators.Num() && (bool)GetCurrentIteratorRef(); }
 		FORCEINLINE bool operator !() const { return !(bool)*this; }
 
 		FORCEINLINE UBlackboardData* Key() const { return GetCurrentIteratorRef().Key().Get(); }
diff --git a/Engine/Source/Runtime/AIModule/Classes/EnvironmentQuery/EnvQueryTypes.h b/Engine/Source/Runtime/AIModule/Classes/EnvironmentQuery/EnvQueryTypes.h
--- a/Engine/Source/Runtime/AIModule/Classes/EnvironmentQuery/EnvQueryTypes.h
+++ b/Engine/Source/Runtime/AIModule/Classes/EnvironmentQuery/EnvQueryTypes.h
@@ -1114,7 +1114,7 @@ public:
 			return GetIndex();
 		}
 
-		FORCEINLINE_EXPLICIT_OPERATOR_BOOL() const
+		FORCEINLINE explicit operator bool() const
 		{
 			return CurrentItem < Instance->Items.Num() && !Instance->bFoundSingleResult && (Deadline < 0 || FPlatformTime::Seconds() < Deadline);
 		}
diff --git a/Engine/Source/Runtime/AIModule/Classes/EnvironmentQuery/Items/EnvQueryItemType_Point.h b/Engine/Source/Runtime/AIModule/Classes/EnvironmentQuery/Items/EnvQueryItemType_Point.h
--- a/Engine/Source/Runtime/AIModule/Classes/EnvironmentQuery/Items/EnvQueryItemType_Point.h
+++ b/Engine/Source/Runtime/AIModule/Classes/EnvironmentQuery/Items/EnvQueryItemType_Point.h
@@ -30,4 +30,4 @@ class AIMODULE_API UEnvQueryItemType_Point : public UEnvQueryItemType_VectorBase
 
 // a specialization to support saving locations with navigation data already gathered
 template<>
-void FEnvQueryInstance::AddItemData<UEnvQueryItemType_Point, FNavLocation>(FNavLocation ItemValue);
\ No newline at end of file
+AIMODULE_API void FEnvQueryInstance::AddItemData<UEnvQueryItemType_Point, FNavLocation>(FNavLocation ItemValue);
\ No newline at end of file
diff --git a/Engine/Source/Runtime/AssetRegistry/Private/NameTableArchive.cpp b/Engine/Source/Runtime/AssetRegistry/Private/NameTableArchive.cpp
--- a/Engine/Source/Runtime/AssetRegistry/Private/NameTableArchive.cpp
+++ b/Engine/Source/Runtime/AssetRegistry/Private/NameTableArchive.cpp
@@ -21,7 +21,7 @@ FNameTableArchiveReader::~FNameTableArchiveReader()
 bool FNameTableArchiveReader::LoadFile(const TCHAR* Filename, int32 SerializationVersion)
 {
 	FileAr = IFileManager::Get().CreateFileReader(Filename, FILEREAD_Silent);
-	if (FileAr)
+	if (FileAr && !FileAr->IsError() && FileAr->TotalSize() > 0)
 	{
 		int32 MagicNumber = 0;
 		*this << MagicNumber;
diff --git a/Engine/Source/Runtime/Core/Private/HAL/Allocators/AnsiAllocator.h b/Engine/Source/Runtime/Core/Private/HAL/Allocators/AnsiAllocator.h
--- a/Engine/Source/Runtime/Core/Private/HAL/Allocators/AnsiAllocator.h
+++ b/Engine/Source/Runtime/Core/Private/HAL/Allocators/AnsiAllocator.h
@@ -61,12 +61,16 @@ public:
 		void ResizeAllocation(int32 PreviousNumElements, int32 NumElements, SIZE_T NumBytesPerElement)
 		{
 			// Avoid calling FMemory::Realloc( nullptr, 0 ) as ANSI C mandates returning a valid pointer which is not what we want.
-			if (Data || NumElements)
+			if (NumElements)
 			{
 				//checkSlow(((uint64)NumElements*(uint64)ElementTypeInfo.GetSize() < (uint64)INT_MAX));
-				FScriptContainerElement* NewData = (FScriptContainerElement*)::realloc(Data, NumElements*NumBytesPerElement);
-				check(NewData);
-				Data = NewData;
+				void* NewRealloc = ::realloc(Data, NumElements*NumBytesPerElement);
+				Data = (FScriptContainerElement*)NewRealloc;
+			}
+			else
+			{
+				::free(Data);
+				Data = nullptr;
 			}
 		}
 		int32 CalculateSlackReserve(int32 NumElements, int32 NumBytesPerElement) const
diff --git a/Engine/Source/Runtime/Core/Private/Misc/Parse.cpp b/Engine/Source/Runtime/Core/Private/Misc/Parse.cpp
--- a/Engine/Source/Runtime/Core/Private/Misc/Parse.cpp
+++ b/Engine/Source/Runtime/Core/Private/Misc/Parse.cpp
@@ -579,7 +579,7 @@ bool FParse::Value( const TCHAR* Stream, const TCHAR* Match, struct FGuid& Guid
 // skips through the command and blanks past it.  Returns 1 of match,
 // 0 if not.
 //
-bool FParse::Command( const TCHAR** Stream, const TCHAR*  Match, bool bParseMightTriggerExecution )
+bool FParse::Command( const TCHAR** Stream, const TCHAR* Match, bool bParseMightTriggerExecution )
 {
 #if !UE_BUILD_SHIPPING
 	if(GConsoleCommandLibrary)
@@ -594,21 +594,28 @@ bool FParse::Command( const TCHAR** Stream, const TCHAR*  Match, bool bParseMigh
 	}
 #endif // !UE_BUILD_SHIPPING
 
-	while( (**Stream==' ')||(**Stream==9) )
+	while (**Stream == TEXT(' ') || **Stream == TEXT('\t'))
+	{
 		(*Stream)++;
+	}
 
-	if( FCString::Strnicmp(*Stream,Match,FCString::Strlen(Match))==0 )
+	int32 MatchLen = FCString::Strlen(Match);
+	if (FCString::Strnicmp(*Stream, Match, MatchLen) == 0)
 	{
-		*Stream += FCString::Strlen(Match);
+		*Stream += MatchLen;
 		if( !FChar::IsAlnum(**Stream))
 //		if( !FChar::IsAlnum(**Stream) && (**Stream != '_') && (**Stream != '.'))		// more correct e.g. a cvar called "log.abc" should work but breaks some code so commented out
 		{
-			while ((**Stream==' ')||(**Stream==9)) (*Stream)++;
+			while (**Stream == TEXT(' ') || **Stream == TEXT('\t'))
+			{
+				(*Stream)++;
+			}
+
 			return 1; // Success.
 		}
 		else
 		{
-			*Stream -= FCString::Strlen(Match);
+			*Stream -= MatchLen;
 			return 0; // Only found partial match.
 		}
 	}
diff --git a/Engine/Source/Runtime/Core/Private/Modules/ModuleManager.cpp b/Engine/Source/Runtime/Core/Private/Modules/ModuleManager.cpp
--- a/Engine/Source/Runtime/Core/Private/Modules/ModuleManager.cpp
+++ b/Engine/Source/Runtime/Core/Private/Modules/ModuleManager.cpp
@@ -1039,14 +1039,14 @@ void FModuleManager::FindModulePathsInDirectory(const FString& InDirectoryName,
 }
 
 
-void FModuleManager::UnloadOrAbandonModuleWithCallback(const FName InModuleName, FOutputDevice &Ar, bool bAbandonOnly)
+void FModuleManager::UnloadOrAbandonModuleWithCallback(const FName InModuleName, FOutputDevice &Ar)
 {
 	auto Module = FindModuleChecked(InModuleName);
 	
 	Module->Module->PreUnloadCallback();
 
 	const bool bIsHotReloadable = DoesLoadedModuleHaveUObjects( InModuleName );
-	if (!bAbandonOnly && bIsHotReloadable && Module->Module->SupportsDynamicReloading())
+	if (bIsHotReloadable && Module->Module->SupportsDynamicReloading())
 	{
 		if( !UnloadModule( InModuleName ))
 		{
@@ -1056,10 +1056,7 @@ void FModuleManager::UnloadOrAbandonModuleWithCallback(const FName InModuleName,
 	else
 	{
 		// Don't warn if abandoning was the intent here
-		if (!bAbandonOnly)
-		{			
-			Ar.Logf(TEXT("Module being reloaded does not support dynamic unloading -- abandoning existing loaded module so that we can load the recompiled version!"));
-		}
+		Ar.Logf(TEXT("Module being reloaded does not support dynamic unloading -- abandoning existing loaded module so that we can load the recompiled version!"));
 		AbandonModule( InModuleName );
 	}
 
@@ -1068,6 +1065,19 @@ void FModuleManager::UnloadOrAbandonModuleWithCallback(const FName InModuleName,
 }
 
 
+void FModuleManager::AbandonModuleWithCallback(const FName InModuleName)
+{
+	auto Module = FindModuleChecked(InModuleName);
+	
+	Module->Module->PreUnloadCallback();
+
+	AbandonModule( InModuleName );
+
+	// Ensure module is unloaded
+	check(!IsModuleLoaded(InModuleName));
+}
+
+
 bool FModuleManager::LoadModuleWithCallback( const FName InModuleName, FOutputDevice &Ar )
 {
 	TSharedPtr<IModuleInterface> LoadedModule = LoadModule( InModuleName, true );
diff --git a/Engine/Source/Runtime/Core/Private/Windows/WindowsPlatformCrashContext.cpp b/Engine/Source/Runtime/Core/Private/Windows/WindowsPlatformCrashContext.cpp
--- a/Engine/Source/Runtime/Core/Private/Windows/WindowsPlatformCrashContext.cpp
+++ b/Engine/Source/Runtime/Core/Private/Windows/WindowsPlatformCrashContext.cpp
@@ -18,6 +18,10 @@
 	#include <dbghelp.h>
 	#include <Shlwapi.h>
 
+#ifndef UE_LOG_CRASH_CALLSTACK
+	#define UE_LOG_CRASH_CALLSTACK 1
+#endif
+
 #pragma comment( lib, "version.lib" )
 #pragma comment( lib, "Shlwapi.lib" )
 
@@ -565,6 +569,7 @@ class FCrashReportingThread
 		{
 			if (WaitForSingleObject(CrashEvent, 500) == WAIT_OBJECT_0)
 			{
+				ResetEvent(CrashHandledEvent);
 				HandleCrashInternal();
 				ResetEvent(CrashEvent);
 				// Let the thread that crashed know we're done.				
@@ -596,7 +601,7 @@ public:
 			SetThreadPriority(Thread, THREAD_PRIORITY_BELOW_NORMAL);
 			// Synchronization objects
 			CrashEvent = CreateEvent(nullptr, true, 0, nullptr);
-			CrashHandledEvent = CreateEvent(nullptr, false, 0, nullptr);
+			CrashHandledEvent = CreateEvent(nullptr, true, 0, nullptr);
 		}
 	}
 
@@ -669,7 +674,12 @@ private:
 		// Then try run time crash processing and broadcast information about a crash.
 		FCoreDelegates::OnHandleSystemError.Broadcast();
 
-		const bool bGenerateRuntimeCallstack = FParse::Param(FCommandLine::Get(), TEXT("ForceLogCallstacks")) || FEngineBuildSettings::IsInternalBuild() || FEngineBuildSettings::IsPerforceBuild() || FEngineBuildSettings::IsSourceDistribution();
+		const bool bGenerateRuntimeCallstack =
+#if UE_LOG_CRASH_CALLSTACK
+			true;
+#else
+			FParse::Param(FCommandLine::Get(), TEXT("ForceLogCallstacks")) || FEngineBuildSettings::IsInternalBuild() || FEngineBuildSettings::IsPerforceBuild() || FEngineBuildSettings::IsSourceDistribution();
+#endif // UE_LOG_CRASH_CALLSTACK
 		if (bGenerateRuntimeCallstack)
 		{
 			const SIZE_T StackTraceSize = 65535;
diff --git a/Engine/Source/Runtime/Core/Public/Async/AsyncWork.h b/Engine/Source/Runtime/Core/Public/Async/AsyncWork.h
--- a/Engine/Source/Runtime/Core/Public/Async/AsyncWork.h
+++ b/Engine/Source/Runtime/Core/Public/Async/AsyncWork.h
@@ -336,71 +336,19 @@ class FAsyncTask
 	}
 
 public:
-	/** Default constructor. */
-	FAsyncTask( )
-	{
-		Init();
-	}
-	/** Passthrough constructor. Generally speaking references will not pass through; use pointers */
-	template<typename T>
-	FAsyncTask( T Arg )
-		: Task(Arg)
-	{
-		Init();
-	}
-	/** Passthrough constructor. Generally speaking references will not pass through; use pointers */
-	template<typename T1,typename T2>
-	FAsyncTask( T1 Arg1, T2 Arg2 )
-		: Task(Arg1,Arg2)
-	{
-		Init();
-	}
-	/** Passthrough constructor. Generally speaking references will not pass through; use pointers */
-	template<typename T1,typename T2, typename T3>
-	FAsyncTask( T1 Arg1, T2 Arg2, T3 Arg3 )
-		: Task(Arg1,Arg2,Arg3)
-	{
-		Init();
-	}
-	/** Passthrough constructor. Generally speaking references will not pass through; use pointers */
-	template<typename T1,typename T2, typename T3, typename T4>
-	FAsyncTask( T1 Arg1, T2 Arg2, T3 Arg3, T4 Arg4 )
-		: Task(Arg1,Arg2,Arg3,Arg4)
-	{
-		Init();
-	}
-	/** Passthrough constructor. Generally speaking references will not pass through; use pointers */
-	template<typename T1,typename T2, typename T3, typename T4, typename T5>
-	FAsyncTask( T1 Arg1, T2 Arg2, T3 Arg3, T4 Arg4, T5 Arg5 )
-		: Task(Arg1,Arg2,Arg3,Arg4,Arg5)
-	{
-		Init();
-	}
-	/** Passthrough constructor. Generally speaking references will not pass through; use pointers */
-	template<typename T1,typename T2, typename T3, typename T4, typename T5, typename T6>
-	FAsyncTask( T1 Arg1, T2 Arg2, T3 Arg3, T4 Arg4, T5 Arg5, T6 Arg6 )
-		: Task(Arg1,Arg2,Arg3,Arg4,Arg5,Arg6)
-	{
-		Init();
-	}
-	/** Passthrough constructor. Generally speaking references will not pass through; use pointers */
-	template<typename T1,typename T2, typename T3, typename T4, typename T5, typename T6, typename T7>
-	FAsyncTask( T1 Arg1, T2 Arg2, T3 Arg3, T4 Arg4, T5 Arg5, T6 Arg6, T7 Arg7 )
-		: Task(Arg1,Arg2,Arg3,Arg4,Arg5,Arg6,Arg7)
-	{
-		Init();
-	}
-	/** Passthrough constructor. Generally speaking references will not pass through; use pointers */
-	template<typename T1,typename T2, typename T3, typename T4, typename T5, typename T6, typename T7, typename T8>
-	FAsyncTask( T1 Arg1, T2 Arg2, T3 Arg3, T4 Arg4, T5 Arg5, T6 Arg6, T7 Arg7, T8 Arg8 )
-		: Task(Arg1,Arg2,Arg3,Arg4,Arg5,Arg6,Arg7,Arg8)
+	FAsyncTask()
+		: Task()
 	{
+		// This constructor shouldn't be necessary as the forwarding constructor should handle it, but
+		// we are getting VC internal compiler errors on CIS when creating arrays of FAsyncTask.
+
 		Init();
 	}
-	/** Passthrough constructor. Generally speaking references will not pass through; use pointers */
-	template<typename T1,typename T2, typename T3, typename T4, typename T5, typename T6, typename T7, typename T8, typename T9>
-	FAsyncTask( T1 Arg1, T2 Arg2, T3 Arg3, T4 Arg4, T5 Arg5, T6 Arg6, T7 Arg7, T8 Arg8, T9 Arg9 )
-		: Task(Arg1,Arg2,Arg3,Arg4,Arg5,Arg6,Arg7,Arg8,Arg9)
+
+	/** Forwarding constructor. */
+	template <typename Arg0Type, typename... ArgTypes>
+	FAsyncTask(Arg0Type&& Arg0, ArgTypes&&... Args)
+		: Task(Forward<Arg0Type>(Arg0), Forward<ArgTypes>(Args)...)
 	{
 		Init();
 	}
diff --git a/Engine/Source/Runtime/Core/Public/Containers/Algo/Find.h b/Engine/Source/Runtime/Core/Public/Containers/Algo/Find.h
new file mode 100644
--- /dev/null
+++ b/Engine/Source/Runtime/Core/Public/Containers/Algo/Find.h
@@ -0,0 +1,77 @@
+// Copyright 1998-2016 Epic Games, Inc. All Rights Reserved.
+
+#pragma once
+
+namespace AlgoImpl
+{
+	template <typename RangeType, typename PredicateType>
+	typename const RangeType::ElementType* FindByPredicate(const RangeType& Range, PredicateType Predicate)
+	{
+		for (const auto& Elem : Range)
+		{
+			if (Predicate(Elem))
+			{
+				return &Elem;
+			}
+		}
+
+		return nullptr;
+	}
+}
+
+namespace Algo
+{
+	/**
+	 * Returns a pointer to the first element matching a value in a range.
+	 *
+	 * @param  Range  The range to search.
+	 * @param  Value  The value to search for.
+	 * @return A pointer to the first element found, or nullptr if none was found.
+	 */
+	template <typename RangeType, typename ValueType>
+	FORCEINLINE typename RangeType::ElementType* Find(RangeType& Range, const ValueType& Value)
+	{
+		return const_cast<typename RangeType::ElementType*>(Find(const_cast<const RangeType&>(Range), Value));
+	}
+
+	/**
+	 * Returns a pointer to the first element matching a value in a range.
+	 *
+	 * @param  Range  The range to search.
+	 * @param  Value  The value to search for.
+	 * @return A pointer to the first element found, or nullptr if none was found.
+	 */
+	template <typename RangeType, typename ValueType>
+	FORCEINLINE const typename RangeType::ElementType* Find(const RangeType& Range, const ValueType& Value)
+	{
+		return AlgoImpl::FindByPredicate(Range, [&](const typename RangeType::ElementType& Elem){
+			return Elem == Value;
+		});
+	}
+
+	/**
+	 * Returns a pointer to the first element matching a predicate in a range.
+	 *
+	 * @param  Range  The range to search.
+	 * @param  Pred   The predicate to search for.
+	 * @return A pointer to the first element found, or nullptr if none was found.
+	 */
+	template <typename RangeType, typename PredicateType>
+	FORCEINLINE typename RangeType::ElementType* FindByPredicate(RangeType& Range, PredicateType Pred)
+	{
+		return const_cast<typename RangeType::ElementType*>(FindByPredicate(const_cast<const RangeType&>(Range), Pred));
+	}
+
+	/**
+	 * Returns a pointer to the first element matching a predicate in a range.
+	 *
+	 * @param  Range  The range to search.
+	 * @param  Pred   The predicate to search for.
+	 * @return A pointer to the first element found, or nullptr if none was found.
+	 */
+	template <typename RangeType, typename PredicateType>
+	FORCEINLINE const typename RangeType::ElementType* FindByPredicate(const RangeType& Range, PredicateType Pred)
+	{
+		return AlgoImpl::FindByPredicate(Range, Pred);
+	}
+}
diff --git a/Engine/Source/Runtime/Core/Public/Containers/Array.h b/Engine/Source/Runtime/Core/Public/Containers/Array.h
--- a/Engine/Source/Runtime/Core/Public/Containers/Array.h
+++ b/Engine/Source/Runtime/Core/Public/Containers/Array.h
@@ -2,6 +2,8 @@
 
 #pragma once
 
+#include <initializer_list>
+
 #include "Containers/ContainerAllocationPolicies.h"
 #include "HAL/Platform.h"
 #include "Serialization/ArchiveBase.h"
@@ -102,7 +104,7 @@ public:
 	//@}
 
 	/** conversion to "bool" returning true if the iterator has not reached the last element. */
-	FORCEINLINE_EXPLICIT_OPERATOR_BOOL() const
+	FORCEINLINE explicit operator bool() const
 	{
 		return Container.IsValidIndex(Index);
 	}
@@ -527,6 +529,17 @@ public:
 		, ArrayMax(0)
 	{}
 
+	/**
+	 * Initializer list constructor
+	 */
+	TArray(std::initializer_list<InElementType> InitList)
+	{
+		// This is not strictly legal, as std::initializer_list's iterators are not guaranteed to be pointers, but
+		// this appears to be the case on all of our implementations.  Also, if it's not true on a new implementation,
+		// it will fail to compile rather than behave badly.
+		CopyToEmpty(InitList.begin(), (int32)InitList.size(), 0, 0);
+	}
+
 	/**
 	 * Copy constructor with changed allocator. Use the common routine to perform the copy.
 	 *
@@ -535,7 +548,7 @@ public:
 	template <typename OtherElementType, typename OtherAllocator>
 	FORCEINLINE explicit TArray(const TArray<OtherElementType, OtherAllocator>& Other)
 	{
-		CopyToEmpty(Other, 0, 0);
+		CopyToEmpty(Other.GetData(), Other.Num(), 0, 0);
 	}
 
 	/**
@@ -545,7 +558,7 @@ public:
 	 */
 	FORCEINLINE TArray(const TArray& Other)
 	{
-		CopyToEmpty(Other, 0, 0);
+		CopyToEmpty(Other.GetData(), Other.Num(), 0, 0);
 	}
 
 	/**
@@ -557,7 +570,23 @@ public:
 	 */
 	FORCEINLINE TArray(const TArray& Other, int32 ExtraSlack)
 	{
-		CopyToEmpty(Other, 0, ExtraSlack);
+		CopyToEmpty(Other.GetData(), Other.Num(), 0, ExtraSlack);
+	}
+
+	/**
+	 * Initializer list assignment operator. First deletes all currently contained elements
+	 * and then copies from initializer list.
+	 *
+	 * @param InitList The initializer_list to copy from.
+	 */
+	AGRESSIVE_ARRAY_FORCEINLINE TArray& operator=(std::initializer_list<InElementType> InitList)
+	{
+		DestructItems(GetData(), ArrayNum);
+		// This is not strictly legal, as std::initializer_list's iterators are not guaranteed to be pointers, but
+		// this appears to be the case on all of our implementations.  Also, if it's not true on a new implementation,
+		// it will fail to compile rather than behave badly.
+		CopyToEmpty(InitList.begin(), (int32)InitList.size(), ArrayMax, 0);
+		return *this;
 	}
 
 	/**
@@ -572,7 +601,7 @@ public:
 	AGRESSIVE_ARRAY_FORCEINLINE TArray& operator=(const TArray<ElementType, OtherAllocator>& Other)
 	{
 		DestructItems(GetData(), ArrayNum);
-		CopyToEmpty(Other, ArrayMax, 0);
+		CopyToEmpty(Other.GetData(), Other.Num(), ArrayMax, 0);
 		return *this;
 	}
 
@@ -587,7 +616,7 @@ public:
 		if (this != &Other)
 		{
 			DestructItems(GetData(), ArrayNum);
-			CopyToEmpty(Other, ArrayMax, 0);
+			CopyToEmpty(Other.GetData(), Other.Num(), ArrayMax, 0);
 		}
 		return *this;
 	}
@@ -626,7 +655,7 @@ private:
 	template <typename FromArrayType, typename ToArrayType>
 	static FORCEINLINE typename TEnableIf<!UE4Array_Private::TCanMoveTArrayPointersBetweenArrayTypes<FromArrayType, ToArrayType>::Value>::Type MoveOrCopy(ToArrayType& ToArray, FromArrayType& FromArray, int32 PrevMax)
 	{
-		ToArray.CopyToEmpty(FromArray, PrevMax, 0);
+		ToArray.CopyToEmpty(FromArray.GetData(), FromArray.Num(), PrevMax, 0);
 	}
 
 	/**
@@ -660,7 +689,7 @@ private:
 	template <typename FromArrayType, typename ToArrayType>
 	static FORCEINLINE typename TEnableIf<!UE4Array_Private::TCanMoveTArrayPointersBetweenArrayTypes<FromArrayType, ToArrayType>::Value>::Type MoveOrCopyWithSlack(ToArrayType& ToArray, FromArrayType& FromArray, int32 PrevMax, int32 ExtraSlack)
 	{
-		ToArray.CopyToEmpty(FromArray, PrevMax, ExtraSlack);
+		ToArray.CopyToEmpty(FromArray.GetData(), FromArray.Num(), PrevMax, ExtraSlack);
 	}
 
 public:
@@ -1434,7 +1463,28 @@ public:
 	void InsertZeroed(int32 Index, int32 Count = 1)
 	{
 		InsertUninitialized(Index, Count);
-		FMemory::Memzero((uint8*)AllocatorInstance.GetAllocation() + Index*sizeof(ElementType), Count*sizeof(ElementType));
+		FMemory::Memzero((uint8*)AllocatorInstance.GetAllocation() + Index * sizeof(ElementType), Count * sizeof(ElementType));
+	}
+
+	/**
+	 * Inserts given elements into the array at given location.
+	 *
+	 * @param Items Array of elements to insert.
+	 * @param InIndex Tells where to insert the new elements.
+	 * @returns Location at which the item was inserted.
+	 */
+	int32 Insert(std::initializer_list<ElementType> InitList, const int32 InIndex)
+	{
+		InsertUninitialized(InIndex, (int32)InitList.size());
+
+		ElementType* Data = (ElementType*)AllocatorInstance.GetAllocation();
+
+		int32 Index = InIndex;
+		for (const ElementType& Element : InitList)
+		{
+			new (Data + Index++) ElementType(Element);
+		}
+		return InIndex;
 	}
 
 	/**
@@ -1447,12 +1497,15 @@ public:
 	int32 Insert(const TArray<ElementType>& Items, const int32 InIndex)
 	{
 		check(this != &Items);
+
 		InsertUninitialized(InIndex, Items.Num());
+
+		ElementType* Data = (ElementType*)AllocatorInstance.GetAllocation();
+
 		int32 Index = InIndex;
 		for (auto It = Items.CreateConstIterator(); It; ++It)
 		{
-			RangeCheck(Index);
-			new(GetData() + Index++) ElementType(MoveTemp(*It));
+			new (Data + Index++) ElementType(MoveTemp(*It));
 		}
 		return InIndex;
 	}
@@ -1779,6 +1832,20 @@ public:
 		ConstructItems<ElementType>(GetData() + Pos, Ptr, Count);
 	}
 
+	/**
+	 * Adds an initializer list of elements to the end of the TArray.
+	 *
+	 * @param InitList The initializer list of elements to add.
+	 * @see Add, Insert
+	 */
+	FORCEINLINE void Append(std::initializer_list<ElementType> InitList)
+	{
+		int32 Count = (int32)InitList.size();
+
+		int32 Pos = AddUninitialized(Count);
+		ConstructItems<ElementType>(GetData() + Pos, InitList.begin(), Count);
+	}
+
 	/**
 	 * Appends the specified array to this array.
 	 * Cannot append to self.
@@ -1805,6 +1872,17 @@ public:
 		return *this;
 	}
 
+	/**
+	 * Appends the specified initializer list to this array.
+	 *
+	 * @param InitList The initializer list to append.
+	 */
+	AGRESSIVE_ARRAY_FORCEINLINE TArray& operator+=(std::initializer_list<ElementType> InitList)
+	{
+		Append(InitList);
+		return *this;
+	}
+
 	/**
 	 * Constructs a new item at the end of the array, possibly reallocating the whole array to fit.
 	 *
@@ -2353,15 +2431,15 @@ private:
 	 *                   the end of the buffer. Counted in elements. Zero by
 	 *                   default.
 	 */
-	template <typename OtherElementType, typename OtherAllocator>
-	AGRESSIVE_ARRAY_FORCEINLINE void CopyToEmpty(const TArray<OtherElementType, OtherAllocator>& Source, int32 PrevMax, int32 ExtraSlack)
+	template <typename OtherElementType>
+	AGRESSIVE_ARRAY_FORCEINLINE void CopyToEmpty(const OtherElementType* OtherData, int32 OtherNum, int32 PrevMax, int32 ExtraSlack)
 	{
 		checkSlow(ExtraSlack >= 0);
-		ArrayNum = Source.Num();
-		if (ArrayNum || ExtraSlack || PrevMax)
+		ArrayNum = OtherNum;
+		if (OtherNum || ExtraSlack || PrevMax)
 		{
-			ResizeForCopy(ArrayNum + ExtraSlack, PrevMax);
-			ConstructItems<ElementType>(GetData(), Source.GetData(), ArrayNum);
+			ResizeForCopy(OtherNum + ExtraSlack, PrevMax);
+			ConstructItems<ElementType>(GetData(), OtherData, OtherNum);
 		}
 		else
 		{
diff --git a/Engine/Source/Runtime/Core/Public/Containers/BitArray.h b/Engine/Source/Runtime/Core/Public/Containers/BitArray.h
--- a/Engine/Source/Runtime/Core/Public/Containers/BitArray.h
+++ b/Engine/Source/Runtime/Core/Public/Containers/BitArray.h
@@ -580,9 +580,8 @@ public:
 			}
 			return *this;
 		}
-		SAFE_BOOL_OPERATORS(FIterator)
 		/** conversion to "bool" returning true if the iterator is valid. */
-		FORCEINLINE_EXPLICIT_OPERATOR_BOOL() const
+		FORCEINLINE explicit operator bool() const
 		{ 
 			return Index < Array.Num(); 
 		}
@@ -622,10 +621,8 @@ public:
 			return *this;
 		}
 
-		SAFE_BOOL_OPERATORS(FConstIterator)
-
 		/** conversion to "bool" returning true if the iterator is valid. */
-		FORCEINLINE_EXPLICIT_OPERATOR_BOOL() const
+		FORCEINLINE explicit operator bool() const
 		{ 
 			return Index < Array.Num(); 
 		}
@@ -665,10 +662,8 @@ public:
 			return *this;
 		}
 
-		SAFE_BOOL_OPERATORS(FConstReverseIterator)
-
 		/** conversion to "bool" returning true if the iterator is valid. */
-		FORCEINLINE_EXPLICIT_OPERATOR_BOOL() const
+		FORCEINLINE explicit operator bool() const
 		{ 
 			return Index >= 0; 
 		}
@@ -770,7 +765,7 @@ public:
 	}
 
 	/** conversion to "bool" returning true if the iterator is valid. */
-	FORCEINLINE_EXPLICIT_OPERATOR_BOOL() const
+	FORCEINLINE explicit operator bool() const
 	{ 
 		return CurrentBitIndex < Array.Num(); 
 	}
@@ -877,10 +872,8 @@ public:
 
 	}
 
-	SAFE_BOOL_OPERATORS(TConstDualSetBitIterator<Allocator,OtherAllocator>)
-
 	/** conversion to "bool" returning true if the iterator is valid. */
-	FORCEINLINE_EXPLICIT_OPERATOR_BOOL() const
+	FORCEINLINE explicit operator bool() const
 	{ 
 		return CurrentBitIndex < ArrayA.Num(); 
 	}
diff --git a/Engine/Source/Runtime/Core/Public/Containers/List.h b/Engine/Source/Runtime/Core/Public/Containers/List.h
--- a/Engine/Source/Runtime/Core/Public/Containers/List.h
+++ b/Engine/Source/Runtime/Core/Public/Containers/List.h
@@ -34,7 +34,7 @@ public:
 	}
 
 	/** conversion to "bool" returning true if the iterator is valid. */
-	FORCEINLINE_EXPLICIT_OPERATOR_BOOL() const
+	FORCEINLINE explicit operator bool() const
 	{ 
 		return CurrentLink != nullptr;
 	}
@@ -365,7 +365,7 @@ public:
 	{ }
 
 	/** conversion to "bool" returning true if the iterator is valid. */
-	FORCEINLINE_EXPLICIT_OPERATOR_BOOL() const
+	FORCEINLINE explicit operator bool() const
 	{ 
 		return CurrentNode != nullptr; 
 	}
diff --git a/Engine/Source/Runtime/Core/Public/Containers/Map.h b/Engine/Source/Runtime/Core/Public/Containers/Map.h
--- a/Engine/Source/Runtime/Core/Public/Containers/Map.h
+++ b/Engine/Source/Runtime/Core/Public/Containers/Map.h
@@ -616,7 +616,7 @@ protected:
 		}
 
 		/** conversion to "bool" returning true if the iterator is valid. */
-		FORCEINLINE_EXPLICIT_OPERATOR_BOOL() const
+		FORCEINLINE explicit operator bool() const
 		{ 
 			return !!PairIt; 
 		}
@@ -661,10 +661,8 @@ protected:
 			return *this;
 		}
 
-		SAFE_BOOL_OPERATORS(TBaseKeyIterator<bConst>)
-
 		/** conversion to "bool" returning true if the iterator is valid. */
-		FORCEINLINE_EXPLICIT_OPERATOR_BOOL() const
+		FORCEINLINE explicit operator bool() const
 		{ 
 			return !!SetIt; 
 		}
diff --git a/Engine/Source/Runtime/Core/Public/Containers/Set.h b/Engine/Source/Runtime/Core/Public/Containers/Set.h
--- a/Engine/Source/Runtime/Core/Public/Containers/Set.h
+++ b/Engine/Source/Runtime/Core/Public/Containers/Set.h
@@ -2,6 +2,8 @@
 
 #pragma once
 
+#include <initializer_list>
+
 #include "Containers/SparseArray.h"
 #include "Misc/StructBuilder.h"
 #include "Templates/Sorting.h"
@@ -267,6 +269,13 @@ private:
 	}
 
 public:
+	/** Initializer list constructor. */
+	TSet(std::initializer_list<ElementType> InitList)
+		: HashSize(0)
+	{
+		Append(InitList);
+	}
+
 	/** Move constructor. */
 	TSet(TSet&& Other)
 		: HashSize(0)
@@ -319,6 +328,14 @@ public:
 		return *this;
 	}
 
+	/** Initializer list assignment operator */
+	TSet& operator=(std::initializer_list<ElementType> InitList)
+	{
+		Reset();
+		Append(InitList);
+		return *this;
+	}
+
 	/**
 	 * Removes all elements from the set, potentially leaving space allocated for an expected number of elements about to be added.
 	 * @param ExpectedNumElements - The number of elements about to be added to the set.
@@ -512,7 +529,7 @@ public:
 	void Append(const TArray<ElementType, ArrayAllocator>& InElements)
 	{
 		Reserve(Elements.Num() + InElements.Num());
-		for (auto& Element : InElements)
+		for (const ElementType& Element : InElements)
 		{
 			Add(Element);
 		}
@@ -522,7 +539,7 @@ public:
 	void Append(TArray<ElementType, ArrayAllocator>&& InElements)
 	{
 		Reserve(Elements.Num() + InElements.Num());
-		for (auto& Element : InElements)
+		for (ElementType& Element : InElements)
 		{
 			Add(MoveTemp(Element));
 		}
@@ -537,7 +554,7 @@ public:
 	void Append(const TSet<ElementType, KeyFuncs, OtherAllocator>& OtherSet)
 	{
 		Reserve(Elements.Num() + OtherSet.Num());
-		for (auto& Element : OtherSet)
+		for (const ElementType& Element : OtherSet)
 		{
 			Add(Element);
 		}
@@ -547,13 +564,22 @@ public:
 	void Append(TSet<ElementType, KeyFuncs, OtherAllocator>&& OtherSet)
 	{
 		Reserve(Elements.Num() + OtherSet.Num());
-		for (auto& Element : OtherSet)
+		for (ElementType& Element : OtherSet)
 		{
 			Add(MoveTemp(Element));
 		}
 		OtherSet.Reset();
 	}
 
+	void Append(std::initializer_list<ElementType> InitList)
+	{
+		Reserve(Elements.Num() + (int32)InitList.size());
+		for (const ElementType& Element : InitList)
+		{
+			Add(Element);
+		}
+	}
+
 	/**
 	 * Removes an element from the set.
 	 * @param Element - A pointer to the element in the set, as returned by Add or Find.
@@ -1043,7 +1069,7 @@ private:
 		}
 
 		/** conversion to "bool" returning true if the iterator is valid. */
-		FORCEINLINE_EXPLICIT_OPERATOR_BOOL() const
+		FORCEINLINE explicit operator bool() const
 		{ 
 			return !!ElementIt; 
 		}
@@ -1117,10 +1143,8 @@ private:
 			return *this;
 		}
 
-		SAFE_BOOL_OPERATORS(TBaseKeyIterator<bConst>)
-
 		/** conversion to "bool" returning true if the iterator is valid. */
-		FORCEINLINE_EXPLICIT_OPERATOR_BOOL() const
+		FORCEINLINE explicit operator bool() const
 		{ 
 			return Id.IsValidId(); 
 		}
diff --git a/Engine/Source/Runtime/Core/Public/Containers/SparseArray.h b/Engine/Source/Runtime/Core/Public/Containers/SparseArray.h
--- a/Engine/Source/Runtime/Core/Public/Containers/SparseArray.h
+++ b/Engine/Source/Runtime/Core/Public/Containers/SparseArray.h
@@ -670,7 +670,7 @@ private:
 		FORCEINLINE friend bool operator!=(const TBaseIterator& Lhs, const TBaseIterator& Rhs) { return Lhs.BitArrayIt != Rhs.BitArrayIt || &Lhs.Array != &Rhs.Array; }
 
 		/** conversion to "bool" returning true if the iterator is valid. */
-		FORCEINLINE_EXPLICIT_OPERATOR_BOOL() const
+		FORCEINLINE explicit operator bool() const
 		{ 
 			return !!BitArrayIt; 
 		}
@@ -769,7 +769,7 @@ public:
 		FORCEINLINE int32 GetIndex() const { return BitArrayIt.GetIndex(); }
 		
 		/** conversion to "bool" returning true if the iterator is valid. */
-		FORCEINLINE_EXPLICIT_OPERATOR_BOOL() const
+		FORCEINLINE explicit operator bool() const
 		{ 
 			return !!BitArrayIt; 
 		}
diff --git a/Engine/Source/Runtime/Core/Public/HAL/Platform.h b/Engine/Source/Runtime/Core/Public/HAL/Platform.h
--- a/Engine/Source/Runtime/Core/Public/HAL/Platform.h
+++ b/Engine/Source/Runtime/Core/Public/HAL/Platform.h
@@ -175,9 +175,6 @@
 #ifndef PLATFORM_COMPILER_HAS_DEFAULTED_FUNCTIONS
 	#define PLATFORM_COMPILER_HAS_DEFAULTED_FUNCTIONS	1
 #endif
-#ifndef PLATFORM_COMPILER_HAS_EXPLICIT_OPERATORS
-	#define PLATFORM_COMPILER_HAS_EXPLICIT_OPERATORS	1
-#endif
 #ifndef PLATFORM_COMPILER_COMMON_LANGUAGE_RUNTIME_COMPILATION
 	#define PLATFORM_COMPILER_COMMON_LANGUAGE_RUNTIME_COMPILATION 0
 #endif
@@ -555,38 +552,6 @@ struct THasOperatorNotEquals
 	enum { Value = FHasOperatorImpl::NotEquals<T>::Value };
 };
 
-#if PLATFORM_COMPILER_HAS_EXPLICIT_OPERATORS
-	#define FORCEINLINE_EXPLICIT_OPERATOR_BOOL FORCEINLINE explicit operator bool
-	#define SAFE_BOOL_OPERATORS(...)		// not needed when compiler supports explicit operator bool()
-#elif PLATFORM_COMPILER_COMMON_LANGUAGE_RUNTIME_COMPILATION
-	// unsafe version
-	#define FORCEINLINE_EXPLICIT_OPERATOR_BOOL FORCEINLINE operator bool
-	#define SAFE_BOOL_OPERATORS(...)
-#else
-	#define FORCEINLINE_EXPLICIT_OPERATOR_BOOL \
-					private: \
-						template <typename TheClassType> \
-						static TheClassType UE_ClassTypeHelper(const TheClassType&); \
-						struct UE_PrivateBoolType \
-						{ \
-							int x; \
-						}; \
-					public: \
-						FORCEINLINE operator int UE_PrivateBoolType::*() const \
-						{ \
-							static_assert(THasOperatorEquals   <decltype(UE_ClassTypeHelper(*this))>::Value, "Class needs operator==()."); \
-							static_assert(THasOperatorNotEquals<decltype(UE_ClassTypeHelper(*this))>::Value, "Class needs operator!=()."); \
-							return UE_OperatorBool() ? &UE_PrivateBoolType::x : 0; \
-						} \
-						FORCEINLINE bool UE_OperatorBool
-
-	// these are left unimplemented to cause link errors if used
-	#define SAFE_BOOL_OPERATORS(...)	\
-			void operator ==( const __VA_ARGS__ & RHS ) const; \
-			void operator !=( const __VA_ARGS__ & RHS ) const;
-#endif
-
-
 // Console ANSICHAR/TCHAR command line handling
 #if PLATFORM_COMPILER_HAS_TCHAR_WMAIN
 #define INT32_MAIN_INT32_ARGC_TCHAR_ARGV() int32 wmain(int32 ArgC, TCHAR* ArgV[])
diff --git a/Engine/Source/Runtime/Core/Public/HTML5/HTML5Platform.h b/Engine/Source/Runtime/Core/Public/HTML5/HTML5Platform.h
--- a/Engine/Source/Runtime/Core/Public/HTML5/HTML5Platform.h
+++ b/Engine/Source/Runtime/Core/Public/HTML5/HTML5Platform.h
@@ -34,7 +34,6 @@ typedef FHTML5Types FPlatformTypes;
 #if _MSC_VER < 1900
 	#define PLATFORM_COMPILER_HAS_DEFAULTED_FUNCTIONS	0
 #endif
-#define PLATFORM_COMPILER_HAS_EXPLICIT_OPERATORS	1
 #define PLATFORM_USES_MICROSOFT_LIBC_FUNCTIONS		1
 #define PLATFORM_MAX_FILEPATH_LENGTH				MAX_PATH
 #define PLATFORM_HAS_BSD_SOCKETS					0
diff --git a/Engine/Source/Runtime/Core/Public/Misc/Optional.h b/Engine/Source/Runtime/Core/Public/Misc/Optional.h
--- a/Engine/Source/Runtime/Core/Public/Misc/Optional.h
+++ b/Engine/Source/Runtime/Core/Public/Misc/Optional.h
@@ -140,7 +140,7 @@ public:
 
 	/** @return true when the value is meaningful; false if calling GetValue() is undefined. */
 	bool IsSet() const { return bIsSet; }
-	FORCEINLINE_EXPLICIT_OPERATOR_BOOL() const { return bIsSet; }
+	FORCEINLINE explicit operator bool() const { return bIsSet; }
 
 	/** @return The optional value; undefined when IsSet() returns false. */
 	const OptionalType& GetValue() const { checkf(IsSet(), TEXT("It is an error to call GetValue() on an unset TOptional. Please either check IsSet() or use Get(DefaultValue) instead.")); return *(OptionalType*)&Value; }
diff --git a/Engine/Source/Runtime/Core/Public/Modules/ModuleManager.h b/Engine/Source/Runtime/Core/Public/Modules/ModuleManager.h
--- a/Engine/Source/Runtime/Core/Public/Modules/ModuleManager.h
+++ b/Engine/Source/Runtime/Core/Public/Modules/ModuleManager.h
@@ -211,10 +211,17 @@ public:
 	 *
 	 * @param InModuleName The name of the module to unload.
 	 * @param Ar The archive to receive error messages, if any.
-	 * @param bAbandonOnly Do not try to unload the module, only abandon.
 	 * @see LoadModuleWithCallback
 	 */
-	void UnloadOrAbandonModuleWithCallback( const FName InModuleName, FOutputDevice &Ar, bool bAbandonOnly = false);
+	void UnloadOrAbandonModuleWithCallback( const FName InModuleName, FOutputDevice &Ar);
+
+	/**
+	 * Calls PreUnload then abandons a module in memory.
+	 *
+	 * @param InModuleName The name of the module to unload.
+	 * @see LoadModuleWithCallback
+	 */
+	void AbandonModuleWithCallback( const FName InModuleName );
 
 	/** Delegate that's used by the module manager to find all the valid modules in a directory matching a pattern */
 	typedef TMap<FString, FString> FModuleNamesMap;
