Fixing UnrealHeaderTool compilation on Linux

diff --git a/Engine/Source/Programs/UnrealHeaderTool/Private/HeaderParser.cpp b/Engine/Source/Programs/UnrealHeaderTool/Private/HeaderParser.cpp
--- a/Engine/Source/Programs/UnrealHeaderTool/Private/HeaderParser.cpp
+++ b/Engine/Source/Programs/UnrealHeaderTool/Private/HeaderParser.cpp
@@ -2629,7 +2629,8 @@ void FHeaderParser::CompileDirective(FClasses& AllClasses)
 		bool bNotDefined = MatchSymbol(TEXT("!"));
 
 		int32 TempInt;
-		if (GetConstInt(TempInt) && TempInt == 0 || TempInt == 1)
+		const bool bParsedInt = GetConstInt(TempInt);
+		if (bParsedInt && (TempInt == 0 || TempInt == 1))
 		{
 			PushCompilerDirective(ECompilerDirective::Insignificant);
 		}
