Fix a couple of build errors from Dev-Core integration.

diff --git a/Engine/Source/Programs/UnrealHeaderTool/Private/CodeGenerator.cpp b/Engine/Source/Programs/UnrealHeaderTool/Private/CodeGenerator.cpp
--- a/Engine/Source/Programs/UnrealHeaderTool/Private/CodeGenerator.cpp
+++ b/Engine/Source/Programs/UnrealHeaderTool/Private/CodeGenerator.cpp
@@ -4891,8 +4891,8 @@ void FNativeClassHeaderGenerator::ApplyAlternatePropertyExportText(UProperty* Pr
 	UArrayProperty* ArrayProperty = Cast<UArrayProperty>(Prop);
 	UProperty* InnerProperty = ArrayProperty ? ArrayProperty->Inner : nullptr;
 	if (InnerProperty && (
-			InnerProperty->IsA<UByteProperty>() && ((UByteProperty*)InnerProperty)->Enum && FClass::IsDynamic(((UByteProperty*)InnerProperty)->Enum) ||
-			InnerProperty->IsA<UEnumProperty>()                                          && FClass::IsDynamic(((UEnumProperty*)InnerProperty)->Enum)
+			(InnerProperty->IsA<UByteProperty>() && ((UByteProperty*)InnerProperty)->Enum && FClass::IsDynamic(((UByteProperty*)InnerProperty)->Enum)) ||
+			(InnerProperty->IsA<UEnumProperty>()                                          && FClass::IsDynamic(((UEnumProperty*)InnerProperty)->Enum))
 		)
 	)
 	{
