Fix for invalid BP test when creating non-byte enums.

diff --git a/Engine/Source/Programs/UnrealHeaderTool/Private/HeaderParser.cpp b/Engine/Source/Programs/UnrealHeaderTool/Private/HeaderParser.cpp
--- a/Engine/Source/Programs/UnrealHeaderTool/Private/HeaderParser.cpp
+++ b/Engine/Source/Programs/UnrealHeaderTool/Private/HeaderParser.cpp
@@ -1389,7 +1389,7 @@ UEnum* FHeaderParser::CompileEnum()
 	}
 
 	static const FName BlueprintTypeName = TEXT("BlueprintType");
-	if (UnderlyingType != EUnderlyingEnumType::uint8 && !EnumToken.MetaData.Contains(BlueprintTypeName))
+	if (UnderlyingType != EUnderlyingEnumType::uint8 && EnumToken.MetaData.Contains(BlueprintTypeName))
 	{
 		FError::Throwf(TEXT("Invalid BlueprintType enum base - currently only uint8 supported"));
 	}
