Copying //UE4/Dev-Core to //UE4/Dev-Main (Source: //UE4/Dev-Core @ 3345728)

diff --git a/Engine/Source/Programs/UnrealHeaderTool/Private/CodeGenerator.cpp b/Engine/Source/Programs/UnrealHeaderTool/Private/CodeGenerator.cpp
--- a/Engine/Source/Programs/UnrealHeaderTool/Private/CodeGenerator.cpp
+++ b/Engine/Source/Programs/UnrealHeaderTool/Private/CodeGenerator.cpp
@@ -49,7 +49,9 @@
 #include "StringUtils.h"
 #include "Features/IModularFeatures.h"
 #include "UHTMakefile/UHTMakefile.h"
-#include "Misc/ScopeExit.h"
+#include "Algo/Sort.h"
+#include "Algo/Reverse.h"
+
 #include "FileLineException.h"
 
 /////////////////////////////////////////////////////
@@ -63,9 +65,6 @@ static TArray<FString> ChangeMessages;
 static bool bWriteContents = false;
 static bool bVerifyContents = false;
 
-static const bool bMultiLineUFUNCTION = true;
-static const bool bMultiLineUPROPERTY = true;
-
 static TSharedRef<FUnrealSourceFile> PerformInitialParseOnHeader(UPackage* InParent, const TCHAR* FileName, EObjectFlags Flags, const TCHAR* Buffer, FUHTMakefile& UHTMakefile);
 
 FCompilerMetadataManager GScriptHelper;
@@ -79,7 +78,15 @@ namespace
 	{
 		return FString::Printf(TEXT("TEXT(\"%s\")"), *InStr);
 	}
+
+	const TCHAR HeaderCopyright[] =
+		TEXT("// Copyright 1998-2017 Epic Games, Inc. All Rights Reserved.\r\n")
+		TEXT("/*===========================================================================\r\n")
+		TEXT("\tGenerated code exported from UnrealHeaderTool.\r\n")
+		TEXT("\tDO NOT modify this manually! Edit the corresponding .h files instead!\r\n")
+		TEXT("===========================================================================*/\r\n");
 }
+
 /**
  * Finds exact match of Identifier in string. Returns nullptr if none is found.
  *
@@ -424,23 +431,29 @@ FParmsAndReturnProperties GetFunctionParmsAndReturn(UFunction* Function)
  * @param	Function	the function to check
  * @return	true if the glue version of the function should be exported.
  */
-bool FNativeClassHeaderGenerator::ShouldExportFunction( UFunction* Function )
+bool ShouldExportUFunction(UFunction* Function)
 {
 	// export any script stubs for native functions declared in interface classes
 	bool bIsBlueprintNativeEvent = (Function->FunctionFlags & FUNC_BlueprintEvent) && (Function->FunctionFlags & FUNC_Native);
 	if (Function->GetOwnerClass()->HasAnyClassFlags(CLASS_Interface) && !bIsBlueprintNativeEvent)
+	{
 		return true;
+	}
 
 	// always export if the function is static
 	if (Function->FunctionFlags & FUNC_Static)
+	{
 		return true;
+	}
 
 	// don't export the function if this is not the original declaration and there is
 	// at least one parent version of the function that is declared native
 	for (UFunction* ParentFunction = Function->GetSuperFunction(); ParentFunction; ParentFunction = ParentFunction->GetSuperFunction())
 	{
 		if (ParentFunction->FunctionFlags & FUNC_Native)
+		{
 			return false;
+		}
 	}
 
 	return true;
@@ -517,26 +530,13 @@ static FString GetMetaDataCodeForObject(const UObject* Object, const TCHAR* Symb
 	return Result;
 }
 
-/**
- * Exports the struct's C++ properties to the specified output device and adds special
- * compiler directives for GCC to pack as we expect.
- *
- * @param	Struct				UStruct to export properties
- * @param	TextIndent			Current text indentation
- * @param	ImportsDefaults		whether this struct will be serialized with a default value
- */
-void FNativeClassHeaderGenerator::ExportProperties(UStruct* Struct, int32 TextIndent, bool bAccessSpecifiers, FUHTStringBuilder* Output)
+void FNativeClassHeaderGenerator::ExportProperties(FOutputDevice& Out, UStruct* Struct, int32 TextIndent)
 {
 	UProperty*	Previous			= NULL;
 	UProperty*	PreviousNonEditorOnly = NULL;
 	UProperty*	LastInSuper			= NULL;
 	UStruct*	InheritanceSuper	= Struct->GetInheritanceSuper();
 	bool		bEmittedHasEditorOnlyMacro = false;
-	bool		bEmittedHasScriptAlign = false;
-
-	check(Output != NULL);
-	FUHTStringBuilder& HeaderOutput = *Output;
-
 
 	// Find last property in the lowest base class that has any properties
 	UStruct* CurrentSuper = InheritanceSuper;
@@ -556,104 +556,26 @@ void FNativeClassHeaderGenerator::ExportProperties(UStruct* Struct, int32 TextIn
 		CurrentSuper = CurrentSuper->GetSuperStruct();
 	}
 
-	EPropertyHeaderExportFlags CurrentExportType = PROPEXPORT_Public;
-
-	// find structs that are nothing but bytes, account for editor only properties being 
-	// removed on consoles
-	int32 NumProperties = 0;
-	int32 NumByteProperties = 0;
-	int32 NumNonEditorOnlyProperties = 0;
-	int32 NumNonEditorOnlyByteProperties = 0;
-	for( TFieldIterator<UProperty> It(Struct, EFieldIteratorFlags::ExcludeSuper); It; ++It )
-	{
-		// treat bitfield and bytes the same
-		bool bIsByteProperty = It->IsA(UByteProperty::StaticClass());// || It->IsA(UBoolProperty::StaticClass());
-		bool bIsEditorOnlyProperty = It->IsEditorOnlyProperty();
-		
-		// count our propertie
-		NumProperties++;
-		if (bIsByteProperty)
-		{
-			NumByteProperties++;
-		}
-		if (!bIsEditorOnlyProperty)
-		{
-			NumNonEditorOnlyProperties++;
-		}
-		if (!bIsEditorOnlyProperty && bIsByteProperty)
-		{
-			NumNonEditorOnlyByteProperties++;
-		}
-	}
-
-	bool bCurrentlyInNotCPPBlock = false;
 	// Iterate over all properties in this struct.
 	for( TFieldIterator<UProperty> It(Struct, EFieldIteratorFlags::ExcludeSuper); It; ++It )
 	{
 		UProperty* Current = *It;
 
-		FUHTStringBuilder PropertyText;
-
 		// Disregard properties with 0 size like functions.
 		if (It.GetStruct() == Struct)
 		{
-			FString AccessSpecifier = TEXT("public");
-			if (bAccessSpecifiers)
-			{
-				// find the class info for this class
-				FClassMetaData* ClassData = GScriptHelper.FindClassData(Struct);
-				check(ClassData);
-
-				// find the compiler token for this property
-				FTokenData* PropData = ClassData->FindTokenData(Current);
-				if ( PropData != NULL )
-				{
-					// if this property has a different access specifier, then export that now
-					if ( (PropData->Token.PropertyExportFlags & CurrentExportType) == 0 )
-					{
-						if ( (PropData->Token.PropertyExportFlags & PROPEXPORT_Private) != 0 )
-						{
-							CurrentExportType = PROPEXPORT_Private;
-							AccessSpecifier = TEXT("private");
-						}
-						else if ( (PropData->Token.PropertyExportFlags & PROPEXPORT_Protected) != 0 )
-						{
-							CurrentExportType = PROPEXPORT_Protected;
-							AccessSpecifier = TEXT("protected");
-						}
-						else
-						{
-							CurrentExportType = PROPEXPORT_Public;
-							AccessSpecifier = TEXT("public");
-						}
-
-						if ( AccessSpecifier.Len() )
-						{
-							// If we are changing the access specifier we need to emit the #endif for the WITH_EDITORONLY_DATA macro first otherwise the access specifier may
-							// only be conditionally compiled in.
-							if( bEmittedHasEditorOnlyMacro )
-							{
-								PropertyText.Logf( TEXT("#endif // WITH_EDITORONLY_DATA\r\n") );
-								bEmittedHasEditorOnlyMacro = false;
-							}
-
-							PropertyText.Logf(TEXT("%s:") LINE_TERMINATOR, *AccessSpecifier);
-						}
-					}
-				}
-			}
 			// If we are switching from editor to non-editor or vice versa and the state of the WITH_EDITORONLY_DATA macro emission doesn't match, generate the 
 			// #if or #endif appropriately.
 			bool RequiresHasEditorOnlyMacro = Current->IsEditorOnlyProperty();
 			if( !bEmittedHasEditorOnlyMacro && RequiresHasEditorOnlyMacro )
 			{
 				// Indent code and export CPP text.
-				PropertyText.Logf( TEXT("#if WITH_EDITORONLY_DATA\r\n") );
+				Out.Logf( TEXT("#if WITH_EDITORONLY_DATA\r\n") );
 				bEmittedHasEditorOnlyMacro = true;
 			}
 			else if( bEmittedHasEditorOnlyMacro && !RequiresHasEditorOnlyMacro )
 			{
-				PropertyText.Logf( TEXT("#endif // WITH_EDITORONLY_DATA\r\n") );
+				Out.Logf( TEXT("#endif // WITH_EDITORONLY_DATA\r\n") );
 				bEmittedHasEditorOnlyMacro = false;
 			}
 
@@ -664,10 +586,10 @@ void FNativeClassHeaderGenerator::ExportProperties(UStruct* Struct, int32 TextIn
 
 				const FString* Dim = GArrayDimensions.Find(Current);
 				Current->ExportCppDeclaration( JustPropertyDecl, EExportedDeclaration::Member, Dim ? **Dim : NULL);
-				ApplyAlternatePropertyExportText(*It, JustPropertyDecl);
+				ApplyAlternatePropertyExportText(*It, JustPropertyDecl, EExportingState::TypeEraseDelegates);
 
 				// Finish up line.
-				PropertyText.Logf(TEXT("%s%s;\r\n"), FCString::Tab(TextIndent + 1), *JustPropertyDecl);
+				Out.Logf(TEXT("%s%s;\r\n"), FCString::Tab(TextIndent + 1), *JustPropertyDecl);
 			}
 
 			LastInSuper	= NULL;
@@ -676,26 +598,13 @@ void FNativeClassHeaderGenerator::ExportProperties(UStruct* Struct, int32 TextIn
 			{
 				PreviousNonEditorOnly = Current;
 			}
-
-			HeaderOutput.Log(PropertyText);
 		}
 	}
-	if (bCurrentlyInNotCPPBlock)
-	{
-		HeaderOutput += TEXT("#endif\r\n");
-		bCurrentlyInNotCPPBlock = false;
-	}
 
 	// End of property list.  If we haven't generated the WITH_EDITORONLY_DATA #endif, do so now.
 	if (bEmittedHasEditorOnlyMacro)
 	{
-		HeaderOutput += TEXT("#endif // WITH_EDITORONLY_DATA\r\n");
-	}
-
-	// if the last property that was exported wasn't public, emit a line to reset the access to "public" so that we don't interfere with cpptext
-	if (CurrentExportType != PROPEXPORT_Public)
-	{
-		HeaderOutput += TEXT("public:") LINE_TERMINATOR;
+		Out.Log(TEXT("#endif // WITH_EDITORONLY_DATA\r\n"));
 	}
 }
 
@@ -917,22 +826,11 @@ FString FNativeClassHeaderGenerator::GetSingletonName(UField* Item, bool bRequir
 		// this is a cross module reference, we need to include the right extern decl
 		FString Extern = FTypeSingletonCache::Get(Item, bRequiresValidObject).GetExternDecl();
 
-		bool bAlreadyInSet = false;
-		UniqueCrossModuleReferences.Add(*Extern, &bAlreadyInSet);
-
-		if (!bAlreadyInSet)
-		{
-			CrossModuleGeneratedFunctionDeclarations.Log(*Extern);
-		}
+		UniqueCrossModuleReferences.Add(*Extern);
 	}
 	return Result;
 }
 
-FString FNativeClassHeaderGenerator::GetSingletonName(FClass* Item, bool bRequiresValidObject)
-{
-	return GetSingletonName((UClass*)Item, bRequiresValidObject);
-}
-
 FString FNativeClassHeaderGenerator::GetOverriddenName(const UField* Item)
 {
 	FString OverriddenName = Item->GetMetaData(TEXT("OverrideNativeName"));
@@ -1088,7 +986,7 @@ FString FNativeClassHeaderGenerator::PropertyNew(FString& Meta, UProperty* Prop,
 	return Lines;
 }
 
-void FNativeClassHeaderGenerator::OutputProperties(FString& Meta, FOutputDevice& OutputDevice, const FString& OuterString, TArray<UProperty*>& Properties, const TCHAR* Spaces)
+void FNativeClassHeaderGenerator::OutputProperties(FString& Meta, FOutputDevice& OutputDevice, const FString& OuterString, const TArray<UProperty*>& Properties, const TCHAR* Spaces)
 {
 	bool bEmittedHasEditorOnlyMacro = false;
 	for (int32 Index = Properties.Num() - 1; Index >= 0; Index--)
@@ -1162,7 +1060,7 @@ void FNativeClassHeaderGenerator::OutputProperty(FString& Meta, FOutputDevice& O
 			SourceStruct = GetEventStructParamsName(Function->GetOuter(), *FunctionName);
 		}
 		else
-		{			
+		{
 			SourceStruct = NameLookupCPP.GetNameCPP(CastChecked<UStruct>(Prop->GetOuter()));
 		}
 		FString PropMacroOuterClass;
@@ -1270,7 +1168,7 @@ static bool IsAlwaysAccessible(UScriptStruct* Script)
 	return Result;
 }
 
-static void FindNoExportStructs(TArray<UScriptStruct*>& Structs, UStruct* Start)
+static void FindNoExportStructsRecursive(TArray<UScriptStruct*>& Structs, UStruct* Start)
 {
 	while (Start)
 	{
@@ -1281,7 +1179,7 @@ static void FindNoExportStructs(TArray<UScriptStruct*>& Structs, UStruct* Start)
 				break;
 			}
 
-			if (!IsAlwaysAccessible(StartScript)) // these are a special cases that already exists and if wrong if exported nievely
+			if (!IsAlwaysAccessible(StartScript)) // these are a special cases that already exists and if wrong if exported naively
 			{
 				// this will topologically sort them in reverse order
 				Structs.Remove(StartScript);
@@ -1293,31 +1191,31 @@ static void FindNoExportStructs(TArray<UScriptStruct*>& Structs, UStruct* Start)
 		{
 			if (UStructProperty* StructProp = Cast<UStructProperty>(Prop))
 			{
-				FindNoExportStructs(Structs, StructProp->Struct);
+				FindNoExportStructsRecursive(Structs, StructProp->Struct);
 			}
 			else if (UArrayProperty* ArrayProp = Cast<UArrayProperty>(Prop))
 			{
 				if (UStructProperty* InnerStructProp = Cast<UStructProperty>(ArrayProp->Inner))
 				{
-					FindNoExportStructs(Structs, InnerStructProp->Struct);
+					FindNoExportStructsRecursive(Structs, InnerStructProp->Struct);
 				}
 			}
 			else if (UMapProperty* MapProp = Cast<UMapProperty>(Prop))
 			{
 				if (UStructProperty* KeyStructProp = Cast<UStructProperty>(MapProp->KeyProp))
 				{
-					FindNoExportStructs(Structs, KeyStructProp->Struct);
+					FindNoExportStructsRecursive(Structs, KeyStructProp->Struct);
 				}
 				if (UStructProperty* ValueStructProp = Cast<UStructProperty>(MapProp->ValueProp))
 				{
-					FindNoExportStructs(Structs, ValueStructProp->Struct);
+					FindNoExportStructsRecursive(Structs, ValueStructProp->Struct);
 				}
 			}
 			else if (USetProperty* SetProp = Cast<USetProperty>(Prop))
 			{
 				if (UStructProperty* ElementStructProp = Cast<UStructProperty>(SetProp->ElementProp))
 				{
-					FindNoExportStructs(Structs, ElementStructProp->Struct);
+					FindNoExportStructsRecursive(Structs, ElementStructProp->Struct);
 				}
 			}
 		}
@@ -1325,38 +1223,47 @@ static void FindNoExportStructs(TArray<UScriptStruct*>& Structs, UStruct* Start)
 	}
 }
 
+static TArray<UScriptStruct*> FindNoExportStructs(UStruct* Start)
+{
+	TArray<UScriptStruct*> Result;
+	FindNoExportStructsRecursive(Result, Start);
+
+	// These come out in reverse order of topology so reverse them
+	Algo::Reverse(Result);
+
+	return Result;
+}
+
 FString GetPackageSingletonName(const UPackage* Package)
 {
 	static FString ClassString = NameLookupCPP.GetNameCPP(UPackage::StaticClass());
 	return FString(TEXT("Z_Construct_")) + ClassString + TEXT("_") + Package->GetName().Replace(TEXT("/"), TEXT("_")) + TEXT("()");
 }
 
-void FNativeClassHeaderGenerator::ExportGeneratedPackageInitCode(const UPackage* InPackage)
+void FNativeClassHeaderGenerator::ExportGeneratedPackageInitCode(FOutputDevice& Out, FUHTStringBuilder& OutDeclarations, const UPackage* InPackage)
 {
 	FString ApiString = GetAPIString();
-	FString SingletonName(GetPackageSingletonName(InPackage));
-
-	FUHTStringBuilder& GeneratedFunctionText = GetGeneratedFunctionTextDevice();
+	FString SingletonName = GetPackageSingletonName(InPackage);
 
-	GeneratedFunctionDeclarations.Logf(TEXT("\t%sclass UPackage* %s;\r\n"), *ApiString, *SingletonName);
+	OutDeclarations.Logf(TEXT("\t%sclass UPackage* %s;\r\n"), *ApiString, *SingletonName);
 
-	GeneratedFunctionText.Logf(TEXT("\tUPackage* %s\r\n"), *SingletonName);
-	GeneratedFunctionText.Logf(TEXT("\t{\r\n"));
-	GeneratedFunctionText.Logf(TEXT("\t\tstatic UPackage* ReturnPackage = NULL;\r\n"));
-	GeneratedFunctionText.Logf(TEXT("\t\tif (!ReturnPackage)\r\n"));
-	GeneratedFunctionText.Logf(TEXT("\t\t{\r\n"));
-	GeneratedFunctionText.Logf(TEXT("\t\t\tReturnPackage = CastChecked<UPackage>(StaticFindObjectFast(UPackage::StaticClass(), NULL, FName(TEXT(\"%s\")), false, false));\r\n"), *InPackage->GetName());
+	Out.Logf(TEXT("\tUPackage* %s\r\n"), *SingletonName);
+	Out.Logf(TEXT("\t{\r\n"));
+	Out.Logf(TEXT("\t\tstatic UPackage* ReturnPackage = NULL;\r\n"));
+	Out.Logf(TEXT("\t\tif (!ReturnPackage)\r\n"));
+	Out.Logf(TEXT("\t\t{\r\n"));
+	Out.Logf(TEXT("\t\t\tReturnPackage = CastChecked<UPackage>(StaticFindObjectFast(UPackage::StaticClass(), NULL, FName(TEXT(\"%s\")), false, false));\r\n"), *InPackage->GetName());
 
 	FString Meta = GetMetaDataCodeForObject(InPackage, TEXT("ReturnPackage"), TEXT("\t\t\t"));
 	if (Meta.Len())
 	{
-		GeneratedFunctionText.Logf(TEXT("#if WITH_METADATA\r\n"));
-		GeneratedFunctionText.Logf(TEXT("\t\t\tUMetaData* MetaData = ReturnPackage->GetMetaData();\r\n"));
-		GeneratedFunctionText.Log(*Meta);
-		GeneratedFunctionText.Logf(TEXT("#endif\r\n"));
+		Out.Logf(TEXT("#if WITH_METADATA\r\n"));
+		Out.Logf(TEXT("\t\t\tUMetaData* MetaData = ReturnPackage->GetMetaData();\r\n"));
+		Out.Log(*Meta);
+		Out.Logf(TEXT("#endif\r\n"));
 	}
 
-	GeneratedFunctionText.Logf(TEXT("\t\t\tReturnPackage->SetPackageFlags(PKG_CompiledIn | 0x%08X);\r\n"), InPackage->GetPackageFlags() & (PKG_ClientOptional | PKG_ServerSideOnly | PKG_EditorOnly | PKG_Developer));
+	Out.Logf(TEXT("\t\t\tReturnPackage->SetPackageFlags(PKG_CompiledIn | 0x%08X);\r\n"), InPackage->GetPackageFlags() & (PKG_ClientOptional | PKG_ServerSideOnly | PKG_EditorOnly | PKG_Developer));
 	TheFlagAudit.Add(InPackage, TEXT("PackageFlags"), InPackage->GetPackageFlags());
 	{
 		FGuid Guid;
@@ -1377,15 +1284,14 @@ void FNativeClassHeaderGenerator::ExportGeneratedPackageInitCode(const UPackage*
 		}
 
 		Guid.A = CombinedCRC;
-		Guid.B = GenerateTextCRC(*GeneratedFunctionDeclarations);
-		GeneratedFunctionText.Logf(TEXT("\t\t\tFGuid Guid;\r\n"));
-		GeneratedFunctionText.Logf(TEXT("\t\t\tGuid.A = 0x%08X;\r\n"), Guid.A);
-		GeneratedFunctionText.Logf(TEXT("\t\t\tGuid.B = 0x%08X;\r\n"), Guid.B);
-		GeneratedFunctionText.Logf(TEXT("\t\t\tGuid.C = 0x%08X;\r\n"), Guid.C);
-		GeneratedFunctionText.Logf(TEXT("\t\t\tGuid.D = 0x%08X;\r\n"), Guid.D);
-		GeneratedFunctionText.Logf(TEXT("\t\t\tReturnPackage->SetGuid(Guid);\r\n"), Guid.D);
-
-		GeneratedFunctionText.Log(TEXT("\r\n"));
+		Guid.B = GenerateTextCRC(*OutDeclarations);
+		Out.Logf(TEXT("\t\t\tFGuid Guid;\r\n"));
+		Out.Logf(TEXT("\t\t\tGuid.A = 0x%08X;\r\n"), Guid.A);
+		Out.Logf(TEXT("\t\t\tGuid.B = 0x%08X;\r\n"), Guid.B);
+		Out.Logf(TEXT("\t\t\tGuid.C = 0x%08X;\r\n"), Guid.C);
+		Out.Logf(TEXT("\t\t\tGuid.D = 0x%08X;\r\n"), Guid.D);
+		Out.Logf(TEXT("\t\t\tReturnPackage->SetGuid(Guid);\r\n"), Guid.D);
+		Out.Log(TEXT("\r\n"));
 
 		for (UField* ScriptType : TObjectRange<UField>())
 		{
@@ -1402,25 +1308,25 @@ void FNativeClassHeaderGenerator::ExportGeneratedPackageInitCode(const UPackage*
 				UField* FieldOuter = Cast<UField>(ScriptType->GetOuter());
 				if (!FieldOuter || !FClass::IsDynamic(FieldOuter))
 				{
-					GeneratedFunctionText.Logf(TEXT("\t\t\t%s;\r\n"), *GetSingletonName(ScriptType));
+					Out.Logf(TEXT("\t\t\t%s;\r\n"), *GetSingletonName(ScriptType));
 				}
 			}
 		}
 	}
 
-	GeneratedFunctionText.Logf(TEXT("\t\t}\r\n"));
-	GeneratedFunctionText.Logf(TEXT("\t\treturn ReturnPackage;\r\n"));
-	GeneratedFunctionText.Logf(TEXT("\t}\r\n"));
-
+	Out.Logf(TEXT("\t\t}\r\n"));
+	Out.Logf(TEXT("\t\treturn ReturnPackage;\r\n"));
+	Out.Logf(TEXT("\t}\r\n"));
 }
 
-void FNativeClassHeaderGenerator::ExportNativeGeneratedInitCode(FUnrealSourceFile& SourceFile, FClass* Class, FUHTStringBuilder& OutFriendText)
+void FNativeClassHeaderGenerator::ExportNativeGeneratedInitCode(FOutputDevice& Out, FOutputDevice& OutDeclarations, const FUnrealSourceFile& SourceFile, FClass* Class, FUHTStringBuilder& OutFriendText)
 {
-	FUHTStringBuilder& GeneratedFunctionText = GetGeneratedFunctionTextDevice();
 	check(!OutFriendText.Len());
-	// Emit code to build the UObjects that used to be in .u files
-	const bool bIsNoExport = Class->HasAnyClassFlags(CLASS_NoExport);
-	const bool bIsDynamic = FClass::IsDynamic(Class);
+
+	const bool   bIsNoExport  = Class->HasAnyClassFlags(CLASS_NoExport);
+	const bool   bIsDynamic   = FClass::IsDynamic(Class);
+	const TCHAR* ClassNameCPP = NameLookupCPP.GetNameCPP(Class);
+
 	FUHTStringBuilder BodyText;
 	FUHTStringBuilder CallSingletons;
 	FString ApiString = GetAPIString();
@@ -1449,13 +1355,12 @@ void FNativeClassHeaderGenerator::ExportNativeGeneratedInitCode(FUnrealSourceFil
 	FunctionsToExport.Sort();
 
 	// Export the init code for each function
-	for (int32 FuncIndex = 0; FuncIndex < FunctionsToExport.Num(); FuncIndex++)
+	for (UFunction* Function : FunctionsToExport)
 	{
-		UFunction* Function = FunctionsToExport[FuncIndex];
-
 		if (!Function->IsA<UDelegateFunction>())
 		{
-			ExportFunction(SourceFile, Function, &FScope::GetTypeScope(Class).Get(), bIsNoExport);
+			OutDeclarations.Log(FTypeSingletonCache::Get(Function).GetExternDecl());
+			ExportFunction(Out, SourceFile, Function, bIsNoExport);
 		}
 
 		CallSingletons.Logf(TEXT("\t\t\t\tOuterClass->LinkChild(%s);\r\n"), *GetSingletonName(Function));
@@ -1469,17 +1374,17 @@ void FNativeClassHeaderGenerator::ExportNativeGeneratedInitCode(FUnrealSourceFil
 		{
 			FString SingletonNameNoRegister(GetSingletonName(Class, false));
 
-			GeneratedFunctionDeclarations.Log(FTypeSingletonCache::Get(Class, false).GetExternDecl());
+			OutDeclarations.Log(FTypeSingletonCache::Get(Class, false).GetExternDecl());
 
 			GeneratedClassRegisterFunctionText.Logf(TEXT("\tUClass* %s\r\n"), *SingletonNameNoRegister);
 			GeneratedClassRegisterFunctionText.Logf(TEXT("\t{\r\n"));
-			GeneratedClassRegisterFunctionText.Logf(TEXT("\t\treturn %s::StaticClass();\r\n"), NameLookupCPP.GetNameCPP(Class));
+			GeneratedClassRegisterFunctionText.Logf(TEXT("\t\treturn %s::StaticClass();\r\n"), ClassNameCPP);
 			GeneratedClassRegisterFunctionText.Logf(TEXT("\t}\r\n"));
 		}
-		FString SingletonName(GetSingletonName(Class));
+		FString SingletonName = GetSingletonName(Class);
 
 		OutFriendText.Logf(TEXT("\tfriend %sclass UClass* %s;\r\n"), *ApiString, *SingletonName);
-		GeneratedFunctionDeclarations.Log(FTypeSingletonCache::Get(Class).GetExternDecl());
+		OutDeclarations.Log(FTypeSingletonCache::Get(Class).GetExternDecl());
 
 		GeneratedClassRegisterFunctionText.Logf(TEXT("\tUClass* %s\r\n"), *SingletonName);
 		GeneratedClassRegisterFunctionText.Logf(TEXT("\t{\r\n"));
@@ -1505,9 +1410,9 @@ void FNativeClassHeaderGenerator::ExportNativeGeneratedInitCode(FUnrealSourceFil
 		{
 			GeneratedClassRegisterFunctionText.Logf(TEXT("\t\t\t%s;\r\n"), *GetPackageSingletonName(CastChecked<UPackage>(Class->GetOutermost())));
 		}
-		GeneratedClassRegisterFunctionText.Logf(TEXT("\t\t\tOuterClass = %s::StaticClass();\r\n"), NameLookupCPP.GetNameCPP(Class));
-		GeneratedClassRegisterFunctionText.Logf(TEXT("\t\t\tif (!(OuterClass->ClassFlags & CLASS_Constructed))\r\n"), NameLookupCPP.GetNameCPP(Class));
-		GeneratedClassRegisterFunctionText.Logf(TEXT("\t\t\t{\r\n"), NameLookupCPP.GetNameCPP(Class));
+		GeneratedClassRegisterFunctionText.Logf(TEXT("\t\t\tOuterClass = %s::StaticClass();\r\n"), ClassNameCPP);
+		GeneratedClassRegisterFunctionText.Logf(TEXT("\t\t\tif (!(OuterClass->ClassFlags & CLASS_Constructed))\r\n"), ClassNameCPP);
+		GeneratedClassRegisterFunctionText.Logf(TEXT("\t\t\t{\r\n"), ClassNameCPP);
 		GeneratedClassRegisterFunctionText.Logf(TEXT("\t\t\t\tUObjectForceRegistration(OuterClass);\r\n"));
 		uint32 Flags = (Class->ClassFlags & CLASS_SaveInCompiledInClasses) | CLASS_Constructed;
 		GeneratedClassRegisterFunctionText.Logf(TEXT("\t\t\t\tOuterClass->ClassFlags |= 0x%08X;\r\n"), Flags);
@@ -1516,7 +1421,7 @@ void FNativeClassHeaderGenerator::ExportNativeGeneratedInitCode(FUnrealSourceFil
 		GeneratedClassRegisterFunctionText.Log(CallSingletons);
 		GeneratedClassRegisterFunctionText.Logf(TEXT("\r\n"));
 
-		FString OuterString = FString(TEXT("OuterClass"));
+		FString OuterString = TEXT("OuterClass");
 
 		TMap<FName, FString>* MetaDataMap = UMetaData::GetMapForObject(Class);
 		{
@@ -1574,7 +1479,7 @@ void FNativeClassHeaderGenerator::ExportNativeGeneratedInitCode(FUnrealSourceFil
 			FString OffsetString(TEXT("0"));
 			if (Inter.PointerOffset)
 			{
-				OffsetString = FString::Printf(TEXT("VTABLE_OFFSET(%s, %s)"), NameLookupCPP.GetNameCPP(Class), NameLookupCPP.GetNameCPP(Inter.Class, true));
+				OffsetString = FString::Printf(TEXT("VTABLE_OFFSET(%s, %s)"), ClassNameCPP, NameLookupCPP.GetNameCPP(Inter.Class, true));
 			}
 			GeneratedClassRegisterFunctionText.Logf(TEXT("\t\t\t\tOuterClass->Interfaces.Add(FImplementedInterface(%s, %s, %s ));\r\n"), 
 				*GetSingletonName(Inter.Class, false),
@@ -1614,75 +1519,68 @@ void FNativeClassHeaderGenerator::ExportNativeGeneratedInitCode(FUnrealSourceFil
 		GeneratedClassRegisterFunctionText.Logf(TEXT("\t\treturn OuterClass;\r\n"));
 		GeneratedClassRegisterFunctionText.Logf(TEXT("\t}\r\n"));
 
-		GeneratedFunctionText += GeneratedClassRegisterFunctionText;
+		Out.Logf(TEXT("%s"), *GeneratedClassRegisterFunctionText);
 	}
 
 	if (OutFriendText.Len() && bIsNoExport)
 	{
-		GeneratedFunctionText.Logf(TEXT("\t/* friend declarations for pasting into noexport class %s\r\n"), NameLookupCPP.GetNameCPP(Class));
-		GeneratedFunctionText.Log(OutFriendText);
-		GeneratedFunctionText.Logf(TEXT("\t*/\r\n"));
+		Out.Logf(TEXT("\t/* friend declarations for pasting into noexport class %s\r\n"), ClassNameCPP);
+		Out.Log(OutFriendText);
+		Out.Logf(TEXT("\t*/\r\n"));
 		OutFriendText.Reset();
 	}
-	FString SingletonName(GetSingletonName(Class));
+
+	FString SingletonName = GetSingletonName(Class);
 	SingletonName.ReplaceInline(TEXT("()"), TEXT(""), ESearchCase::CaseSensitive); // function address
 
-	{	
-		FString OverriddenClassName = *FNativeClassHeaderGenerator::GetOverriddenName(Class);
-		const TCHAR* ClassNameCPP = NameLookupCPP.GetNameCPP(Class);
+	FString OverriddenClassName = *FNativeClassHeaderGenerator::GetOverriddenName(Class);
 
-		const FString InitSearchableValuesFunctionName = bIsDynamic ? Class->GetMetaData(TEXT("InitializeStaticSearchableValues")) : FString();
-		const FString InitSearchableValuesFunctionParam = InitSearchableValuesFunctionName.IsEmpty() ? FString(TEXT("nullptr")) :
-			FString::Printf(TEXT("&%s::%s"), ClassNameCPP, *InitSearchableValuesFunctionName);
+	const FString InitSearchableValuesFunctionName = bIsDynamic ? Class->GetMetaData(TEXT("InitializeStaticSearchableValues")) : FString();
+	const FString InitSearchableValuesFunctionParam = InitSearchableValuesFunctionName.IsEmpty() ? FString(TEXT("nullptr")) :
+		FString::Printf(TEXT("&%s::%s"), ClassNameCPP, *InitSearchableValuesFunctionName);
 
-		GeneratedFunctionText.Logf(TEXT("\tstatic FCompiledInDefer Z_CompiledInDefer_UClass_%s(%s, &%s::StaticClass, TEXT(\"%s\"), TEXT(\"%s\"), %s, %s, %s, %s);\r\n"),
-			ClassNameCPP, 
-			*SingletonName, 
-			ClassNameCPP, 
-			bIsDynamic ? *FClass::GetTypePackageName(Class) : *Class->GetOutermost()->GetName(),
-			bIsDynamic ? *OverriddenClassName : ClassNameCPP,
-			bIsDynamic ? TEXT("true") : TEXT("false"),
-			bIsDynamic ? *AsTEXT(FClass::GetTypePackageName(Class)) : TEXT("nullptr"),
-			bIsDynamic ? *AsTEXT(FNativeClassHeaderGenerator::GetOverriddenPathName(Class)) : TEXT("nullptr"),
-			*InitSearchableValuesFunctionParam);
-		
-		// Append base class' CRC at the end of the generated code, this will force update derived classes
-		// when base class changes during hot-reload.
-		uint32 BaseClassCRC = 0;
-		if (Class->GetSuperClass() && !Class->GetSuperClass()->HasAnyClassFlags(CLASS_Intrinsic))
-		{
-			BaseClassCRC = GGeneratedCodeCRCs.FindChecked(Class->GetSuperClass());
-		}
-		GeneratedClassRegisterFunctionText.Logf(TEXT("\r\n// %u\r\n"), BaseClassCRC);
+	// Append base class' CRC at the end of the generated code, this will force update derived classes
+	// when base class changes during hot-reload.
+	uint32 BaseClassCRC = 0;
+	if (Class->GetSuperClass() && !Class->GetSuperClass()->HasAnyClassFlags(CLASS_Intrinsic))
+	{
+		BaseClassCRC = GGeneratedCodeCRCs.FindChecked(Class->GetSuperClass());
+	}
+	GeneratedClassRegisterFunctionText.Logf(TEXT("\r\n// %u\r\n"), BaseClassCRC);
 
-		// Calculate generated class initialization code CRC so that we know when it changes after hot-reload
-		uint32 ClassCrc = GenerateTextCRC(*GeneratedClassRegisterFunctionText);
-		GGeneratedCodeCRCs.Add(Class, ClassCrc);
-		UHTMakefile.AddGeneratedCodeCRC(&SourceFile, Class, ClassCrc);
-		// Emit the IMPLEMENT_CLASS macro to go in the generated cpp file.
-		if (!bIsDynamic)
-		{
-			GeneratedPackageCPP.Logf(TEXT("\tIMPLEMENT_CLASS(%s, %u);\r\n"), ClassNameCPP, ClassCrc);
-		}
-		else
-		{
-			GeneratedPackageCPP.Logf(TEXT("\tIMPLEMENT_DYNAMIC_CLASS(%s, TEXT(\"%s\"), %u);\r\n"), ClassNameCPP, *OverriddenClassName, ClassCrc);
-		}
+	// Calculate generated class initialization code CRC so that we know when it changes after hot-reload
+	uint32 ClassCrc = GenerateTextCRC(*GeneratedClassRegisterFunctionText);
+	GGeneratedCodeCRCs.Add(Class, ClassCrc);
+	UHTMakefile.AddGeneratedCodeCRC(&SourceFile, Class, ClassCrc);
+	// Emit the IMPLEMENT_CLASS macro to go in the generated cpp file.
+	if (!bIsDynamic)
+	{
+		Out.Logf(TEXT("\tIMPLEMENT_CLASS(%s, %u);\r\n"), ClassNameCPP, ClassCrc);
 	}
+	else
+	{
+		Out.Logf(TEXT("\tIMPLEMENT_DYNAMIC_CLASS(%s, TEXT(\"%s\"), %u);\r\n"), ClassNameCPP, *OverriddenClassName, ClassCrc);
+	}
+
+	Out.Logf(TEXT("\tstatic FCompiledInDefer Z_CompiledInDefer_UClass_%s(%s, &%s::StaticClass, TEXT(\"%s\"), TEXT(\"%s\"), %s, %s, %s, %s);\r\n"),
+		ClassNameCPP,
+		*SingletonName,
+		ClassNameCPP,
+		bIsDynamic ? *FClass::GetTypePackageName(Class) : *Class->GetOutermost()->GetName(),
+		bIsDynamic ? *OverriddenClassName : ClassNameCPP,
+		bIsDynamic ? TEXT("true") : TEXT("false"),
+		bIsDynamic ? *AsTEXT(FClass::GetTypePackageName(Class)) : TEXT("nullptr"),
+		bIsDynamic ? *AsTEXT(FNativeClassHeaderGenerator::GetOverriddenPathName(Class)) : TEXT("nullptr"),
+		*InitSearchableValuesFunctionParam);
 }
 
-void FNativeClassHeaderGenerator::ExportFunction(FUnrealSourceFile& SourceFile, UFunction* Function, FScope* Scope, bool bIsNoExport)
+void FNativeClassHeaderGenerator::ExportFunction(FOutputDevice& Out, const FUnrealSourceFile& SourceFile, UFunction* Function, bool bIsNoExport)
 {
 	UFunction* SuperFunction = Function->GetSuperFunction();
 
-	auto* CompilerInfo = FFunctionData::FindForFunction(Function);
-
 	bool bIsDelegate = Function->HasAnyFunctionFlags(FUNC_Delegate);
 
-	const FFuncInfo& FunctionData = CompilerInfo->GetFunctionData();
-
-	const FString SingletonName(GetSingletonName(Function));
-	GeneratedFunctionDeclarations.Log(FTypeSingletonCache::Get(Function).GetExternDecl());
+	const FString SingletonName = GetSingletonName(Function);
 
 	FUHTStringBuilder CurrentFunctionText;
 
@@ -1691,16 +1589,13 @@ void FNativeClassHeaderGenerator::ExportFunction(FUnrealSourceFile& SourceFile,
 
 	if (bIsNoExport || !(Function->FunctionFlags&FUNC_Event))  // non-events do not export a params struct, so lets do that locally for offset determination
 	{
-		TGuardValue<bool> Guard(bIsExportingForOffsetDeterminationOnly, true);
-		TArray<UScriptStruct*> Structs;
-		FindNoExportStructs(Structs, Function);
-		if (Structs.Num())
+		TArray<UScriptStruct*> Structs = FindNoExportStructs(Function);
+		for (UScriptStruct* Struct : Structs)
 		{
-			ExportMirrorsForNoexportStructs(Structs, /*Indent=*/ 2, CurrentFunctionText);
+			ExportMirrorsForNoexportStruct(CurrentFunctionText, Struct, /*Indent=*/ 2);
 		}
-		TArray<UFunction*> CallbackFunctions;
-		CallbackFunctions.Add(Function);
-		ExportEventParms(*Scope, CallbackFunctions, CurrentFunctionText, /*Indent=*/ 2, /*bOutputConstructor=*/ false);
+
+		ExportEventParm(CurrentFunctionText, ForwardDeclarations, Function, /*Indent=*/ 2, /*bOutputConstructor=*/ false, EExportingState::TypeEraseDelegates);
 	}
 
 	if (UObject* Outer = Function->GetOuter())
@@ -1776,6 +1671,8 @@ void FNativeClassHeaderGenerator::ExportFunction(FUnrealSourceFile& SourceFile,
 		OutputProperty(Meta, CurrentFunctionText, OuterString, Props[Index], TEXT("\t\t\t"));
 	}
 
+	FFunctionData* CompilerInfo = FFunctionData::FindForFunction(Function);
+	const FFuncInfo& FunctionData = CompilerInfo->GetFunctionData();
 	if (FunctionData.FunctionFlags & (FUNC_NetRequest | FUNC_NetResponse))
 	{
 		CurrentFunctionText.Logf(TEXT("\t\t\tReturnFunction->RPCId=%d;\r\n"), FunctionData.RPCId);
@@ -1800,64 +1697,99 @@ void FNativeClassHeaderGenerator::ExportFunction(FUnrealSourceFile& SourceFile,
 	uint32 FunctionCrc = GenerateTextCRC(*CurrentFunctionText);
 	GGeneratedCodeCRCs.Add(Function, FunctionCrc);
 	UHTMakefile.AddGeneratedCodeCRC(&SourceFile, Function, FunctionCrc);
-	GetGeneratedFunctionTextDevice() += CurrentFunctionText;
+	Out.Log(CurrentFunctionText);
 }
 
-void FNativeClassHeaderGenerator::ExportNatives(FClass* Class)
+void FNativeClassHeaderGenerator::ExportNatives(FOutputDevice& Out, FClass* Class)
 {
-	// Skip no export classes.
-	if (Class->HasAnyClassFlags(CLASS_NoExport))
-	{
-		return;
-	}
+	const TCHAR* ClassCPPName = NameLookupCPP.GetNameCPP(Class);
+	FString TypeName = Class->HasAnyClassFlags(CLASS_Interface) ? *FString::Printf(TEXT("I%s"), *Class->GetName()) : ClassCPPName;
 
-	GeneratedPackageCPP.Logf(TEXT("\tvoid %s::StaticRegisterNatives%s()\r\n"),NameLookupCPP.GetNameCPP(Class),NameLookupCPP.GetNameCPP(Class));
-	GeneratedPackageCPP.Logf(TEXT("\t{\r\n"));
+	Out.Logf(TEXT("\tvoid %s::StaticRegisterNatives%s()\r\n"), ClassCPPName, ClassCPPName);
+	Out.Log(TEXT("\t{\r\n"));
 
 	{
-		TArray<UFunction*> FunctionsToExport;
-		for (auto* Function : TFieldRange<UFunction>(Class, EFieldIteratorFlags::ExcludeSuper))
+		TArray<TTuple<UFunction*, FString>> AnsiNamedFunctionsToExport;
+		TArray<TTuple<UFunction*, FString>> TCharNamedFunctionsToExport;
+		for (UFunction* Function : TFieldRange<UFunction>(Class, EFieldIteratorFlags::ExcludeSuper))
 		{
-			if( (Function->FunctionFlags & (FUNC_Native | FUNC_NetRequest)) == FUNC_Native )
+			if ((Function->FunctionFlags & (FUNC_Native | FUNC_NetRequest)) == FUNC_Native)
 			{
-				FunctionsToExport.Add(Function);
+				FString OverriddenName = FNativeClassHeaderGenerator::GetOverriddenNameForLiteral(Function);
+				if (OverriddenName.StartsWith(TEXT("TEXT(")))
+				{
+					TCharNamedFunctionsToExport.Emplace(Function, MoveTemp(OverriddenName));
+				}
+				else
+				{
+					AnsiNamedFunctionsToExport.Emplace(Function, MoveTemp(OverriddenName));
+				}
 			}
 		}
 
-		FunctionsToExport.Sort();
+		Algo::SortBy(AnsiNamedFunctionsToExport,  [](const TTuple<UFunction*, FString>& Pair){ return Pair.Get<0>()->GetFName(); });
+		Algo::SortBy(TCharNamedFunctionsToExport, [](const TTuple<UFunction*, FString>& Pair){ return Pair.Get<0>()->GetFName(); });
 
-		for (auto Func : FunctionsToExport)
+		if (AnsiNamedFunctionsToExport.Num() > 0 || TCharNamedFunctionsToExport.Num() > 0)
 		{
-			GeneratedPackageCPP.Logf(TEXT("\t\tFNativeFunctionRegistrar::RegisterFunction(%s::StaticClass(), %s,(Native)&%s::exec%s);\r\n"),
-				NameLookupCPP.GetNameCPP(Class),
-				*FNativeClassHeaderGenerator::GetOverriddenNameForLiteral(Func),
-				Class->HasAnyClassFlags(CLASS_Interface) ? *FString::Printf(TEXT("I%s"), *Class->GetName()) : NameLookupCPP.GetNameCPP(Class),
-				*Func->GetName()
-			);
+			Out.Logf(TEXT("\t\tUClass* Class = %s::StaticClass();\r\n"), ClassCPPName);
+		}
+
+		if (AnsiNamedFunctionsToExport.Num())
+		{
+			Out.Log(TEXT("\t\tstatic const TNameNativePtrPair<ANSICHAR> AnsiFuncs[] = {\r\n"));
+
+			for (const TTuple<UFunction*, FString>& Func : AnsiNamedFunctionsToExport)
+			{
+				Out.Logf(
+					TEXT("\t\t\t{ %s, (Native)&%s::exec%s },\r\n"),
+					*Func.Get<1>(),
+					*TypeName,
+					*Func.Get<0>()->GetName()
+				);
+			}
+
+			Out.Log(TEXT("\t\t};\r\n"));
+			Out.Logf(TEXT("\t\tFNativeFunctionRegistrar::RegisterFunctions(Class, AnsiFuncs, %d);\r\n"), AnsiNamedFunctionsToExport.Num());
+		}
+
+		if (TCharNamedFunctionsToExport.Num())
+		{
+			Out.Log(TEXT("\t\tstatic const TNameNativePtrPair<TCHAR> TCharFuncs[] = {\r\n"));
+
+			for (const TTuple<UFunction*, FString>& Func : TCharNamedFunctionsToExport)
+			{
+				Out.Logf(
+					TEXT("\t\t\t{ %s, (Native)&%s::exec%s },\r\n"),
+					*Func.Get<1>(),
+					*TypeName,
+					*Func.Get<0>()->GetName()
+				);
+			}
+
+			Out.Log(TEXT("\t\t};\r\n"));
+			Out.Logf(TEXT("\t\tFNativeFunctionRegistrar::RegisterFunctions(Class, TCharFuncs, %d);\r\n"), TCharNamedFunctionsToExport.Num());
 		}
 	}
 
-	for (auto* Struct : TFieldRange<UScriptStruct>(Class, EFieldIteratorFlags::ExcludeSuper))
+	for (UScriptStruct* Struct : TFieldRange<UScriptStruct>(Class, EFieldIteratorFlags::ExcludeSuper))
 	{
 		if (Struct->StructFlags & STRUCT_Native)
 		{
-			GeneratedPackageCPP.Logf( TEXT("\t\tUScriptStruct::DeferCppStructOps(FName(TEXT(\"%s\")),new UScriptStruct::TCppStructOps<%s%s>);\r\n"), *Struct->GetName(), Struct->GetPrefixCPP(), *Struct->GetName() );
+			Out.Logf( TEXT("\t\tUScriptStruct::DeferCppStructOps(FName(TEXT(\"%s\")),new UScriptStruct::TCppStructOps<%s%s>);\r\n"), *Struct->GetName(), Struct->GetPrefixCPP(), *Struct->GetName() );
 		}
 	}
 
-	GeneratedPackageCPP.Logf(TEXT("\t}\r\n"));
+	Out.Logf(TEXT("\t}\r\n"));
 }
 
-void FNativeClassHeaderGenerator::ExportInterfaceCallFunctions(const TArray<UFunction*>& InCallbackFunctions, UClass* Class, FUHTStringBuilder& HeaderOutput)
+void FNativeClassHeaderGenerator::ExportInterfaceCallFunctions(FOutputDevice& OutCpp, FUHTStringBuilder& Out, const TArray<UFunction*>& CallbackFunctions, const TCHAR* ClassName)
 {
-	TArray<UFunction*> CallbackFunctions = InCallbackFunctions;
-	CallbackFunctions.Sort();
+	FString APIString = GetAPIString();
 
-	for (auto FuncIt = CallbackFunctions.CreateConstIterator(); FuncIt; ++FuncIt)
+	for (UFunction* Function : CallbackFunctions)
 	{
-		UFunction* Function     = *FuncIt;
-		FString    FunctionName = Function->GetName();
-		FString    ClassName    = Class->GetName();
+		FString FunctionName = Function->GetName();
 
 		auto* CompilerInfo = FFunctionData::FindForFunction(Function);
 
@@ -1865,15 +1797,15 @@ void FNativeClassHeaderGenerator::ExportInterfaceCallFunctions(const TArray<UFun
 		const TCHAR* ConstQualifier = FunctionData.FunctionReference->HasAllFunctionFlags(FUNC_Const) ? TEXT("const ") : TEXT("");
 		FString ExtraParam = FString::Printf(TEXT("%sUObject* O"), ConstQualifier);
 
-		ExportNativeFunctionHeader(FunctionData, HeaderOutput, EExportFunctionType::Interface, EExportFunctionHeaderStyle::Declaration, *ExtraParam);
-		HeaderOutput.Logf( TEXT(";") LINE_TERMINATOR );
+		ExportNativeFunctionHeader(Out, ForwardDeclarations, FunctionData, EExportFunctionType::Interface, EExportFunctionHeaderStyle::Declaration, *ExtraParam, *APIString);
+		Out.Logf( TEXT(";") LINE_TERMINATOR );
 
 
-		ExportNativeFunctionHeader(FunctionData, GeneratedPackageCPP, EExportFunctionType::Interface, EExportFunctionHeaderStyle::Definition, *ExtraParam);
-		GeneratedPackageCPP.Logf( LINE_TERMINATOR TEXT("\t{") LINE_TERMINATOR );
+		ExportNativeFunctionHeader(OutCpp, ForwardDeclarations, FunctionData, EExportFunctionType::Interface, EExportFunctionHeaderStyle::Definition, *ExtraParam, *APIString);
+		OutCpp.Logf( LINE_TERMINATOR TEXT("\t{") LINE_TERMINATOR );
 
-		GeneratedPackageCPP.Logf(TEXT("\t\tcheck(O != NULL);") LINE_TERMINATOR);
-		GeneratedPackageCPP.Logf(TEXT("\t\tcheck(O->GetClass()->ImplementsInterface(U%s::StaticClass()));") LINE_TERMINATOR, *ClassName);
+		OutCpp.Logf(TEXT("\t\tcheck(O != NULL);") LINE_TERMINATOR);
+		OutCpp.Logf(TEXT("\t\tcheck(O->GetClass()->ImplementsInterface(U%s::StaticClass()));") LINE_TERMINATOR, ClassName);
 
 		auto Parameters = GetFunctionParmsAndReturn(FunctionData.FunctionReference);
 
@@ -1882,23 +1814,23 @@ void FNativeClassHeaderGenerator::ExportInterfaceCallFunctions(const TArray<UFun
 		if (bHasParms)
 		{
 			FString EventParmStructName = GetEventStructParamsName(Function->GetOuter(), *FunctionName);
-			GeneratedPackageCPP.Logf(TEXT("\t\t%s Parms;") LINE_TERMINATOR, *EventParmStructName);
+			OutCpp.Logf(TEXT("\t\t%s Parms;") LINE_TERMINATOR, *EventParmStructName);
 		}
 
-		GeneratedPackageCPP.Logf(TEXT("\t\tUFunction* const Func = O->FindFunction(%s_%s);") LINE_TERMINATOR, *API, *FunctionName);
-		GeneratedPackageCPP.Logf(TEXT("\t\tif (Func)") LINE_TERMINATOR);
-		GeneratedPackageCPP.Logf(TEXT("\t\t{") LINE_TERMINATOR);
+		OutCpp.Logf(TEXT("\t\tUFunction* const Func = O->FindFunction(%s_%s);") LINE_TERMINATOR, *API, *FunctionName);
+		OutCpp.Logf(TEXT("\t\tif (Func)") LINE_TERMINATOR);
+		OutCpp.Logf(TEXT("\t\t{") LINE_TERMINATOR);
 
 		// code to populate Parms struct
 		for (auto It = Parameters.Parms.CreateConstIterator(); It; ++It)
 		{
 			UProperty* Param = *It;
 
-			GeneratedPackageCPP.Logf(TEXT("\t\t\tParms.%s=%s;") LINE_TERMINATOR, *Param->GetName(), *Param->GetName());
+			OutCpp.Logf(TEXT("\t\t\tParms.%s=%s;") LINE_TERMINATOR, *Param->GetName(), *Param->GetName());
 		}
 
-		const FString ObjectRef = FunctionData.FunctionReference->HasAllFunctionFlags(FUNC_Const) ? FString::Printf(TEXT("const_cast<UObject*>(O)"), *ClassName) : TEXT("O");
-		GeneratedPackageCPP.Logf(TEXT("\t\t\t%s->ProcessEvent(Func, %s);") LINE_TERMINATOR, *ObjectRef, bHasParms ? TEXT("&Parms") : TEXT("NULL"));
+		const FString ObjectRef = FunctionData.FunctionReference->HasAllFunctionFlags(FUNC_Const) ? FString::Printf(TEXT("const_cast<UObject*>(O)")) : TEXT("O");
+		OutCpp.Logf(TEXT("\t\t\t%s->ProcessEvent(Func, %s);") LINE_TERMINATOR, *ObjectRef, bHasParms ? TEXT("&Parms") : TEXT("NULL"));
 
 		for (auto It = Parameters.Parms.CreateConstIterator(); It; ++It)
 		{
@@ -1906,26 +1838,26 @@ void FNativeClassHeaderGenerator::ExportInterfaceCallFunctions(const TArray<UFun
 
 			if( Param->HasAllPropertyFlags(CPF_OutParm) && !Param->HasAnyPropertyFlags(CPF_ConstParm|CPF_ReturnParm))
 			{
-				GeneratedPackageCPP.Logf(TEXT("\t\t\t%s=Parms.%s;") LINE_TERMINATOR, *Param->GetName(), *Param->GetName());
+				OutCpp.Logf(TEXT("\t\t\t%s=Parms.%s;") LINE_TERMINATOR, *Param->GetName(), *Param->GetName());
 			}
 		}
 
-		GeneratedPackageCPP += TEXT("\t\t}") LINE_TERMINATOR;
+		OutCpp.Log(TEXT("\t\t}") LINE_TERMINATOR);
 		
 
 		// else clause to call back into native if it's a BlueprintNativeEvent
 		if (Function->FunctionFlags & FUNC_Native)
 		{
-			GeneratedPackageCPP.Logf(TEXT("\t\telse if (auto I = (%sI%s*)(O->GetNativeInterfaceAddress(U%s::StaticClass())))") LINE_TERMINATOR, ConstQualifier, *ClassName, *ClassName);
-			GeneratedPackageCPP += TEXT("\t\t{") LINE_TERMINATOR;
+			OutCpp.Logf(TEXT("\t\telse if (auto I = (%sI%s*)(O->GetNativeInterfaceAddress(U%s::StaticClass())))") LINE_TERMINATOR, ConstQualifier, ClassName, ClassName);
+			OutCpp.Log(TEXT("\t\t{") LINE_TERMINATOR);
 
-			GeneratedPackageCPP += TEXT("\t\t\t");
+			OutCpp.Log(TEXT("\t\t\t"));
 			if (Parameters.Return)
 			{
-				GeneratedPackageCPP += TEXT("Parms.ReturnValue = ");
+				OutCpp.Log(TEXT("Parms.ReturnValue = "));
 			}
 
-			GeneratedPackageCPP.Logf(TEXT("I->%s_Implementation("), *FunctionName);
+			OutCpp.Logf(TEXT("I->%s_Implementation("), *FunctionName);
 
 			bool First = true;
 			for (auto It = Parameters.Parms.CreateConstIterator(); It; ++It)
@@ -1934,24 +1866,24 @@ void FNativeClassHeaderGenerator::ExportInterfaceCallFunctions(const TArray<UFun
 
 				if (!First)
 				{
-					GeneratedPackageCPP.Logf(TEXT(","));
+					OutCpp.Logf(TEXT(","));
 				}
 				First = false;
 
-				GeneratedPackageCPP.Logf(TEXT("%s"), *Param->GetName());
+				OutCpp.Logf(TEXT("%s"), *Param->GetName());
 			}
 
-			GeneratedPackageCPP.Logf(TEXT(");") LINE_TERMINATOR);
+			OutCpp.Logf(TEXT(");") LINE_TERMINATOR);
 
-			GeneratedPackageCPP.Logf(TEXT("\t\t}") LINE_TERMINATOR);
+			OutCpp.Logf(TEXT("\t\t}") LINE_TERMINATOR);
 		}
 
 		if (Parameters.Return)
 		{
-			GeneratedPackageCPP.Logf(TEXT("\t\treturn Parms.ReturnValue;") LINE_TERMINATOR);
+			OutCpp.Logf(TEXT("\t\treturn Parms.ReturnValue;") LINE_TERMINATOR);
 		}
 
-		GeneratedPackageCPP.Logf(TEXT("\t}") LINE_TERMINATOR);
+		OutCpp.Logf(TEXT("\t}") LINE_TERMINATOR);
 	}
 }
 
@@ -1978,9 +1910,8 @@ FString GetGeneratedMacroDeprecationWarning(const TCHAR* MacroName)
  */
 FString GetPreservedAccessSpecifierString(FClass* Class)
 {
-	FString PreservedAccessSpecifier(FString::Printf(TEXT("static_assert(false, \"Unknown access specifier for GENERATED_BODY() macro in class %s.\");"), *GetNameSafe(Class)));
-	FClassMetaData* Data = GScriptHelper.FindClassData(Class);
-	if (Data)
+	FString PreservedAccessSpecifier;
+	if (FClassMetaData* Data = GScriptHelper.FindClassData(Class))
 	{
 		switch (Data->GeneratedBodyMacroAccessSpecifier)
 		{
@@ -1994,6 +1925,7 @@ FString GetPreservedAccessSpecifierString(FClass* Class)
 			PreservedAccessSpecifier = "public:";
 			break;
 		case EAccessSpecifier::ACCESS_NotAnAccessSpecifier :
+			PreservedAccessSpecifier = FString::Printf(TEXT("static_assert(false, \"Unknown access specifier for GENERATED_BODY() macro in class %s.\");"), *GetNameSafe(Class));
 			break;
 		}
 	}
@@ -2004,16 +1936,13 @@ FString GetPreservedAccessSpecifierString(FClass* Class)
 /**
  * Exports interface body macros.
  */
-void ExportInterfaceBodyMacros(FUHTStringBuilder& Out, FUnrealSourceFile& SourceFile, int32 ClassGeneratedBodyLine, FClass* Class, const FString& StdCtorCall, const FString& EnhCtorCall, const FString& UInterfaceBoilerplate)
+void ExportInterfaceBodyMacros(FOutputDevice& Out, const FUnrealSourceFile& SourceFile, int32 ClassGeneratedBodyLine, FClass* Class, const FString& StdCtorCall, const FString& EnhCtorCall)
 {
-	Out.Log(TEXT("#undef GENERATED_UINTERFACE_BODY_COMMON\r\n"));
-	Out.Log(Macroize(TEXT("GENERATED_UINTERFACE_BODY_COMMON()"), *UInterfaceBoilerplate));
-
-	auto DeprecationWarning = GetGeneratedMacroDeprecationWarning(TEXT("GENERATED_UINTERFACE_BODY"));
+	FString DeprecationWarning = GetGeneratedMacroDeprecationWarning(TEXT("GENERATED_UINTERFACE_BODY"));
 
-	auto DeprecationPushString = TEXT("PRAGMA_DISABLE_DEPRECATION_WARNINGS") LINE_TERMINATOR;
-	auto DeprecationPopString = TEXT("PRAGMA_ENABLE_DEPRECATION_WARNINGS") LINE_TERMINATOR;
-	auto Offset = TEXT("\t");
+	const TCHAR* DeprecationPushString = TEXT("PRAGMA_DISABLE_DEPRECATION_WARNINGS") LINE_TERMINATOR;
+	const TCHAR* DeprecationPopString = TEXT("PRAGMA_ENABLE_DEPRECATION_WARNINGS") LINE_TERMINATOR;
+	const TCHAR* Offset = TEXT("\t");
 
 	Out.Log(Macroize(*SourceFile.GetGeneratedBodyMacroName(ClassGeneratedBodyLine, true), *(FString() +
 		Offset + DeprecationWarning +
@@ -2030,13 +1959,7 @@ void ExportInterfaceBodyMacros(FUHTStringBuilder& Out, FUnrealSourceFile& Source
 		Offset + DeprecationPopString)));
 }
 
-void FNativeClassHeaderGenerator::ExportSourceFileHeader(FClasses& AllClasses, FUnrealSourceFile* SourceFile)
-{
-	TSet<const FUnrealSourceFile*>	VisitedSet;
-	ExportSourceFileHeaderRecursive(AllClasses, SourceFile, VisitedSet, false);
-}
-
-void WriteMacro(FUHTStringBuilder& Output, const FString& MacroName, const FString& MacroContent)
+void WriteMacro(FOutputDevice& Output, const FString& MacroName, const FString& MacroContent)
 {
 	Output.Log(*Macroize(*MacroName, *MacroContent));
 }
@@ -2047,16 +1970,16 @@ void WriteMacro(FUHTStringBuilder& Output, const FString& MacroName, const FStri
  * @param Out Output device.
  * @param SourceFile Source file.
  */
-void ExportAutoIncludes(FStringOutputDevice& Out, const FUnrealSourceFile& SourceFile)
+void ExportAutoIncludes(FOutputDevice& Out, const FUnrealSourceFile& SourceFile)
 {
-	for (const auto& Include : SourceFile.GetIncludes())
+	for (const FHeaderProvider& Include : SourceFile.GetIncludes())
 	{
 		if (!Include.IsAutoInclude())
 		{
 			continue;
 		}
 
-		const auto* AutoIncludedSourceFile = Include.GetResolved();
+		const FUnrealSourceFile* AutoIncludedSourceFile = Include.GetResolved();
 
 		if (AutoIncludedSourceFile == nullptr)
 		{
@@ -2100,331 +2023,332 @@ static FString PrivatePropertiesOffsetGetters(const UStruct* Struct, const FStri
 	return Result;
 }
 
-void FNativeClassHeaderGenerator::ExportClassesFromSourceFileInner(FUnrealSourceFile& SourceFile)
+void FNativeClassHeaderGenerator::ExportClassFromSourceFileInner(
+	FOutputDevice&           OutGeneratedHeaderText,
+	FOutputDevice&           OutCpp,
+	FOutputDevice&           OutDeclarations,
+	FClass*                  Class,
+	const FUnrealSourceFile& SourceFile
+)
 {
-	NameLookupCPP.SetCurrentSourceFile(&SourceFile);
-	UHTMakefile.AddToHeaderOrder(&SourceFile);
-	TArray<UEnum*>				Enums;
-	TArray<UScriptStruct*>		Structs;
-	TArray<UDelegateFunction*>	DelegateFunctions;
-
-	GeneratedHeaderText.Logf(
-		TEXT("#ifdef %s")																	LINE_TERMINATOR
-		TEXT("#error \"%s.generated.h already included, missing '#pragma once' in %s.h\"")	LINE_TERMINATOR
-		TEXT("#endif")																		LINE_TERMINATOR
-		TEXT("#define %s")																	LINE_TERMINATOR
-		LINE_TERMINATOR,
-		*SourceFile.GetFileDefineName(), *SourceFile.GetStrippedFilename(), *SourceFile.GetStrippedFilename(), *SourceFile.GetFileDefineName());
+	FUHTStringBuilder StandardUObjectConstructorsMacroCall;
+	FUHTStringBuilder EnhancedUObjectConstructorsMacroCall;
 
-	ExportAutoIncludes(GeneratedHeaderText, SourceFile);
-
-	// get the lists of fields that belong to this class that should be exported
-	SourceFile.GetScope()->SplitTypesIntoArrays(Enums, Structs, DelegateFunctions);
-
-	// export enum declarations
-	ExportEnums(Enums);
-
-	// export delegate definitions
-	ExportDelegateDefinitions(SourceFile, DelegateFunctions, false);
+	FClassMetaData* ClassData = GScriptHelper.FindClassData(Class);
+	check(ClassData);
 
-	// Export enums declared in non-UClass headers.
-	ExportGeneratedEnumsInitCode(SourceFile, Enums);
+	// C++ -> VM stubs (native function execs)
+	FUHTStringBuilder ClassMacroCalls;
+	FUHTStringBuilder ClassNoPureDeclsMacroCalls;
+	ExportNativeFunctions(OutGeneratedHeaderText, ClassMacroCalls, ClassNoPureDeclsMacroCalls, SourceFile, Class, ClassData);
 
-	// export boilerplate macros for structs
-	ExportGeneratedStructBodyMacros(SourceFile, Structs);
+	// Get Callback functions
+	TArray<UFunction*> CallbackFunctions;
+	{
+		for (UFunction* Function : TFieldRange<UFunction>(Class, EFieldIteratorFlags::ExcludeSuper))
+		{
+			if ((Function->FunctionFlags & FUNC_Event) && Function->GetSuperFunction() == nullptr)
+			{
+				CallbackFunctions.Add(Function);
+			}
+		}
 
-	// export delegate wrapper function implementations
-	ExportDelegateDefinitions(SourceFile, DelegateFunctions, true);
+		CallbackFunctions.Sort();
+	}
 
-	for (auto* RawUClass : SourceFile.GetDefinedClasses())
+	FUHTStringBuilder PrologMacroCalls;
+	if (CallbackFunctions.Num() != 0)
 	{
-		PrologMacroCalls.Reset();
-		InClassMacroCalls.Reset();
-		InClassNoPureDeclsMacroCalls.Reset();
-		StandardUObjectConstructorsMacroCall.Reset();
-		EnhancedUObjectConstructorsMacroCall.Reset();
+		FUHTStringBuilder UClassMacroContent;
 
-		auto* Class = (FClass*)RawUClass;
-		if (Class->ClassFlags & CLASS_Intrinsic)
+		// export parameters structs for all events and delegates
+		for (UFunction* Function : CallbackFunctions)
 		{
-			continue;
+			ExportEventParm(UClassMacroContent, ForwardDeclarations, Function, /*Indent=*/ 1, /*bOutputConstructor=*/ true, EExportingState::Normal);
 		}
 
-		auto* ClassData = GScriptHelper.FindClassData(Class);
-		check(ClassData);
-
-		// C++ -> VM stubs (native function execs)
-		ExportNativeFunctions(SourceFile, Class, ClassData);
+		FString MacroName = SourceFile.GetGeneratedMacroName(ClassData, TEXT("_EVENT_PARMS"));
+		WriteMacro(OutGeneratedHeaderText, MacroName, UClassMacroContent);
+		PrologMacroCalls.Logf(TEXT("\t%s\r\n"), *MacroName);
 
 		// VM -> C++ proxies (events and delegates).
-		TArray<UFunction*> CallbackFunctions = ExportCallbackFunctions(SourceFile, Class, ClassData);
-
-		// Class definition.
-		ExportNatives(Class);
+		FUHTStringBuilder NullOutput;
+		FOutputDevice& CallbackOut = Class->HasAnyClassFlags(CLASS_NoExport) ? NullOutput : OutCpp;
+		FString CallbackWrappersMacroName = SourceFile.GetGeneratedMacroName(ClassData, TEXT("_CALLBACK_WRAPPERS"));
+		ExportCallbackFunctions(
+			OutGeneratedHeaderText,
+			CallbackOut,
+			ReferencedNames,
+			ForwardDeclarations,
+			CallbackFunctions,
+			*CallbackWrappersMacroName,
+			(Class->ClassFlags & CLASS_Interface) ? EExportCallbackType::Interface : EExportCallbackType::Class,
+			*API,
+			*GetAPIString()
+		);
 
-		FUHTStringBuilder FriendText;
-		ExportNativeGeneratedInitCode(SourceFile, Class, FriendText);
+		ClassMacroCalls.Logf(TEXT("\t%s\r\n"), *CallbackWrappersMacroName);
+		ClassNoPureDeclsMacroCalls.Logf(TEXT("\t%s\r\n"), *CallbackWrappersMacroName);
+	}
 
-		FClass* SuperClass = Class->GetSuperClass();
+	// Class definition.
+	if (!Class->HasAnyClassFlags(CLASS_NoExport))
+	{
+		ExportNatives(OutCpp, Class);
+	}
 
-		// the name for the C++ version of the UClass
-		const TCHAR* ClassCPPName = NameLookupCPP.GetNameCPP(Class);
-		const TCHAR* SuperClassCPPName = (SuperClass != nullptr) ? NameLookupCPP.GetNameCPP(SuperClass) : nullptr;
+	FUHTStringBuilder FriendText;
+	ExportNativeGeneratedInitCode(OutCpp, OutDeclarations, SourceFile, Class, FriendText);
 
-		FString PPOMacroName;
+	FClass* SuperClass = Class->GetSuperClass();
 
-		if (Class->HasAnyClassFlags(CLASS_Interface))
-		{
-			{
-				FUHTStringBuilder UInterfaceBoilerplate;
-
-				// Export the class's native function registration.
-				UInterfaceBoilerplate.Logf(TEXT("\tprivate:\r\n"));
-				UInterfaceBoilerplate.Logf(TEXT("\tstatic void StaticRegisterNatives%s();\r\n"), NameLookupCPP.GetNameCPP(Class));
-				UInterfaceBoilerplate.Log(*FriendText);
-				FriendText.Reset();
-				UInterfaceBoilerplate.Logf(TEXT("public:\r\n"));
-
-				// Build the DECLARE_CLASS line
-				FString APIArg(API);
-				if (!Class->HasAnyClassFlags(CLASS_MinimalAPI))
-				{
-					APIArg = TEXT("NO");
-				}
+	// the name for the C++ version of the UClass
+	const TCHAR* ClassCPPName = NameLookupCPP.GetNameCPP(Class);
+	const TCHAR* SuperClassCPPName = (SuperClass != nullptr) ? NameLookupCPP.GetNameCPP(SuperClass) : nullptr;
 
-				const bool bCastedClass = Class->HasAnyCastFlag(CASTCLASS_AllFlags) && SuperClass && Class->ClassCastFlags != SuperClass->ClassCastFlags;
-				UInterfaceBoilerplate.Logf(TEXT("\tDECLARE_CLASS(%s, %s, COMPILED_IN_FLAGS(CLASS_Abstract%s), %s, TEXT(\"%s\"), %s_API)\r\n"),
-					ClassCPPName,
-					SuperClassCPPName,
-					*GetClassFlagExportText(Class),
-					bCastedClass ? *FString::Printf(TEXT("CASTCLASS_%s"), ClassCPPName) : TEXT("0"),
-					*FClass::GetTypePackageName(Class),
-					*APIArg);
+	FString APIArg = API;
+	if (!Class->HasAnyClassFlags(CLASS_MinimalAPI))
+	{
+		APIArg = TEXT("NO");
+	}
 
-				UInterfaceBoilerplate.Logf(TEXT("\tDECLARE_SERIALIZER(%s)\r\n"), ClassCPPName);
-				UInterfaceBoilerplate.Log(TEXT("\tenum {IsIntrinsic=COMPILED_IN_INTRINSIC};\r\n"));
+	FString PPOMacroName;
 
-				if (Class->ClassWithin != Class->GetSuperClass()->ClassWithin)
-				{
-					UInterfaceBoilerplate.Logf(TEXT("\tDECLARE_WITHIN(%s)\r\n"), NameLookupCPP.GetNameCPP(Class->GetClassWithin()));
-				}
+	if (Class->HasAnyClassFlags(CLASS_Interface))
+	{
+		{
+			FUHTStringBuilder UInterfaceBoilerplate;
 
-				ExportConstructorsMacros(SourceFile.GetGeneratedMacroName(ClassData), Class);
+			// Export the class's native function registration.
+			UInterfaceBoilerplate.Logf(TEXT("\tprivate:\r\n"));
+			UInterfaceBoilerplate.Logf(TEXT("\tstatic void StaticRegisterNatives%s();\r\n"), ClassCPPName);
+			UInterfaceBoilerplate.Log(*FriendText);
+			FriendText.Reset();
+			UInterfaceBoilerplate.Logf(TEXT("public:\r\n"));
 
-				ExportInterfaceBodyMacros(GeneratedHeaderText, SourceFile, ClassData->GetGeneratedBodyLine(), Class, StandardUObjectConstructorsMacroCall, EnhancedUObjectConstructorsMacroCall, UInterfaceBoilerplate);
-			}
+			const bool bCastedClass = Class->HasAnyCastFlag(CASTCLASS_AllFlags) && SuperClass && Class->ClassCastFlags != SuperClass->ClassCastFlags;
+			UInterfaceBoilerplate.Logf(TEXT("\tDECLARE_CLASS(%s, %s, COMPILED_IN_FLAGS(CLASS_Abstract%s), %s, TEXT(\"%s\"), %s_API)\r\n"),
+				ClassCPPName,
+				SuperClassCPPName,
+				*GetClassFlagExportText(Class),
+				bCastedClass ? *FString::Printf(TEXT("CASTCLASS_%s"), ClassCPPName) : TEXT("0"),
+				*FClass::GetTypePackageName(Class),
+				*APIArg);
 
-			// =============================================
-			// Export the pure interface version of the class
+			UInterfaceBoilerplate.Logf(TEXT("\tDECLARE_SERIALIZER(%s)\r\n"), ClassCPPName);
+			UInterfaceBoilerplate.Log(TEXT("\tenum {IsIntrinsic=COMPILED_IN_INTRINSIC};\r\n"));
 
-			// the name of the pure interface class
-			FString InterfaceCPPName = FString::Printf(TEXT("I%s"), *Class->GetName());
-			FString SuperInterfaceCPPName;
-			if (SuperClass != NULL)
+			if (Class->ClassWithin != Class->GetSuperClass()->ClassWithin)
 			{
-				SuperInterfaceCPPName = FString::Printf(TEXT("I%s"), *SuperClass->GetName());
+				UInterfaceBoilerplate.Logf(TEXT("\tDECLARE_WITHIN(%s)\r\n"), NameLookupCPP.GetNameCPP(Class->GetClassWithin()));
 			}
 
-			// Thunk functions
-			FUHTStringBuilder InterfaceBoilerplate;
+			ExportConstructorsMacros(OutGeneratedHeaderText, OutCpp, StandardUObjectConstructorsMacroCall, EnhancedUObjectConstructorsMacroCall, SourceFile.GetGeneratedMacroName(ClassData), Class, *APIArg);
 
-			InterfaceBoilerplate.Logf(TEXT("protected:\r\n\tvirtual ~%s() {}\r\npublic:\r\n"), *InterfaceCPPName);
-			InterfaceBoilerplate.Logf(TEXT("\ttypedef %s UClassType;\r\n"), ClassCPPName);
-
-			ExportInterfaceCallFunctions(CallbackFunctions, Class, InterfaceBoilerplate);
-
-			// we'll need a way to get to the UObject portion of a native interface, so that we can safely pass native interfaces
-			// to script VM functions
-			if (SuperClass->IsChildOf(UInterface::StaticClass()))
-			{
-				InterfaceBoilerplate.Logf(TEXT("\tvirtual UObject* _getUObject() const = 0;\r\n"));
-			}
+			OutGeneratedHeaderText.Log(TEXT("#undef GENERATED_UINTERFACE_BODY_COMMON\r\n"));
+			OutGeneratedHeaderText.Log(Macroize(TEXT("GENERATED_UINTERFACE_BODY_COMMON()"), *UInterfaceBoilerplate));
 
-			// Replication, add in the declaration for GetLifetimeReplicatedProps() automatically if there are any net flagged properties
-			bool bNeedsRep = false;
-			auto ClassName = FString(NameLookupCPP.GetNameCPP(Class));
-			for (TFieldIterator<UProperty> It(Class, EFieldIteratorFlags::ExcludeSuper); It; ++It)
-			{
-				if ((It->PropertyFlags & CPF_Net) != 0)
-				{
-					bNeedsRep = true;
-					break;
-				}
-			}
+			ExportInterfaceBodyMacros(OutGeneratedHeaderText, SourceFile, ClassData->GetGeneratedBodyLine(), Class, StandardUObjectConstructorsMacroCall, EnhancedUObjectConstructorsMacroCall);
+		}
 
-			auto ClassRange = ClassDefinitionRange();
-			if (ClassDefinitionRanges.Contains(Class))
-			{
-				ClassRange = ClassDefinitionRanges[Class];
-				ClassRange.Validate();
-			}
+		// =============================================
+		// Export the pure interface version of the class
 
-			auto ClassStart = ClassRange.Start;
-			auto ClassEnd = ClassRange.End;
-			auto ClassDefinition = FString(ClassEnd - ClassStart, ClassStart);
+		// the name of the pure interface class
+		FString InterfaceCPPName = FString::Printf(TEXT("I%s"), *Class->GetName());
+		FString SuperInterfaceCPPName;
+		if (SuperClass != NULL)
+		{
+			SuperInterfaceCPPName = FString::Printf(TEXT("I%s"), *SuperClass->GetName());
+		}
 
-			bool bHasGetLifetimeReplicatedProps = HasIdentifierExactMatch(ClassDefinition, TEXT("GetLifetimeReplicatedProps"));
+		// Thunk functions
+		FUHTStringBuilder InterfaceBoilerplate;
 
-			if (bNeedsRep && !bHasGetLifetimeReplicatedProps)
-			{
-				if (SourceFile.GetGeneratedCodeVersionForStruct(Class) == EGeneratedCodeVersion::V1)
-				{
-					InterfaceBoilerplate.Logf(TEXT("\void GetLifetimeReplicatedProps(TArray<FLifetimeProperty>& OutLifetimeProps) const;\r\n"));
-				}
-				else
-				{
-					FError::Throwf(TEXT("Class %s has Net flagged properties and should declare member function: void GetLifetimeReplicatedProps(TArray<FLifetimeProperty>& OutLifetimeProps) const override"), ClassCPPName);
-				}
-			}
+		InterfaceBoilerplate.Logf(TEXT("protected:\r\n\tvirtual ~%s() {}\r\npublic:\r\n"), *InterfaceCPPName);
+		InterfaceBoilerplate.Logf(TEXT("\ttypedef %s UClassType;\r\n"), ClassCPPName);
 
-			auto NoPureDeclsMacroName = SourceFile.GetGeneratedMacroName(ClassData, TEXT("_INCLASS_IINTERFACE_NO_PURE_DECLS"));
-			WriteMacro(GeneratedHeaderText, NoPureDeclsMacroName, InterfaceBoilerplate);
-			InClassNoPureDeclsMacroCalls.Logf(TEXT("\t%s\r\n"), *NoPureDeclsMacroName);
+		ExportInterfaceCallFunctions(OutCpp, InterfaceBoilerplate, CallbackFunctions, *Class->GetName());
 
-			auto MacroName = SourceFile.GetGeneratedMacroName(ClassData, TEXT("_INCLASS_IINTERFACE"));
-			WriteMacro(GeneratedHeaderText, MacroName, InterfaceBoilerplate);
-			InClassMacroCalls.Logf(TEXT("\t%s\r\n"), *MacroName);
-		}
-		else
+		// we'll need a way to get to the UObject portion of a native interface, so that we can safely pass native interfaces
+		// to script VM functions
+		if (SuperClass->IsChildOf(UInterface::StaticClass()))
 		{
-			FUHTStringBuilder ClassBoilerplate;
+			InterfaceBoilerplate.Logf(TEXT("\tvirtual UObject* _getUObject() const = 0;\r\n"));
+		}
 
-			// Export the class's native function registration.
-			ClassBoilerplate.Logf(TEXT("\tprivate:\r\n"));
-			ClassBoilerplate.Logf(TEXT("\tstatic void StaticRegisterNatives%s();\r\n"), NameLookupCPP.GetNameCPP(Class));
-			ClassBoilerplate.Log(*FriendText);
-			ClassBoilerplate.Logf(TEXT("\tpublic:\r\n"));
-
-			// Build the DECLARE_CLASS line
-			FString APIArg(API);
-			if (!Class->HasAnyClassFlags(CLASS_MinimalAPI))
+		// Replication, add in the declaration for GetLifetimeReplicatedProps() automatically if there are any net flagged properties
+		bool bNeedsRep = false;
+		for (TFieldIterator<UProperty> It(Class, EFieldIteratorFlags::ExcludeSuper); It; ++It)
+		{
+			if ((It->PropertyFlags & CPF_Net) != 0)
 			{
-				APIArg = TEXT("NO");
+				bNeedsRep = true;
+				break;
 			}
-			const bool bCastedClass = Class->HasAnyCastFlag(CASTCLASS_AllFlags) && SuperClass && Class->ClassCastFlags != SuperClass->ClassCastFlags;
-			ClassBoilerplate.Logf(TEXT("\tDECLARE_CLASS(%s, %s, COMPILED_IN_FLAGS(%s%s), %s, TEXT(\"%s\"), %s_API)\r\n"),
-				ClassCPPName,
-				SuperClassCPPName ? SuperClassCPPName : TEXT("None"),
-				Class->HasAnyClassFlags(CLASS_Abstract) ? TEXT("CLASS_Abstract") : TEXT("0"),
-				*GetClassFlagExportText(Class),
-				bCastedClass ? *FString::Printf(TEXT("CASTCLASS_%s"), ClassCPPName) : TEXT("0"),
-				*FClass::GetTypePackageName(Class),
-				*APIArg);
+		}
 
-			ClassBoilerplate.Logf(TEXT("\tDECLARE_SERIALIZER(%s)\r\n"), ClassCPPName);
-			ClassBoilerplate.Log(TEXT("\t/** Indicates whether the class is compiled into the engine */\r\n"));
-			ClassBoilerplate.Log(TEXT("\tenum {IsIntrinsic=COMPILED_IN_INTRINSIC};\r\n"));
+		auto ClassRange = ClassDefinitionRange();
+		if (ClassDefinitionRanges.Contains(Class))
+		{
+			ClassRange = ClassDefinitionRanges[Class];
+			ClassRange.Validate();
+		}
 
-			if (SuperClass && Class->ClassWithin != SuperClass->ClassWithin)
-			{
-				ClassBoilerplate.Logf(TEXT("\tDECLARE_WITHIN(%s)\r\n"), NameLookupCPP.GetNameCPP(Class->GetClassWithin()));
-			}
+		auto ClassStart = ClassRange.Start;
+		auto ClassEnd = ClassRange.End;
+		auto ClassDefinition = FString(ClassEnd - ClassStart, ClassStart);
 
-			// export the class's config name
-			if (SuperClass && Class->ClassConfigName != NAME_None && Class->ClassConfigName != SuperClass->ClassConfigName)
-			{
-				ClassBoilerplate.Logf(TEXT("\tstatic const TCHAR* StaticConfigName() {return TEXT(\"%s\");}\r\n\r\n"), *Class->ClassConfigName.ToString());
-			}
+		bool bHasGetLifetimeReplicatedProps = HasIdentifierExactMatch(ClassDefinition, TEXT("GetLifetimeReplicatedProps"));
 
-			// export implementation of _getUObject for classes that implement interfaces
-			if (Class->Interfaces.Num() > 0)
+		if (bNeedsRep && !bHasGetLifetimeReplicatedProps)
+		{
+			if (SourceFile.GetGeneratedCodeVersionForStruct(Class) == EGeneratedCodeVersion::V1)
 			{
-				ClassBoilerplate.Logf(TEXT("\tvirtual UObject* _getUObject() const override { return const_cast<%s*>(this); }\r\n"), ClassCPPName);
+				InterfaceBoilerplate.Logf(TEXT("\void GetLifetimeReplicatedProps(TArray<FLifetimeProperty>& OutLifetimeProps) const;\r\n"));
 			}
-
-			// Replication, add in the declaration for GetLifetimeReplicatedProps() automatically if there are any net flagged properties
-			bool bHasNetFlaggedProperties = false;
-			for (TFieldIterator<UProperty> It(Class, EFieldIteratorFlags::ExcludeSuper); It; ++It)
+			else
 			{
-				if ((It->PropertyFlags & CPF_Net) != 0)
-				{
-					bHasNetFlaggedProperties = true;
-					break;
-				}
+				FError::Throwf(TEXT("Class %s has Net flagged properties and should declare member function: void GetLifetimeReplicatedProps(TArray<FLifetimeProperty>& OutLifetimeProps) const override"), ClassCPPName);
 			}
+		}
 
-			auto ClassRange = ClassDefinitionRange();
-			if (ClassDefinitionRanges.Contains(Class))
-			{
-				ClassRange = ClassDefinitionRanges[Class];
-				ClassRange.Validate();
-			}
+		auto NoPureDeclsMacroName = SourceFile.GetGeneratedMacroName(ClassData, TEXT("_INCLASS_IINTERFACE_NO_PURE_DECLS"));
+		WriteMacro(OutGeneratedHeaderText, NoPureDeclsMacroName, InterfaceBoilerplate);
+		ClassNoPureDeclsMacroCalls.Logf(TEXT("\t%s\r\n"), *NoPureDeclsMacroName);
 
-			auto ClassStart = ClassRange.Start;
-			auto ClassEnd = ClassRange.End;
-			auto ClassDefinition = FString(ClassEnd - ClassStart, ClassStart);
+		auto MacroName = SourceFile.GetGeneratedMacroName(ClassData, TEXT("_INCLASS_IINTERFACE"));
+		WriteMacro(OutGeneratedHeaderText, MacroName, InterfaceBoilerplate);
+		ClassMacroCalls.Logf(TEXT("\t%s\r\n"), *MacroName);
+	}
+	else
+	{
+		FUHTStringBuilder ClassBoilerplate;
 
-			bool bHasGetLifetimeReplicatedProps = HasIdentifierExactMatch(ClassDefinition, TEXT("GetLifetimeReplicatedProps"));
+		// Export the class's native function registration.
+		ClassBoilerplate.Logf(TEXT("\tprivate:\r\n"));
+		ClassBoilerplate.Logf(TEXT("\tstatic void StaticRegisterNatives%s();\r\n"), ClassCPPName);
+		ClassBoilerplate.Log(*FriendText);
+		ClassBoilerplate.Logf(TEXT("\tpublic:\r\n"));
 
-			if (bHasNetFlaggedProperties && !bHasGetLifetimeReplicatedProps)
-			{
-				// Default version autogenerates declarations.
-				if (SourceFile.GetGeneratedCodeVersionForStruct(Class) == EGeneratedCodeVersion::V1)
-				{
-					ClassBoilerplate.Logf(TEXT("\tvoid GetLifetimeReplicatedProps(TArray<FLifetimeProperty>& OutLifetimeProps) const override;\r\n"));
-				}
-				else
-				{
-					FError::Throwf(TEXT("Class %s has Net flagged properties and should declare member function: void GetLifetimeReplicatedProps(TArray<FLifetimeProperty>& OutLifetimeProps) const override"), ClassCPPName);
-				}
-			}
-			{
-				auto NoPureDeclsMacroName = SourceFile.GetGeneratedMacroName(ClassData, TEXT("_INCLASS_NO_PURE_DECLS"));
-				WriteMacro(GeneratedHeaderText, NoPureDeclsMacroName, ClassBoilerplate);
-				InClassNoPureDeclsMacroCalls.Logf(TEXT("\t%s\r\n"), *NoPureDeclsMacroName);
+		const bool bCastedClass = Class->HasAnyCastFlag(CASTCLASS_AllFlags) && SuperClass && Class->ClassCastFlags != SuperClass->ClassCastFlags;
+		ClassBoilerplate.Logf(TEXT("\tDECLARE_CLASS(%s, %s, COMPILED_IN_FLAGS(%s%s), %s, TEXT(\"%s\"), %s_API)\r\n"),
+			ClassCPPName,
+			SuperClassCPPName ? SuperClassCPPName : TEXT("None"),
+			Class->HasAnyClassFlags(CLASS_Abstract) ? TEXT("CLASS_Abstract") : TEXT("0"),
+			*GetClassFlagExportText(Class),
+			bCastedClass ? *FString::Printf(TEXT("CASTCLASS_%s"), ClassCPPName) : TEXT("0"),
+			*FClass::GetTypePackageName(Class),
+			*APIArg);
 
-				FString MacroName = SourceFile.GetGeneratedMacroName(ClassData, TEXT("_INCLASS"));
-				WriteMacro(GeneratedHeaderText, MacroName, ClassBoilerplate);
-				InClassMacroCalls.Logf(TEXT("\t%s\r\n"), *MacroName);
+		ClassBoilerplate.Logf(TEXT("\tDECLARE_SERIALIZER(%s)\r\n"), ClassCPPName);
+		ClassBoilerplate.Log(TEXT("\t/** Indicates whether the class is compiled into the engine */\r\n"));
+		ClassBoilerplate.Log(TEXT("\tenum {IsIntrinsic=COMPILED_IN_INTRINSIC};\r\n"));
 
-				ExportConstructorsMacros(SourceFile.GetGeneratedMacroName(ClassData), Class);
-			}
-			{
-				const FString PrivatePropertiesOffsets = PrivatePropertiesOffsetGetters(Class, ClassCPPName);
-				const FString PPOMacroNameRaw = SourceFile.GetGeneratedMacroName(ClassData, TEXT("_PRIVATE_PROPERTY_OFFSET"));
-				PPOMacroName = FString::Printf(TEXT("\t%s\r\n"), *PPOMacroNameRaw);
-				WriteMacro(GeneratedHeaderText, PPOMacroNameRaw, PrivatePropertiesOffsets);
-			}
+		if (SuperClass && Class->ClassWithin != SuperClass->ClassWithin)
+		{
+			ClassBoilerplate.Logf(TEXT("\tDECLARE_WITHIN(%s)\r\n"), NameLookupCPP.GetNameCPP(Class->GetClassWithin()));
 		}
 
+		// export the class's config name
+		if (SuperClass && Class->ClassConfigName != NAME_None && Class->ClassConfigName != SuperClass->ClassConfigName)
 		{
-			FString MacroName = SourceFile.GetGeneratedMacroName(ClassData->GetPrologLine(), TEXT("_PROLOG"));
-			WriteMacro(GeneratedHeaderText, MacroName, PrologMacroCalls);
+			ClassBoilerplate.Logf(TEXT("\tstatic const TCHAR* StaticConfigName() {return TEXT(\"%s\");}\r\n\r\n"), *Class->ClassConfigName.ToString());
 		}
 
+		// export implementation of _getUObject for classes that implement interfaces
+		if (Class->Interfaces.Num() > 0)
 		{
-			bool bIsIInterface = Class->HasAnyClassFlags(CLASS_Interface);
-
-			auto MacroName = FString::Printf(TEXT("GENERATED_%s_BODY()"), bIsIInterface ? TEXT("IINTERFACE") : TEXT("UCLASS"));
-
-			auto DeprecationWarning = bIsIInterface ? FString(TEXT("")) : GetGeneratedMacroDeprecationWarning(*MacroName);
-
-			auto DeprecationPushString = TEXT("PRAGMA_DISABLE_DEPRECATION_WARNINGS") LINE_TERMINATOR;
-			auto DeprecationPopString = TEXT("PRAGMA_ENABLE_DEPRECATION_WARNINGS") LINE_TERMINATOR;
+			ClassBoilerplate.Logf(TEXT("\tvirtual UObject* _getUObject() const override { return const_cast<%s*>(this); }\r\n"), ClassCPPName);
+		}
 
-			auto Public = TEXT("public:" LINE_TERMINATOR);
+		// Replication, add in the declaration for GetLifetimeReplicatedProps() automatically if there are any net flagged properties
+		bool bHasNetFlaggedProperties = false;
+		for (TFieldIterator<UProperty> It(Class, EFieldIteratorFlags::ExcludeSuper); It; ++It)
+		{
+			if ((It->PropertyFlags & CPF_Net) != 0)
+			{
+				bHasNetFlaggedProperties = true;
+				break;
+			}
+		}
 
-			auto GeneratedBodyLine = bIsIInterface ? ClassData->GetInterfaceGeneratedBodyLine() : ClassData->GetGeneratedBodyLine();
-			auto LegacyGeneratedBody = FString(bIsIInterface ? TEXT("") : PPOMacroName)
-				+ InClassMacroCalls 
-				+ (bIsIInterface ? TEXT("") : StandardUObjectConstructorsMacroCall);
-			auto GeneratedBody = FString(bIsIInterface ? TEXT("") : PPOMacroName)
-				+ InClassNoPureDeclsMacroCalls 
-				+ (bIsIInterface ? TEXT("") : EnhancedUObjectConstructorsMacroCall);
+		auto ClassRange = ClassDefinitionRange();
+		if (ClassDefinitionRanges.Contains(Class))
+		{
+			ClassRange = ClassDefinitionRanges[Class];
+			ClassRange.Validate();
+		}
 
-			auto WrappedLegacyGeneratedBody = DeprecationWarning + DeprecationPushString + Public + LegacyGeneratedBody + Public + DeprecationPopString;
-			auto WrappedGeneratedBody = FString(DeprecationPushString) + Public + GeneratedBody + GetPreservedAccessSpecifierString(Class) + DeprecationPopString;
+		auto ClassStart = ClassRange.Start;
+		auto ClassEnd = ClassRange.End;
+		auto ClassDefinition = FString(ClassEnd - ClassStart, ClassStart);
 
-			auto BodyMacros = Macroize(*SourceFile.GetGeneratedBodyMacroName(GeneratedBodyLine, true), *WrappedLegacyGeneratedBody) +
-				Macroize(*SourceFile.GetGeneratedBodyMacroName(GeneratedBodyLine, false), *WrappedGeneratedBody);
+		bool bHasGetLifetimeReplicatedProps = HasIdentifierExactMatch(ClassDefinition, TEXT("GetLifetimeReplicatedProps"));
 
-			GeneratedHeaderText.Log(*BodyMacros);
-		}
+		if (bHasNetFlaggedProperties && !bHasGetLifetimeReplicatedProps)
+		{
+			// Default version autogenerates declarations.
+			if (SourceFile.GetGeneratedCodeVersionForStruct(Class) == EGeneratedCodeVersion::V1)
+			{
+				ClassBoilerplate.Logf(TEXT("\tvoid GetLifetimeReplicatedProps(TArray<FLifetimeProperty>& OutLifetimeProps) const override;\r\n"));
+			}
+			else
+			{
+				FError::Throwf(TEXT("Class %s has Net flagged properties and should declare member function: void GetLifetimeReplicatedProps(TArray<FLifetimeProperty>& OutLifetimeProps) const override"), ClassCPPName);
+			}
+		}
+		{
+			auto NoPureDeclsMacroName = SourceFile.GetGeneratedMacroName(ClassData, TEXT("_INCLASS_NO_PURE_DECLS"));
+			WriteMacro(OutGeneratedHeaderText, NoPureDeclsMacroName, ClassBoilerplate);
+			ClassNoPureDeclsMacroCalls.Logf(TEXT("\t%s\r\n"), *NoPureDeclsMacroName);
+
+			FString MacroName = SourceFile.GetGeneratedMacroName(ClassData, TEXT("_INCLASS"));
+			WriteMacro(OutGeneratedHeaderText, MacroName, ClassBoilerplate);
+			ClassMacroCalls.Logf(TEXT("\t%s\r\n"), *MacroName);
+
+			ExportConstructorsMacros(OutGeneratedHeaderText, OutCpp, StandardUObjectConstructorsMacroCall, EnhancedUObjectConstructorsMacroCall, SourceFile.GetGeneratedMacroName(ClassData), Class, *APIArg);
+		}
+		{
+			const FString PrivatePropertiesOffsets = PrivatePropertiesOffsetGetters(Class, ClassCPPName);
+			const FString PPOMacroNameRaw = SourceFile.GetGeneratedMacroName(ClassData, TEXT("_PRIVATE_PROPERTY_OFFSET"));
+			PPOMacroName = FString::Printf(TEXT("\t%s\r\n"), *PPOMacroNameRaw);
+			WriteMacro(OutGeneratedHeaderText, PPOMacroNameRaw, PrivatePropertiesOffsets);
+		}
+	}
+
+	{
+		FString MacroName = SourceFile.GetGeneratedMacroName(ClassData->GetPrologLine(), TEXT("_PROLOG"));
+		WriteMacro(OutGeneratedHeaderText, MacroName, PrologMacroCalls);
 	}
 
-	if (GPublicSourceFileSet.Contains(&SourceFile) && !ListOfPublicHeaderGroupIncludes.Contains(&SourceFile))
 	{
-		ListOfPublicHeaderGroupIncludes.Add(&SourceFile);
+		bool bIsIInterface = Class->HasAnyClassFlags(CLASS_Interface);
+
+		auto MacroName = FString::Printf(TEXT("GENERATED_%s_BODY()"), bIsIInterface ? TEXT("IINTERFACE") : TEXT("UCLASS"));
+
+		auto DeprecationWarning = bIsIInterface ? FString(TEXT("")) : GetGeneratedMacroDeprecationWarning(*MacroName);
+
+		auto DeprecationPushString = TEXT("PRAGMA_DISABLE_DEPRECATION_WARNINGS") LINE_TERMINATOR;
+		auto DeprecationPopString = TEXT("PRAGMA_ENABLE_DEPRECATION_WARNINGS") LINE_TERMINATOR;
+
+		auto Public = TEXT("public:" LINE_TERMINATOR);
+
+		auto GeneratedBodyLine = bIsIInterface ? ClassData->GetInterfaceGeneratedBodyLine() : ClassData->GetGeneratedBodyLine();
+		auto LegacyGeneratedBody = FString(bIsIInterface ? TEXT("") : PPOMacroName)
+			+ ClassMacroCalls 
+			+ (bIsIInterface ? TEXT("") : StandardUObjectConstructorsMacroCall);
+		auto GeneratedBody = FString(bIsIInterface ? TEXT("") : PPOMacroName)
+			+ ClassNoPureDeclsMacroCalls 
+			+ (bIsIInterface ? TEXT("") : EnhancedUObjectConstructorsMacroCall);
+
+		auto WrappedLegacyGeneratedBody = DeprecationWarning + DeprecationPushString + Public + LegacyGeneratedBody + Public + DeprecationPopString;
+		auto WrappedGeneratedBody = FString(DeprecationPushString) + Public + GeneratedBody + GetPreservedAccessSpecifierString(Class) + DeprecationPopString;
+
+		auto BodyMacros = Macroize(*SourceFile.GetGeneratedBodyMacroName(GeneratedBodyLine, true), *WrappedLegacyGeneratedBody) +
+			Macroize(*SourceFile.GetGeneratedBodyMacroName(GeneratedBodyLine, false), *WrappedGeneratedBody);
+
+		OutGeneratedHeaderText.Log(*BodyMacros);
 	}
 }
 
@@ -2435,93 +2359,30 @@ void FNativeClassHeaderGenerator::ExportClassesFromSourceFileInner(FUnrealSource
 * @param Class Class to generate constructor for.
 * @param API API string for this constructor.
 */
-void ExportCopyConstructorDefinition(FUHTStringBuilder& Out, FClass* Class, const FString& API)
+void ExportCopyConstructorDefinition(FOutputDevice& Out, const TCHAR* API, const TCHAR* ClassCPPName)
 {
-	auto ClassNameCPP = NameLookupCPP.GetNameCPP(Class);
 	Out.Logf(TEXT("private:\r\n"));
 	Out.Logf(TEXT("\t/** Private move- and copy-constructors, should never be used */\r\n"));
-	Out.Logf(TEXT("\t%s_API %s(%s&&);\r\n"), *API, ClassNameCPP, ClassNameCPP);
-	Out.Logf(TEXT("\t%s_API %s(const %s&);\r\n"), *API, ClassNameCPP, ClassNameCPP);
+	Out.Logf(TEXT("\t%s_API %s(%s&&);\r\n"), API, ClassCPPName, ClassCPPName);
+	Out.Logf(TEXT("\t%s_API %s(const %s&);\r\n"), API, ClassCPPName, ClassCPPName);
 	Out.Logf(TEXT("public:\r\n"));
 }
 
-#if WITH_HOT_RELOAD_CTORS
-/**
- * Generates vtable helper constructor declaration.
- *
- * @param Out Output device to generate to.
- * @param Class Class to generate constructor for.
- * @param API API string for this constructor.
- */
-void ExportVTableHelperConstructorDecl(FStringOutputDevice& Out, FClass* Class, const FString& API)
-{
-	auto* ClassData = GScriptHelper.FindClassData(Class);
-	check(ClassData);
-	if (!ClassData->bCustomVTableHelperConstructorDeclared)
-	{
-		Out.Logf(TEXT("\tDECLARE_VTABLE_PTR_HELPER_CTOR(%s_API, %s);" LINE_TERMINATOR), *API, NameLookupCPP.GetNameCPP(Class));
-	}
-}
-
-/**
- * Generates vtable helper ctor caller.
- *
- * @param Out Output device to generate to.
- * @param Class Class to generate for.
- */
-void ExportVTableHelperCtorCaller(FStringOutputDevice& Out, FClass* Class)
-{
-	Out.Logf(TEXT("DEFINE_VTABLE_PTR_HELPER_CTOR_CALLER(%s);" LINE_TERMINATOR), NameLookupCPP.GetNameCPP(Class));
-}
-
-/**
- * Generates vtable helper ctor caller dummy.
- *
- * @param Out Output device to generate to.
- */
-void ExportVTableHelperCtorCallerDummy(FStringOutputDevice& Out)
-{
-	Out.Logf(TEXT("DEFINE_VTABLE_PTR_HELPER_CTOR_CALLER_DUMMY();" LINE_TERMINATOR));
-}
-
-/**
- * Generates vtable helper constructor body.
- *
- * @param Out Output device to generate to.
- * @param Class Class to generate constructor for.
- * @param API API string for this constructor.
- */
-void ExportVTableHelperConstructorBody(FStringOutputDevice& Out, FClass* Class)
-{
-	auto* ClassData = GScriptHelper.FindClassData(Class);
-	check(ClassData);
-	if (!ClassData->bCustomVTableHelperConstructorDeclared)
-	{
-		Out.Logf(TEXT("\tDEFINE_VTABLE_PTR_HELPER_CTOR(%s);" LINE_TERMINATOR), NameLookupCPP.GetNameCPP(Class));
-	}
-}
-
 /**
  * Generates vtable helper caller and eventual constructor body.
  *
  * @param Out Output device to generate to.
  * @param Class Class to generate for.
  * @param API API string.
- * @param bExportVTableConstructors Tells if we are going to export constructors as well.
  */
-void ExportVTableHelperCtorAndCaller(FStringOutputDevice& Out, FClass* Class, const FString& API, bool bExportVTableConstructors)
+void ExportVTableHelperCtorAndCaller(FOutputDevice& Out, FClassMetaData* ClassData, const TCHAR* API, const TCHAR* ClassCPPName)
 {
-	if (bExportVTableConstructors)
-	{
-		ExportVTableHelperConstructorDecl(Out, Class, API);
-		ExportVTableHelperCtorCaller(Out, Class);
-	}
-	else
+	if (!ClassData->bCustomVTableHelperConstructorDeclared)
 	{
-		ExportVTableHelperCtorCallerDummy(Out);
+		Out.Logf(TEXT("\tDECLARE_VTABLE_PTR_HELPER_CTOR(%s_API, %s);" LINE_TERMINATOR), API, ClassCPPName);
 	}
+	Out.Logf(TEXT("DEFINE_VTABLE_PTR_HELPER_CTOR_CALLER(%s);" LINE_TERMINATOR), ClassCPPName);
 }
-#endif // WITH_HOT_RELOAD_CTORS
 
 /**
  * Generates standard constructor declaration.
@@ -2529,31 +2390,19 @@ void ExportVTableHelperCtorAndCaller(FStringOutputDevice& Out, FClass* Class, co
  * @param Out Output device to generate to.
  * @param Class Class to generate constructor for.
  * @param API API string for this constructor.
- * @param bExportVTableConstructors Export VTable constructors.
  */
-void ExportStandardConstructorsMacro(FUHTStringBuilder& Out, FClass* Class, const FString& API
-#if WITH_HOT_RELOAD_CTORS
-		, bool bExportVTableConstructors
-#endif // WITH_HOT_RELOAD_CTORS
-)
+void ExportStandardConstructorsMacro(FOutputDevice& Out, FClass* Class, FClassMetaData* ClassData, const TCHAR* API, const TCHAR* ClassCPPName)
 {
 	if (!Class->HasAnyClassFlags(CLASS_CustomConstructor))
 	{
 		Out.Logf(TEXT("\t/** Standard constructor, called after all reflected properties have been initialized */\r\n"));
-
-		auto* ClassData = GScriptHelper.FindClassData(Class);
-		check(ClassData);
-
-		Out.Logf(TEXT("\t%s_API %s(const FObjectInitializer& ObjectInitializer%s);\r\n"), *API, NameLookupCPP.GetNameCPP(Class),
+		Out.Logf(TEXT("\t%s_API %s(const FObjectInitializer& ObjectInitializer%s);\r\n"), API, ClassCPPName,
 			ClassData->bDefaultConstructorDeclared ? TEXT("") : TEXT(" = FObjectInitializer::Get()"));
 	}
-	Out.Logf(TEXT("\tDEFINE_DEFAULT_OBJECT_INITIALIZER_CONSTRUCTOR_CALL(%s)\r\n"), NameLookupCPP.GetNameCPP(Class));
+	Out.Logf(TEXT("\tDEFINE_DEFAULT_OBJECT_INITIALIZER_CONSTRUCTOR_CALL(%s)\r\n"), ClassCPPName);
 
-#if WITH_HOT_RELOAD_CTORS
-	ExportVTableHelperCtorAndCaller(Out, Class, API, bExportVTableConstructors);
-#endif // WITH_HOT_RELOAD_CTORS
-
-	ExportCopyConstructorDefinition(Out, Class, API);
+	ExportVTableHelperCtorAndCaller(Out, ClassData, API, ClassCPPName);
+	ExportCopyConstructorDefinition(Out, API, ClassCPPName);
 }
 
 /**
@@ -2563,10 +2412,8 @@ void ExportStandardConstructorsMacro(FUHTStringBuilder& Out, FClass* Class, cons
  * @param Class Class to generate constructor for.
  * @param API API string for this constructor.
  */
-void ExportConstructorDefinition(FUHTStringBuilder& Out, FClass* Class, const FString& API)
+void ExportConstructorDefinition(FOutputDevice& Out, FClass* Class, FClassMetaData* ClassData, const TCHAR* API, const TCHAR* ClassCPPName)
 {
-	auto* ClassData = GScriptHelper.FindClassData(Class);
-	check(ClassData);
 	if (!ClassData->bConstructorDeclared)
 	{
 		Out.Logf(TEXT("\t/** Standard constructor, called after all reflected properties have been initialized */\r\n"));
@@ -2574,10 +2421,10 @@ void ExportConstructorDefinition(FUHTStringBuilder& Out, FClass* Class, const FS
 		// Assume super class has OI constructor, this may not always be true but we should always be able to check this.
 		// In any case, it will default to old behaviour before we even checked this.
 		bool bSuperClassObjectInitializerConstructorDeclared = true;
-		auto* SuperClass = Class->GetSuperClass();		
+		FClass* SuperClass = Class->GetSuperClass();
 		if (SuperClass != nullptr)
 		{
-			auto* SuperClassData = GScriptHelper.FindClassData(SuperClass);
+			FClassMetaData* SuperClassData = GScriptHelper.FindClassData(SuperClass);
 			if (SuperClassData)
 			{
 				bSuperClassObjectInitializerConstructorDeclared = SuperClassData->bObjectInitializerConstructorDeclared;
@@ -2585,18 +2432,18 @@ void ExportConstructorDefinition(FUHTStringBuilder& Out, FClass* Class, const FS
 		}
 		if (bSuperClassObjectInitializerConstructorDeclared)
 		{
-			Out.Logf(TEXT("\t%s_API %s(const FObjectInitializer& ObjectInitializer = FObjectInitializer::Get()) : Super(ObjectInitializer) { };\r\n"), *API, NameLookupCPP.GetNameCPP(Class));
+			Out.Logf(TEXT("\t%s_API %s(const FObjectInitializer& ObjectInitializer = FObjectInitializer::Get()) : Super(ObjectInitializer) { };\r\n"), API, ClassCPPName);
 			ClassData->bObjectInitializerConstructorDeclared = true;
 		}
 		else
 		{
-			Out.Logf(TEXT("\t%s_API %s() { };\r\n"), *API, NameLookupCPP.GetNameCPP(Class));
+			Out.Logf(TEXT("\t%s_API %s() { };\r\n"), API, ClassCPPName);
 			ClassData->bDefaultConstructorDeclared = true;
 		}
 
 		ClassData->bConstructorDeclared = true;
 	}
-	ExportCopyConstructorDefinition(Out, Class, API);
+	ExportCopyConstructorDefinition(Out, API, ClassCPPName);
 }
 
 /**
@@ -2605,22 +2452,19 @@ void ExportConstructorDefinition(FUHTStringBuilder& Out, FClass* Class, const FS
  * @param Out Output device to generate to.
  * @param Class Class to generate constructor call definition for.
  */
-void ExportDefaultConstructorCallDefinition(FUHTStringBuilder& Out, FClass* Class)
+void ExportDefaultConstructorCallDefinition(FOutputDevice& Out, FClassMetaData* ClassData, const TCHAR* ClassCPPName)
 {
-	auto* ClassData = GScriptHelper.FindClassData(Class);
-	check(ClassData);
-
 	if (ClassData->bObjectInitializerConstructorDeclared)
 	{
-		Out.Logf(TEXT("\tDEFINE_DEFAULT_OBJECT_INITIALIZER_CONSTRUCTOR_CALL(%s)\r\n"), NameLookupCPP.GetNameCPP(Class));
+		Out.Logf(TEXT("\tDEFINE_DEFAULT_OBJECT_INITIALIZER_CONSTRUCTOR_CALL(%s)\r\n"), ClassCPPName);
 	}
 	else if (ClassData->bDefaultConstructorDeclared)
 	{
-		Out.Logf(TEXT("\tDEFINE_DEFAULT_CONSTRUCTOR_CALL(%s)\r\n"), NameLookupCPP.GetNameCPP(Class));
+		Out.Logf(TEXT("\tDEFINE_DEFAULT_CONSTRUCTOR_CALL(%s)\r\n"), ClassCPPName);
 	}
 	else
 	{
-		Out.Logf(TEXT("\tDEFINE_FORBIDDEN_DEFAULT_CONSTRUCTOR_CALL(%s)\r\n"), NameLookupCPP.GetNameCPP(Class));
+		Out.Logf(TEXT("\tDEFINE_FORBIDDEN_DEFAULT_CONSTRUCTOR_CALL(%s)\r\n"), ClassCPPName);
 	}
 }
 
@@ -2630,21 +2474,12 @@ void ExportDefaultConstructorCallDefinition(FUHTStringBuilder& Out, FClass* Clas
  * @param Out Output device to generate to.
  * @param Class Class to generate constructor for.
  * @param API API string for this constructor.
- * @param bExportVTableConstructors Export VTable constructors.
  */
-void ExportEnhancedConstructorsMacro(FUHTStringBuilder& Out, FClass* Class, const FString& API
-#if WITH_HOT_RELOAD_CTORS
-		, bool bExportVTableConstructors
-#endif // WITH_HOT_RELOAD_CTORS
-)
+void ExportEnhancedConstructorsMacro(FOutputDevice& Out, FClass* Class, FClassMetaData* ClassData, const TCHAR* API, const TCHAR* ClassCPPName)
 {
-	ExportConstructorDefinition(Out, Class, API);
-	
-#if WITH_HOT_RELOAD_CTORS
-	ExportVTableHelperCtorAndCaller(Out, Class, API, bExportVTableConstructors);
-#endif // WITH_HOT_RELOAD_CTORS
-	
-	ExportDefaultConstructorCallDefinition(Out, Class);
+	ExportConstructorDefinition(Out, Class, ClassData, API, ClassCPPName);
+	ExportVTableHelperCtorAndCaller(Out, ClassData, API, ClassCPPName);
+	ExportDefaultConstructorCallDefinition(Out, ClassData, ClassCPPName);
 }
 
 /**
@@ -2663,209 +2498,127 @@ FString GetBuildPath(FUnrealSourceFile& SourceFile)
 	return Out;
 }
 
-FString FNativeClassHeaderGenerator::GetListOfPublicHeaderGroupIncludesString(const UPackage* InPackage)
+void FNativeClassHeaderGenerator::ExportConstructorsMacros(FOutputDevice& OutGeneratedHeaderText, FOutputDevice& Out, FOutputDevice& StandardUObjectConstructorsMacroCall, FOutputDevice& EnhancedUObjectConstructorsMacroCall, const FString& ConstructorsMacroPrefix, FClass* Class, const TCHAR* APIArg)
 {
-	FUHTStringBuilder Out;
+	const TCHAR* ClassCPPName = NameLookupCPP.GetNameCPP(Class);
 
-	// Fill with the rest source files from this package.
-	for (auto* SourceFile : GPublicSourceFileSet)
-	{
-		if (SourceFile->GetPackage() == InPackage && !ListOfPublicHeaderGroupIncludes.Contains(SourceFile))
-		{
-			ListOfPublicHeaderGroupIncludes.Add(SourceFile);
-		}
-	}
-
-	for (auto* SourceFile : ListOfPublicHeaderGroupIncludes)
-	{
-		Out.Logf(TEXT("#include \"%s\"") LINE_TERMINATOR, *GetBuildPath(*SourceFile));
-	}
-
-	Out.Log(LINE_TERMINATOR);
-
-	return Out;
-}
-
-void FNativeClassHeaderGenerator::ExportConstructorsMacros(const FString& ConstructorsMacroPrefix, FClass* Class)
-{
-	FString APIArg(API);
-	if (!Class->HasAnyClassFlags(CLASS_MinimalAPI))
-	{
-		APIArg = TEXT("NO");
-	}
+	FClassMetaData* ClassData = GScriptHelper.FindClassData(Class);
+	check(ClassData);
 
 	FUHTStringBuilder StdMacro;
 	FUHTStringBuilder EnhMacro;
 	FString StdMacroName = ConstructorsMacroPrefix + TEXT("_STANDARD_CONSTRUCTORS");
 	FString EnhMacroName = ConstructorsMacroPrefix + TEXT("_ENHANCED_CONSTRUCTORS");
 
-#if WITH_HOT_RELOAD_CTORS
-	ExportStandardConstructorsMacro(StdMacro, Class, APIArg, bExportVTableConstructors);
-	ExportEnhancedConstructorsMacro(EnhMacro, Class, APIArg, bExportVTableConstructors);
+	ExportStandardConstructorsMacro(StdMacro, Class, ClassData, APIArg, ClassCPPName);
+	ExportEnhancedConstructorsMacro(EnhMacro, Class, ClassData, APIArg, ClassCPPName);
 
-	if (bExportVTableConstructors)
+	if (!ClassData->bCustomVTableHelperConstructorDeclared)
 	{
-		ExportVTableHelperConstructorBody(GetGeneratedFunctionTextDevice(), Class);
+		Out.Logf(TEXT("\tDEFINE_VTABLE_PTR_HELPER_CTOR(%s);" LINE_TERMINATOR), ClassCPPName);
 	}
-#else // WITH_HOT_RELOAD_CTORS
-	ExportStandardConstructorsMacro(StdMacro, Class, APIArg);
-	ExportEnhancedConstructorsMacro(EnhMacro, Class, APIArg);
-#endif // WITH_HOT_RELOAD_CTORS
 
-	GeneratedHeaderText.Log(*Macroize(*StdMacroName, *StdMacro));
-	GeneratedHeaderText.Log(*Macroize(*EnhMacroName, *EnhMacro));
+	OutGeneratedHeaderText.Log(*Macroize(*StdMacroName, *StdMacro));
+	OutGeneratedHeaderText.Log(*Macroize(*EnhMacroName, *EnhMacro));
 
 	StandardUObjectConstructorsMacroCall.Logf(TEXT("\t%s\r\n"), *StdMacroName);
 	EnhancedUObjectConstructorsMacroCall.Logf(TEXT("\t%s\r\n"), *EnhMacroName);
 }
 
-void FNativeClassHeaderGenerator::ExportClassesFromSourceFileWrapper(FUnrealSourceFile& SourceFile)
+bool FNativeClassHeaderGenerator::WriteHeader(const TCHAR* Path, const FString& InBodyText, const TSet<FString>& InFwdDecl)
 {
-	check(!GeneratedHeaderText.Len());
-	ExportClassesFromSourceFileInner(SourceFile);
-
-	if (!GeneratedHeaderText.Len() && !EnumForeachText.Len())
-	{
-		// Nothing to export
-		return;
-	}
-
-	GeneratedHeaderText.Log(TEXT("#undef CURRENT_FILE_ID\r\n"));
-	GeneratedHeaderText.Logf(TEXT("#define CURRENT_FILE_ID %s\r\n\r\n\r\n"), *SourceFile.GetFileId());
-
-	const FString SourceFilename	= SourceFile.GetFilename();
-	const FString PkgName			= FPackageName::GetShortName(Package);
-
-	FString NewFileName = SourceFilename.EndsWith(TEXT(".h")) ? SourceFilename : (SourceFilename + TEXT(".h"));
-
-	FString PkgDir;
-	FString GeneratedIncludeDirectory;
-	if (FindPackageLocation(*PkgName, PkgDir, GeneratedIncludeDirectory) == false)
-	{
-		UE_LOG(LogCompile, Error, TEXT("Failed to find path for package %s"), *PkgName);
-	}
-
-	const FString ClassHeaderPath(GeneratedIncludeDirectory / FPaths::GetBaseFilename(SourceFile.GetFilename()) + TEXT(".generated.h"));
-
-	GeneratedHeaderText += EnumForeachText;
-
 	FUHTStringBuilder GeneratedHeaderTextWithCopyright;
-
-	GeneratedHeaderTextWithCopyright.Log(
-		TEXT("// Copyright 1998-2017 Epic Games, Inc. All Rights Reserved.\r\n")
-		TEXT("/*===========================================================================\r\n")
-		TEXT("\tC++ class header boilerplate exported from UnrealHeaderTool.\r\n")
-		TEXT("\tThis is automatically generated by the tools.\r\n")
-		TEXT("\tDO NOT modify this manually! Edit the corresponding .h files instead!\r\n")
-		TEXT("===========================================================================*/\r\n")
-		LINE_TERMINATOR
-		TEXT("#include \"ObjectMacros.h\"\r\n")
-		LINE_TERMINATOR);
-
+	GeneratedHeaderTextWithCopyright.Logf(TEXT("%s"), HeaderCopyright);
+	GeneratedHeaderTextWithCopyright.Log(LINE_TERMINATOR);
+	GeneratedHeaderTextWithCopyright.Log(TEXT("#include \"ObjectMacros.h\"\r\n"));
+	GeneratedHeaderTextWithCopyright.Log(LINE_TERMINATOR);
 	GeneratedHeaderTextWithCopyright.Log(TEXT("PRAGMA_DISABLE_DEPRECATION_WARNINGS") LINE_TERMINATOR);
 
-	TSet<FString> ForwardDeclarationStrings;
-	for (auto Property : ForwardDeclarations)
+	for (const FString& FWDecl : InFwdDecl)
 	{
-		auto FWDecl = Property->GetCPPTypeForwardDeclaration();
 		if (FWDecl.Len() > 0)
 		{
-			ForwardDeclarationStrings.Add(FWDecl);
+			GeneratedHeaderTextWithCopyright.Logf(TEXT("%s\r\n"), *FWDecl);
 		}
 	}
 
-	for (auto FWDeclString : ForwardDeclarationStrings)
-	{
-		GeneratedForwardDeclarations.Logf(TEXT("%s\r\n"), *FWDeclString);
-	}
-
-	GeneratedHeaderTextWithCopyright.Log(*GeneratedForwardDeclarations);
-	GeneratedHeaderTextWithCopyright.Log(*GeneratedHeaderText);
+	GeneratedHeaderTextWithCopyright.Log(*InBodyText);
 	GeneratedHeaderTextWithCopyright.Log(TEXT("PRAGMA_ENABLE_DEPRECATION_WARNINGS") LINE_TERMINATOR);
 
-	SourceFile.SetGeneratedFilename(ClassHeaderPath);
-	SourceFile.SetHasChanged(SaveHeaderIfChanged(*ClassHeaderPath, *GeneratedHeaderTextWithCopyright));
-
-	check(SourceFilename.EndsWith(TEXT(".h")));
-
-	ConvertToBuildIncludePath(Package, NewFileName);
-
-	auto IncludeStr = FString::Printf(
-		TEXT("#ifndef %s")			LINE_TERMINATOR
-		TEXT("\t#include \"%s\"")	LINE_TERMINATOR
-		TEXT("#endif")				LINE_TERMINATOR,
-		*SourceFile.GetFileDefineName(), *NewFileName);
-
-	// Keep track of all of the UObject headers for this module, in the same order that we digest them in
-	// @todo uht: We're wrapping these includes in checks for header guards, ONLY because most existing UObject headers are missing '#pragma once'
-	ListOfAllUObjectHeaderIncludes                .Log(*IncludeStr);
-	ListOfPublicClassesUObjectHeaderModuleIncludes.Log(*IncludeStr);
-
-	ForwardDeclarations.Reset();
-	GeneratedForwardDeclarations.Reset();
-	GeneratedHeaderText.Reset();
-	EnumForeachText.Reset();
+	bool bHasChanged = SaveHeaderIfChanged(Path, *GeneratedHeaderTextWithCopyright);
+	return bHasChanged;
 }
 
-void FNativeClassHeaderGenerator::ExportSourceFileHeaderRecursive(FClasses& AllClasses, FUnrealSourceFile* SourceFile, TSet<const FUnrealSourceFile*>& VisitedSet, bool bCheckDependenciesOnly)
+void FNativeClassHeaderGenerator::ExportSourceFileHeader(
+	FOutputDevice&                     OutGeneratedHeaderText,
+	TFunctionRef<FUHTStringBuilder&()> OutputGetter,
+	FOutputDevice&                     OutDeclarations,
+	const FUnrealSourceFile&           SourceFile,
+	const TArray<UEnum*>&              Enums,
+	const TArray<UScriptStruct*>&      Structs,
+	const TArray<UDelegateFunction*>&  DelegateFunctions
+)
 {
-	bool bIsCorrectHeader = SourceFile->GetPackage() == Package;
+	OutGeneratedHeaderText.Logf(
+		TEXT("#ifdef %s")																	LINE_TERMINATOR
+		TEXT("#error \"%s.generated.h already included, missing '#pragma once' in %s.h\"")	LINE_TERMINATOR
+		TEXT("#endif")																		LINE_TERMINATOR
+		TEXT("#define %s")																	LINE_TERMINATOR
+		LINE_TERMINATOR,
+		*SourceFile.GetFileDefineName(), *SourceFile.GetStrippedFilename(), *SourceFile.GetStrippedFilename(), *SourceFile.GetFileDefineName());
 
-	// Check for circular header dependencies between export classes.
-	if (!bIsCorrectHeader)
-	{
-		if (VisitedSet.Num() == 0)
-		{
-			// The first export class we found doesn't belong in this header: No need to keep exporting along this dependency path.
-			return;
-		}
+	ExportAutoIncludes(OutGeneratedHeaderText, SourceFile);
 
-		// From now on, we're not going to export anything. Instead, we're going to check that no deeper dependency tries to export to this header file.
-		bCheckDependenciesOnly = true;
+	// export delegate definitions
+	for (UDelegateFunction* Func : DelegateFunctions)
+	{
+		ExportDelegateDeclaration(OutputGetter(), OutDeclarations, SourceFile, Func);
 	}
 
-	// Check if the Class has already been exported, after we've checked for circular header dependencies.
-	if (GExportedSourceFiles.Contains(SourceFile))
+	// Export enums declared in non-UClass headers.
+	for (UEnum* Enum : Enums)
 	{
-		return;
+		// Is this ever not the case?
+		if (Enum->GetOuter()->IsA(UPackage::StaticClass()))
+		{
+			ExportGeneratedEnumInitCode(OutputGetter(), OutDeclarations, SourceFile, Enum);
+		}
 	}
 
-	// Check for circular dependencies.
-	if (VisitedSet.Contains(SourceFile))
+	// export boilerplate macros for structs
+	// reverse the order.
+	for (UScriptStruct* Struct : Structs)
 	{
-		UE_LOG(LogCompile, Error, TEXT("Circular dependency detected for filename %s!"), *SourceFile->GetFilename());
-		return;
+		ExportGeneratedStructBodyMacros(OutGeneratedHeaderText, OutputGetter(), OutDeclarations, SourceFile, Struct);
 	}
 
-	// Temporarily mark the Class as VISITED. Make sure to clear this flag before returning!
-	VisitedSet.Add(SourceFile);
-
-	TArray<FUnrealSourceFile*> DependOnList;
-	for (auto& Include : SourceFile->GetIncludes())
+	// export delegate wrapper function implementations
+	if (DelegateFunctions.Num())
 	{
-		auto* OtherSourceFile = Include.Resolve();
-		if (OtherSourceFile && !DependOnList.Contains(OtherSourceFile))
+		for (UDelegateFunction* Func : DelegateFunctions)
 		{
-			DependOnList.Add(OtherSourceFile);
+			ExportDelegateDefinition(OutGeneratedHeaderText, SourceFile, Func);
 		}
+
+		OutGeneratedHeaderText.Log(TEXT("\r\n\r\n"));
 	}
 
-	for (auto* DependOn : DependOnList)
+	TArray<UClass*> DefinedClasses = SourceFile.GetDefinedClasses();
+	for (UClass* Class : DefinedClasses)
 	{
-		ExportSourceFileHeaderRecursive(AllClasses, DependOn, VisitedSet, bCheckDependenciesOnly);
+		if (!(Class->ClassFlags & CLASS_Intrinsic))
+		{
+			ExportClassFromSourceFileInner(OutGeneratedHeaderText, OutputGetter(), OutDeclarations, (FClass*)Class, SourceFile);
+		}
 	}
 
-	// Export class header.
-	if (bIsCorrectHeader && !bCheckDependenciesOnly)
-	{
-		// Mark class as exported.
-		GExportedSourceFiles.Add(SourceFile);
+	OutGeneratedHeaderText.Log(TEXT("#undef CURRENT_FILE_ID\r\n"));
+	OutGeneratedHeaderText.Logf(TEXT("#define CURRENT_FILE_ID %s\r\n\r\n\r\n"), *SourceFile.GetFileId());
 
-		ExportClassesFromSourceFileWrapper(*SourceFile);
+	for (UEnum* Enum : Enums)
+	{
+		ExportEnum(OutGeneratedHeaderText, Enum);
 	}
-
-	// We're done visiting this Class.
-	VisitedSet.Remove(SourceFile);
 }
 
 /**
@@ -2910,412 +2663,374 @@ FString FNativeClassHeaderGenerator::GetClassFlagExportText( UClass* Class )
 *
 * @param	Enums	the enums to export
 */
-void FNativeClassHeaderGenerator::ExportEnums( const TArray<UEnum*>& Enums )
+void FNativeClassHeaderGenerator::ExportEnum(FOutputDevice& Out, UEnum* Enum)
 {
-	// Enum definitions.
-	for( int32 EnumIdx = Enums.Num() - 1; EnumIdx >= 0; EnumIdx-- )
+	// Export FOREACH macro
+	Out.Logf( TEXT("#define FOREACH_ENUM_%s(op) "), *Enum->GetName().ToUpper() );
+	for (int32 i = 0; i < Enum->NumEnums() - 1; i++)
 	{
-		UEnum* Enum = Enums[EnumIdx];
-
-		// Export FOREACH macro
-		EnumForeachText.Logf( TEXT("#define FOREACH_ENUM_%s(op) "), *Enum->GetName().ToUpper() );
-		for (int32 i = 0; i < Enum->NumEnums() - 1; i++)
-		{
-			const FString QualifiedEnumValue = Enum->GetNameByIndex(i).ToString();
-			EnumForeachText.Logf( TEXT("\\\r\n\top(%s) "), *QualifiedEnumValue );
-		}
-		EnumForeachText.Logf( TEXT("\r\n") );
+		const FString QualifiedEnumValue = Enum->GetNameByIndex(i).ToString();
+		Out.Logf( TEXT("\\\r\n\top(%s) "), *QualifiedEnumValue );
 	}
+	Out.Logf( TEXT("\r\n") );
 }
 
 // Exports the header text for the list of structs specified (GENERATED_BODY impls)
-void FNativeClassHeaderGenerator::ExportGeneratedStructBodyMacros(FUnrealSourceFile& SourceFile, const TArray<UScriptStruct*>& NativeStructs)
+void FNativeClassHeaderGenerator::ExportGeneratedStructBodyMacros(FOutputDevice& OutGeneratedHeaderText, FOutputDevice& Out, FOutputDevice& OutDeclarations, const FUnrealSourceFile& SourceFile, UScriptStruct* Struct)
 {
-	// reverse the order.
-	for (int32 i = NativeStructs.Num() - 1; i >= 0; --i)
-	{
-		UScriptStruct* Struct = NativeStructs[i];
-		const bool bIsDynamic = FClass::IsDynamic(Struct);
-		const FString ActualStructName = FNativeClassHeaderGenerator::GetOverriddenName(Struct);
-
-		UStruct* BaseStruct = Struct->GetSuperStruct();
-
-		// Export struct.
-		if (Struct->StructFlags & STRUCT_Native)
-		{
-			check(Struct->StructMacroDeclaredLineNumber != INDEX_NONE);
+	const bool bIsDynamic = FClass::IsDynamic(Struct);
+	const FString ActualStructName = FNativeClassHeaderGenerator::GetOverriddenName(Struct);
 
-			const FString FriendApiString = GetAPIString();
-			const FString StaticConstructionString = GetSingletonName(Struct);
-
-			FString RequiredAPI;
-			if (!(Struct->StructFlags & STRUCT_RequiredAPI))
-			{
-				RequiredAPI = FriendApiString;
-			}
+	UStruct* BaseStruct = Struct->GetSuperStruct();
 
-			const TCHAR* StructNameCPP = NameLookupCPP.GetNameCPP(Struct);
-			const FString FriendLine = FString::Printf(TEXT("\tfriend %sclass UScriptStruct* %s;\r\n"), *FriendApiString, *StaticConstructionString);
-			const FString StaticClassLine = FString::Printf(TEXT("\t%sstatic class UScriptStruct* StaticStruct();\r\n"), *RequiredAPI);
-			const FString PrivatePropertiesOffset = PrivatePropertiesOffsetGetters(Struct, StructNameCPP);
-			const FString SuperTypedef = BaseStruct ? FString::Printf(TEXT("\ttypedef %s Super;\r\n"), NameLookupCPP.GetNameCPP(BaseStruct)) : FString();
+	// Export struct.
+	if (Struct->StructFlags & STRUCT_Native)
+	{
+		check(Struct->StructMacroDeclaredLineNumber != INDEX_NONE);
 
-			const FString CombinedLine = FriendLine + StaticClassLine + PrivatePropertiesOffset + SuperTypedef;
-			const FString MacroName = SourceFile.GetGeneratedBodyMacroName(Struct->StructMacroDeclaredLineNumber);
+		const FString FriendApiString = GetAPIString();
+		const FString StaticConstructionString = GetSingletonName(Struct);
 
-			const FString Macroized = Macroize(*MacroName, *CombinedLine);
-			GeneratedHeaderText.Log(*Macroized);
+		FString RequiredAPI;
+		if (!(Struct->StructFlags & STRUCT_RequiredAPI))
+		{
+			RequiredAPI = FriendApiString;
+		}
 
-			FString SingletonName = StaticConstructionString.Replace(TEXT("()"), TEXT(""), ESearchCase::CaseSensitive); // function address
-			FString GetCRCName = FString::Printf(TEXT("Get_%s_CRC"), *SingletonName);
-			
-			GeneratedPackageCPP.Logf(TEXT("class UScriptStruct* %s::StaticStruct()\r\n"), StructNameCPP);
-			GeneratedPackageCPP.Logf(TEXT("{\r\n"));
+		const TCHAR* StructNameCPP = NameLookupCPP.GetNameCPP(Struct);
+		const FString FriendLine = FString::Printf(TEXT("\tfriend %sclass UScriptStruct* %s;\r\n"), *FriendApiString, *StaticConstructionString);
+		const FString StaticClassLine = FString::Printf(TEXT("\t%sstatic class UScriptStruct* StaticStruct();\r\n"), *RequiredAPI);
+		const FString PrivatePropertiesOffset = PrivatePropertiesOffsetGetters(Struct, StructNameCPP);
+		const FString SuperTypedef = BaseStruct ? FString::Printf(TEXT("\ttypedef %s Super;\r\n"), NameLookupCPP.GetNameCPP(BaseStruct)) : FString();
 
-			// UStructs can have UClass or UPackage outer (if declared in non-UClass headers).
-			FString OuterName;
-			if (Struct->GetOuter()->IsA(UStruct::StaticClass()))
-			{
-				OuterName = NameLookupCPP.GetNameCPP(CastChecked<UStruct>(Struct->GetOuter()));
-				OuterName += TEXT("::StaticClass()");
-			}
-			else if (!bIsDynamic)
-			{
-				OuterName = GetPackageSingletonName(CastChecked<UPackage>(Struct->GetOuter()));
-				GeneratedPackageCPP.Logf(TEXT("\textern %sclass UPackage* %s;\r\n"), *FriendApiString, *OuterName);
-			}
-			else
-			{
-				OuterName = TEXT("StructPackage");
-				GeneratedPackageCPP.Logf(TEXT("\tclass UPackage* %s = FindOrConstructDynamicTypePackage(TEXT(\"%s\"));\r\n"), *OuterName, *FClass::GetTypePackageName(Struct));
-			}
+		const FString CombinedLine = FriendLine + StaticClassLine + PrivatePropertiesOffset + SuperTypedef;
+		const FString MacroName = SourceFile.GetGeneratedBodyMacroName(Struct->StructMacroDeclaredLineNumber);
 
-			if (!bIsDynamic)
-			{
-				GeneratedPackageCPP.Logf(TEXT("\tstatic class UScriptStruct* Singleton = NULL;\r\n"));
-			}
-			else
-			{
-				GeneratedPackageCPP.Logf(TEXT("\tclass UScriptStruct* Singleton = Cast<UScriptStruct>(StaticFindObjectFast(UScriptStruct::StaticClass(), %s, TEXT(\"%s\")));\r\n"),
-					*OuterName, *ActualStructName);
-			}
-			GeneratedPackageCPP.Logf(TEXT("\tif (!Singleton)\r\n"));
-			GeneratedPackageCPP.Logf(TEXT("\t{\r\n"));
-			GeneratedPackageCPP.Logf(TEXT("\t\textern %sclass UScriptStruct* %s;\r\n"), *FriendApiString, *StaticConstructionString);
-			GeneratedPackageCPP.Logf(TEXT("\t\textern %suint32 %s();\r\n"), *FriendApiString, *GetCRCName);
-
-			GeneratedPackageCPP.Logf(TEXT("\t\tSingleton = GetStaticStruct(%s, %s, TEXT(\"%s\"), sizeof(%s), %s());\r\n"),
-				*SingletonName, *OuterName, *ActualStructName, StructNameCPP, *GetCRCName);
-
-			GeneratedPackageCPP.Logf(TEXT("\t}\r\n"));
-			GeneratedPackageCPP.Logf(TEXT("\treturn Singleton;\r\n"));
-			GeneratedPackageCPP.Logf(TEXT("}\r\n"));
-
-			GeneratedPackageCPP.Logf(TEXT("static FCompiledInDeferStruct Z_CompiledInDeferStruct_UScriptStruct_%s(%s::StaticStruct, TEXT(\"%s\"), TEXT(\"%s\"), %s, %s, %s);\r\n"),
-				StructNameCPP, StructNameCPP, *Struct->GetOutermost()->GetName(), *ActualStructName,
-				bIsDynamic ? TEXT("true") : TEXT("false"),
-				bIsDynamic ? *AsTEXT(FClass::GetTypePackageName(Struct)) : TEXT("nullptr"),
-				bIsDynamic ? *AsTEXT(FNativeClassHeaderGenerator::GetOverriddenPathName(Struct)) : TEXT("nullptr"));
-
-			// Generate StaticRegisterNatives equivalent for structs without classes.
-			if (!Struct->GetOuter()->IsA(UStruct::StaticClass()))
-			{
-				const FString ShortPackageName = FPackageName::GetShortName(Struct->GetOuter()->GetName());
-				GeneratedPackageCPP.Logf(TEXT("static struct FScriptStruct_%s_StaticRegisterNatives%s\r\n"), *ShortPackageName, StructNameCPP);
-				GeneratedPackageCPP.Logf(TEXT("{\r\n"));
-				GeneratedPackageCPP.Logf(TEXT("\tFScriptStruct_%s_StaticRegisterNatives%s()\r\n"), *ShortPackageName, StructNameCPP);
-				GeneratedPackageCPP.Logf(TEXT("\t{\r\n"));
+		const FString Macroized = Macroize(*MacroName, *CombinedLine);
+		OutGeneratedHeaderText.Log(*Macroized);
 
-				GeneratedPackageCPP.Logf(TEXT("\t\tUScriptStruct::DeferCppStructOps(FName(TEXT(\"%s\")),new UScriptStruct::TCppStructOps<%s>);\r\n"), *ActualStructName, StructNameCPP);
+		FString SingletonName = StaticConstructionString.Replace(TEXT("()"), TEXT(""), ESearchCase::CaseSensitive); // function address
+		FString GetCRCName = FString::Printf(TEXT("Get_%s_CRC"), *SingletonName);
+		
+		Out.Logf(TEXT("class UScriptStruct* %s::StaticStruct()\r\n"), StructNameCPP);
+		Out.Logf(TEXT("{\r\n"));
 
-				GeneratedPackageCPP.Logf(TEXT("\t}\r\n"));
-				GeneratedPackageCPP.Logf(TEXT("} ScriptStruct_%s_StaticRegisterNatives%s;\r\n"), *ShortPackageName, StructNameCPP);
-			}
+		// UStructs can have UClass or UPackage outer (if declared in non-UClass headers).
+		FString OuterName;
+		if (Struct->GetOuter()->IsA(UStruct::StaticClass()))
+		{
+			OuterName = NameLookupCPP.GetNameCPP(CastChecked<UStruct>(Struct->GetOuter()));
+			OuterName += TEXT("::StaticClass()");
 		}
-
+		else if (!bIsDynamic)
 		{
-			UScriptStruct* ScriptStruct = Struct;
-
-			const FString SingletonName(GetSingletonName(ScriptStruct));
-			GeneratedFunctionDeclarations.Log(FTypeSingletonCache::Get(ScriptStruct).GetExternDecl());
+			OuterName = GetPackageSingletonName(CastChecked<UPackage>(Struct->GetOuter()));
+			Out.Logf(TEXT("\textern %sclass UPackage* %s;\r\n"), *FriendApiString, *OuterName);
+		}
+		else
+		{
+			OuterName = TEXT("StructPackage");
+			Out.Logf(TEXT("\tclass UPackage* %s = FindOrConstructDynamicTypePackage(TEXT(\"%s\"));\r\n"), *OuterName, *FClass::GetTypePackageName(Struct));
+		}
 
-			FUHTStringBuilder GeneratedStructRegisterFunctionText;
+		if (!bIsDynamic)
+		{
+			Out.Logf(TEXT("\tstatic class UScriptStruct* Singleton = NULL;\r\n"));
+		}
+		else
+		{
+			Out.Logf(TEXT("\tclass UScriptStruct* Singleton = Cast<UScriptStruct>(StaticFindObjectFast(UScriptStruct::StaticClass(), %s, TEXT(\"%s\")));\r\n"),
+				*OuterName, *ActualStructName);
+		}
+		Out.Logf(TEXT("\tif (!Singleton)\r\n"));
+		Out.Logf(TEXT("\t{\r\n"));
+		Out.Logf(TEXT("\t\textern %sclass UScriptStruct* %s;\r\n"), *FriendApiString, *StaticConstructionString);
+		Out.Logf(TEXT("\t\textern %suint32 %s();\r\n"), *FriendApiString, *GetCRCName);
 
-			GeneratedStructRegisterFunctionText.Logf(TEXT("\tUScriptStruct* %s\r\n"), *SingletonName);
-			GeneratedStructRegisterFunctionText.Logf(TEXT("\t{\r\n"));
+		Out.Logf(TEXT("\t\tSingleton = GetStaticStruct(%s, %s, TEXT(\"%s\"), sizeof(%s), %s());\r\n"),
+			*SingletonName, *OuterName, *ActualStructName, StructNameCPP, *GetCRCName);
 
-			// if this is a no export struct, we will put a local struct here for offset determination
-			TArray<UScriptStruct*> Structs;
-			FindNoExportStructs(Structs, ScriptStruct);
-			if (Structs.Num())
-			{
-				TGuardValue<bool> Guard(bIsExportingForOffsetDeterminationOnly, true);
-				ExportMirrorsForNoexportStructs(Structs, /*Indent=*/ 2, GeneratedStructRegisterFunctionText);
-			}
+		Out.Logf(TEXT("\t}\r\n"));
+		Out.Logf(TEXT("\treturn Singleton;\r\n"));
+		Out.Logf(TEXT("}\r\n"));
 
-			FString CRCFuncName = FString::Printf(TEXT("Get_%s_CRC"), *SingletonName.Replace(TEXT("()"), TEXT(""), ESearchCase::CaseSensitive));
+		Out.Logf(TEXT("static FCompiledInDeferStruct Z_CompiledInDeferStruct_UScriptStruct_%s(%s::StaticStruct, TEXT(\"%s\"), TEXT(\"%s\"), %s, %s, %s);\r\n"),
+			StructNameCPP, StructNameCPP, *Struct->GetOutermost()->GetName(), *ActualStructName,
+			bIsDynamic ? TEXT("true") : TEXT("false"),
+			bIsDynamic ? *AsTEXT(FClass::GetTypePackageName(Struct)) : TEXT("nullptr"),
+			bIsDynamic ? *AsTEXT(FNativeClassHeaderGenerator::GetOverriddenPathName(Struct)) : TEXT("nullptr"));
 
-			// Structs can either have a UClass or UPackage as outer (if delcared in non-UClass header).
-			if (ScriptStruct->GetOuter()->IsA(UStruct::StaticClass()))
-			{
-				GeneratedStructRegisterFunctionText.Logf(TEXT("\t\tUStruct* Outer = %s;\r\n"), *GetSingletonName(CastChecked<UStruct>(ScriptStruct->GetOuter())));
-			}
-			else if (!bIsDynamic)
-			{
-				GeneratedStructRegisterFunctionText.Logf(TEXT("\t\tUPackage* Outer = %s;\r\n"), *GetPackageSingletonName(CastChecked<UPackage>(ScriptStruct->GetOuter())));
-			}
-			else
-			{
-				GeneratedStructRegisterFunctionText.Logf(TEXT("\t\tUPackage* Outer = FindOrConstructDynamicTypePackage(TEXT(\"%s\"));\r\n"), *FClass::GetTypePackageName(ScriptStruct));
-			}
+		// Generate StaticRegisterNatives equivalent for structs without classes.
+		if (!Struct->GetOuter()->IsA(UStruct::StaticClass()))
+		{
+			const FString ShortPackageName = FPackageName::GetShortName(Struct->GetOuter()->GetName());
+			Out.Logf(TEXT("static struct FScriptStruct_%s_StaticRegisterNatives%s\r\n"), *ShortPackageName, StructNameCPP);
+			Out.Logf(TEXT("{\r\n"));
+			Out.Logf(TEXT("\tFScriptStruct_%s_StaticRegisterNatives%s()\r\n"), *ShortPackageName, StructNameCPP);
+			Out.Logf(TEXT("\t{\r\n"));
 
-			GeneratedStructRegisterFunctionText.Logf(TEXT("\t\textern uint32 %s();\r\n"), *CRCFuncName);
-			if (!bIsDynamic)
-			{
-				GeneratedStructRegisterFunctionText.Logf(TEXT("\t\tstatic UScriptStruct* ReturnStruct = FindExistingStructIfHotReloadOrDynamic(Outer, TEXT(\"%s\"), sizeof(%s), %s(), false);\r\n"), *ActualStructName, NameLookupCPP.GetNameCPP(Struct), *CRCFuncName);
-			}
-			else
-			{
-				GeneratedStructRegisterFunctionText.Logf(TEXT("\t\tUScriptStruct* ReturnStruct = FindExistingStructIfHotReloadOrDynamic(Outer, TEXT(\"%s\"), sizeof(%s), %s(), true);\r\n"), *ActualStructName, NameLookupCPP.GetNameCPP(Struct), *CRCFuncName);
-			}
-			GeneratedStructRegisterFunctionText.Logf(TEXT("\t\tif (!ReturnStruct)\r\n"));
-			GeneratedStructRegisterFunctionText.Logf(TEXT("\t\t{\r\n"));
-			FString BaseStructString(TEXT("NULL"));
-			if (BaseStruct)
-			{
-				CastChecked<UScriptStruct>(BaseStruct); // this better actually be a script struct
-				BaseStructString = GetSingletonName(BaseStruct);
-			}
-			FString CppStructOpsString(TEXT("NULL"));
-			FString ExplicitSizeString;
-			FString ExplicitAlignmentString;
-			if ((ScriptStruct->StructFlags&STRUCT_Native) != 0)
-			{
-				//@todo .u we don't need the auto register versions of these (except for hotreload, which should be fixed)
-				CppStructOpsString = FString::Printf(TEXT("new UScriptStruct::TCppStructOps<%s>"), NameLookupCPP.GetNameCPP(ScriptStruct));
-			}
-			else
-			{
-				ExplicitSizeString = FString::Printf(TEXT(", sizeof(%s), ALIGNOF(%s)"), NameLookupCPP.GetNameCPP(ScriptStruct), NameLookupCPP.GetNameCPP(ScriptStruct));
-			}
+			Out.Logf(TEXT("\t\tUScriptStruct::DeferCppStructOps(FName(TEXT(\"%s\")),new UScriptStruct::TCppStructOps<%s>);\r\n"), *ActualStructName, StructNameCPP);
 
-			const TCHAR* UStructObjectFlags = bIsDynamic ? TEXT("RF_Public|RF_Transient") : TEXT("RF_Public|RF_Transient|RF_MarkAsNative");
-			GeneratedStructRegisterFunctionText.Logf(TEXT("\t\t\tReturnStruct = new(EC_InternalUseOnlyConstructor, Outer, TEXT(\"%s\"), %s) UScriptStruct(FObjectInitializer(), %s, %s, EStructFlags(0x%08X)%s);\r\n"),
-				*ActualStructName,
-				UStructObjectFlags,
-				*BaseStructString,
-				*CppStructOpsString,
-				(uint32)(ScriptStruct->StructFlags & ~STRUCT_ComputedFlags),
-				*ExplicitSizeString
-				);
-			TheFlagAudit.Add(ScriptStruct, TEXT("StructFlags"), (uint64)(ScriptStruct->StructFlags & ~STRUCT_ComputedFlags));
+			Out.Logf(TEXT("\t}\r\n"));
+			Out.Logf(TEXT("} ScriptStruct_%s_StaticRegisterNatives%s;\r\n"), *ShortPackageName, StructNameCPP);
+		}
+	}
 
-			TArray<UProperty*> Props;
-			for (TFieldIterator<UProperty> ItInner(ScriptStruct, EFieldIteratorFlags::ExcludeSuper); ItInner; ++ItInner)
-			{
-				Props.Add(*ItInner);
-			}
-			FString OuterString = FString(TEXT("ReturnStruct"));
-			FString Meta = GetMetaDataCodeForObject(ScriptStruct, *OuterString, TEXT("\t\t\t"));
-			OutputProperties(Meta, GeneratedStructRegisterFunctionText, OuterString, Props, TEXT("\t\t\t"));
-			GeneratedStructRegisterFunctionText.Logf(TEXT("\t\t\tReturnStruct->StaticLink();\r\n"));
+	const FString SingletonName = GetSingletonName(Struct);
+	OutDeclarations.Log(FTypeSingletonCache::Get(Struct).GetExternDecl());
 
-			if (Meta.Len())
-			{
-				GeneratedStructRegisterFunctionText.Logf(TEXT("#if WITH_METADATA\r\n"));
-				GeneratedStructRegisterFunctionText.Logf(TEXT("\t\t\tUMetaData* MetaData = ReturnStruct->GetOutermost()->GetMetaData();\r\n"));
-				GeneratedStructRegisterFunctionText.Log(*Meta);
-				GeneratedStructRegisterFunctionText.Logf(TEXT("#endif\r\n"));
-			}
+	FUHTStringBuilder GeneratedStructRegisterFunctionText;
 
-			GeneratedStructRegisterFunctionText.Logf(TEXT("\t\t}\r\n"));
-			GeneratedStructRegisterFunctionText.Logf(TEXT("\t\treturn ReturnStruct;\r\n"));
-			GeneratedStructRegisterFunctionText.Logf(TEXT("\t}\r\n"));
+	GeneratedStructRegisterFunctionText.Logf(TEXT("\tUScriptStruct* %s\r\n"), *SingletonName);
+	GeneratedStructRegisterFunctionText.Logf(TEXT("\t{\r\n"));
 
-			uint32 StructCrc = GenerateTextCRC(*GeneratedStructRegisterFunctionText);
-			GGeneratedCodeCRCs.Add(ScriptStruct, StructCrc);
-			UHTMakefile.AddGeneratedCodeCRC(&SourceFile, ScriptStruct, StructCrc);
+	// if this is a no export struct, we will put a local struct here for offset determination
+	TArray<UScriptStruct*> Structs = FindNoExportStructs(Struct);
+	for (UScriptStruct* NoExportStruct : Structs)
+	{
+		ExportMirrorsForNoexportStruct(GeneratedStructRegisterFunctionText, NoExportStruct, /*Indent=*/ 2);
+	}
 
-			auto& GeneratedFunctionText = GetGeneratedFunctionTextDevice();
-			GeneratedFunctionText += GeneratedStructRegisterFunctionText;
-			GeneratedFunctionText.Logf(TEXT("\tuint32 %s() { return %uU; }\r\n"), *CRCFuncName, StructCrc);
+	FString CRCFuncName = FString::Printf(TEXT("Get_%s_CRC"), *SingletonName.Replace(TEXT("()"), TEXT(""), ESearchCase::CaseSensitive));
 
-			//CallSingletons.Logf(TEXT("\t\t\t\tOuterClass->LinkChild(%s); // %u\r\n"), *SingletonName, StructCrc);
-		}
+	// Structs can either have a UClass or UPackage as outer (if declared in non-UClass header).
+	if (Struct->GetOuter()->IsA(UStruct::StaticClass()))
+	{
+		GeneratedStructRegisterFunctionText.Logf(TEXT("\t\tUStruct* Outer = %s;\r\n"), *GetSingletonName(CastChecked<UStruct>(Struct->GetOuter())));
 	}
-}
-
-void FNativeClassHeaderGenerator::ExportGeneratedEnumsInitCode(FUnrealSourceFile& SourceFile, const TArray<UEnum*>& Enums)
-{
-	// reverse the order.
-	for (int32 i = Enums.Num() - 1; i >= 0; --i)
+	else if (!bIsDynamic)
 	{
-		UEnum* Enum = Enums[i];
-
-		// Export Enum.
-		if (!Enum->GetOuter()->IsA(UPackage::StaticClass()))
-		{
-			continue;
-		}
-
-		const bool bIsDynamic = FClass::IsDynamic(Enum);
-		const FString FriendApiString = GetAPIString();
-		const FString StaticConstructionString = GetSingletonName(Enum);
+		GeneratedStructRegisterFunctionText.Logf(TEXT("\t\tUPackage* Outer = %s;\r\n"), *GetPackageSingletonName(CastChecked<UPackage>(Struct->GetOuter())));
+	}
+	else
+	{
+		GeneratedStructRegisterFunctionText.Logf(TEXT("\t\tUPackage* Outer = FindOrConstructDynamicTypePackage(TEXT(\"%s\"));\r\n"), *FClass::GetTypePackageName(Struct));
+	}
 
-		FString SingletonName = StaticConstructionString.Replace(TEXT("()"), TEXT(""), ESearchCase::CaseSensitive); // function address
-		FString PackageSingletonName = GetPackageSingletonName(CastChecked<UPackage>(Enum->GetOuter()));
-		if (!bIsDynamic)
-		{
-			PackageSingletonName = GetPackageSingletonName(CastChecked<UPackage>(Enum->GetOuter()));
-		}
-		else
-		{
-			PackageSingletonName = FClass::GetTypePackageName(Enum);
-		}
+	GeneratedStructRegisterFunctionText.Logf(TEXT("\t\textern uint32 %s();\r\n"), *CRCFuncName);
+	if (!bIsDynamic)
+	{
+		GeneratedStructRegisterFunctionText.Logf(TEXT("\t\tstatic UScriptStruct* ReturnStruct = FindExistingStructIfHotReloadOrDynamic(Outer, TEXT(\"%s\"), sizeof(%s), %s(), false);\r\n"), *ActualStructName, NameLookupCPP.GetNameCPP(Struct), *CRCFuncName);
+	}
+	else
+	{
+		GeneratedStructRegisterFunctionText.Logf(TEXT("\t\tUScriptStruct* ReturnStruct = FindExistingStructIfHotReloadOrDynamic(Outer, TEXT(\"%s\"), sizeof(%s), %s(), true);\r\n"), *ActualStructName, NameLookupCPP.GetNameCPP(Struct), *CRCFuncName);
+	}
+	GeneratedStructRegisterFunctionText.Logf(TEXT("\t\tif (!ReturnStruct)\r\n"));
+	GeneratedStructRegisterFunctionText.Logf(TEXT("\t\t{\r\n"));
+	FString BaseStructString(TEXT("NULL"));
+	if (BaseStruct)
+	{
+		CastChecked<UScriptStruct>(BaseStruct); // this better actually be a script struct
+		BaseStructString = GetSingletonName(BaseStruct);
+	}
+	FString CppStructOpsString(TEXT("NULL"));
+	FString ExplicitSizeString;
+	FString ExplicitAlignmentString;
+	if ((Struct->StructFlags&STRUCT_Native) != 0)
+	{
+		//@todo .u we don't need the auto register versions of these (except for hotreload, which should be fixed)
+		CppStructOpsString = FString::Printf(TEXT("new UScriptStruct::TCppStructOps<%s>"), NameLookupCPP.GetNameCPP(Struct));
+	}
+	else
+	{
+		ExplicitSizeString = FString::Printf(TEXT(", sizeof(%s), ALIGNOF(%s)"), NameLookupCPP.GetNameCPP(Struct), NameLookupCPP.GetNameCPP(Struct));
+	}
 
-		GeneratedPackageCPP.Logf(TEXT("static class UEnum* %s_StaticEnum()\r\n"), *Enum->GetName());
-		GeneratedPackageCPP.Logf(TEXT("{\r\n"));
-		
-		if (!bIsDynamic)
-		{
-			GeneratedPackageCPP.Logf(TEXT("\textern %sclass UPackage* %s;\r\n"), *FriendApiString, *PackageSingletonName);
-			GeneratedPackageCPP.Logf(TEXT("\tstatic class UEnum* Singleton = NULL;\r\n"));
-		}
-		else
-		{
-			GeneratedPackageCPP.Logf(TEXT("\tclass UPackage* EnumPackage = FindOrConstructDynamicTypePackage(TEXT(\"%s\"));\r\n"), *PackageSingletonName);
-			GeneratedPackageCPP.Logf(TEXT("\tclass UEnum* Singleton = Cast<UEnum>(StaticFindObjectFast(UEnum::StaticClass(), EnumPackage, TEXT(\"%s\")));\r\n"), *FNativeClassHeaderGenerator::GetOverriddenName(Enum));
-		}
-		GeneratedPackageCPP.Logf(TEXT("\tif (!Singleton)\r\n"));
-		GeneratedPackageCPP.Logf(TEXT("\t{\r\n"));
-		GeneratedPackageCPP.Logf(TEXT("\t\textern %sclass UEnum* %s;\r\n"), *FriendApiString, *StaticConstructionString);
-		if (!bIsDynamic)
-		{
-			GeneratedPackageCPP.Logf(TEXT("\t\tSingleton = GetStaticEnum(%s, %s, TEXT(\"%s\"));\r\n"), *SingletonName, *PackageSingletonName, *Enum->GetName());
-		}
-		else
-		{
-			GeneratedPackageCPP.Logf(TEXT("\t\tSingleton = GetStaticEnum(%s, EnumPackage, TEXT(\"%s\"));\r\n"), *SingletonName, *FNativeClassHeaderGenerator::GetOverriddenName(Enum));
-		}
-			
-		GeneratedPackageCPP.Logf(TEXT("\t}\r\n"));
-		GeneratedPackageCPP.Logf(TEXT("\treturn Singleton;\r\n"));
-		GeneratedPackageCPP.Logf(TEXT("}\r\n"));
+	const TCHAR* UStructObjectFlags = bIsDynamic ? TEXT("RF_Public|RF_Transient") : TEXT("RF_Public|RF_Transient|RF_MarkAsNative");
+	GeneratedStructRegisterFunctionText.Logf(TEXT("\t\t\tReturnStruct = new(EC_InternalUseOnlyConstructor, Outer, TEXT(\"%s\"), %s) UScriptStruct(FObjectInitializer(), %s, %s, EStructFlags(0x%08X)%s);\r\n"),
+		*ActualStructName,
+		UStructObjectFlags,
+		*BaseStructString,
+		*CppStructOpsString,
+		(uint32)(Struct->StructFlags & ~STRUCT_ComputedFlags),
+		*ExplicitSizeString
+		);
+	TheFlagAudit.Add(Struct, TEXT("StructFlags"), (uint64)(Struct->StructFlags & ~STRUCT_ComputedFlags));
 
-		const FString EnumNameCpp = Enum->GetName(); //UserDefinedEnum should already have a valid cpp name.
-		GeneratedPackageCPP.Logf(TEXT("static FCompiledInDeferEnum Z_CompiledInDeferEnum_UEnum_%s(%s_StaticEnum, TEXT(\"%s\"), TEXT(\"%s\"), %s, %s, %s);\r\n"), 
-			*EnumNameCpp, *EnumNameCpp, *Enum->GetOutermost()->GetName(), *FNativeClassHeaderGenerator::GetOverriddenName(Enum),
-			bIsDynamic ? TEXT("true") : TEXT("false"),
-			bIsDynamic ? *AsTEXT(FClass::GetTypePackageName(Enum)) : TEXT("nullptr"),
-			bIsDynamic ? *AsTEXT(FNativeClassHeaderGenerator::GetOverriddenPathName(Enum)) : TEXT("nullptr"));
-		{
-			const FString EnumSingletonName = GetSingletonName(Enum);
-			GeneratedFunctionDeclarations.Log(FTypeSingletonCache::Get(Enum).GetExternDecl());
+	TArray<UProperty*> Props;
+	for (TFieldIterator<UProperty> ItInner(Struct, EFieldIteratorFlags::ExcludeSuper); ItInner; ++ItInner)
+	{
+		Props.Add(*ItInner);
+	}
+	FString OuterString = FString(TEXT("ReturnStruct"));
+	FString Meta = GetMetaDataCodeForObject(Struct, *OuterString, TEXT("\t\t\t"));
+	OutputProperties(Meta, GeneratedStructRegisterFunctionText, OuterString, Props, TEXT("\t\t\t"));
+	GeneratedStructRegisterFunctionText.Logf(TEXT("\t\t\tReturnStruct->StaticLink();\r\n"));
 
-			FUHTStringBuilder GeneratedEnumRegisterFunctionText;
+	if (Meta.Len())
+	{
+		GeneratedStructRegisterFunctionText.Logf(TEXT("#if WITH_METADATA\r\n"));
+		GeneratedStructRegisterFunctionText.Logf(TEXT("\t\t\tUMetaData* MetaData = ReturnStruct->GetOutermost()->GetMetaData();\r\n"));
+		GeneratedStructRegisterFunctionText.Log(*Meta);
+		GeneratedStructRegisterFunctionText.Logf(TEXT("#endif\r\n"));
+	}
 
-			FString CRCFuncName = FString::Printf(TEXT("Get_%s_CRC"), *SingletonName.Replace(TEXT("()"), TEXT(""), ESearchCase::CaseSensitive));
+	GeneratedStructRegisterFunctionText.Logf(TEXT("\t\t}\r\n"));
+	GeneratedStructRegisterFunctionText.Logf(TEXT("\t\treturn ReturnStruct;\r\n"));
+	GeneratedStructRegisterFunctionText.Logf(TEXT("\t}\r\n"));
 
-			GeneratedEnumRegisterFunctionText.Logf(TEXT("\tUEnum* %s\r\n"), *EnumSingletonName);
-			GeneratedEnumRegisterFunctionText.Logf(TEXT("\t{\r\n"));
-			// Enums can either have a UClass or UPackage as outer (if declared in non-UClass header).
-			if (Enum->GetOuter()->IsA(UStruct::StaticClass()))
-			{
-				GeneratedEnumRegisterFunctionText.Logf(TEXT("\t\tUClass* Outer=%s;\r\n"), *GetSingletonName(CastChecked<UStruct>(Enum->GetOuter())));
-			}
-			else if (!bIsDynamic)
-			{
-				GeneratedEnumRegisterFunctionText.Logf(TEXT("\t\tUPackage* Outer=%s;\r\n"), *GetPackageSingletonName(CastChecked<UPackage>(Enum->GetOuter())));
-			}
-			else
-			{
-				GeneratedEnumRegisterFunctionText.Logf(TEXT("\t\tUPackage* Outer = FindOrConstructDynamicTypePackage(TEXT(\"%s\"));\r\n"), *PackageSingletonName);
-			}
-			GeneratedEnumRegisterFunctionText.Logf(TEXT("\t\textern uint32 %s();\r\n"), *CRCFuncName);
-			if (!bIsDynamic)
-			{
-				GeneratedEnumRegisterFunctionText.Logf(TEXT("\t\tstatic UEnum* ReturnEnum = FindExistingEnumIfHotReloadOrDynamic(Outer, TEXT(\"%s\"), 0, %s(), false);\r\n"), *Enum->GetName(), *CRCFuncName);
-			}
-			else
-			{
-				GeneratedEnumRegisterFunctionText.Logf(TEXT("\t\tUEnum* ReturnEnum = FindExistingEnumIfHotReloadOrDynamic(Outer, TEXT(\"%s\"), 0, %s(), true);\r\n"), *FNativeClassHeaderGenerator::GetOverriddenName(Enum), *CRCFuncName);
-			}
-			GeneratedEnumRegisterFunctionText.Logf(TEXT("\t\tif (!ReturnEnum)\r\n"));
-			GeneratedEnumRegisterFunctionText.Logf(TEXT("\t\t{\r\n"));
+	uint32 StructCrc = GenerateTextCRC(*GeneratedStructRegisterFunctionText);
+	GGeneratedCodeCRCs.Add(Struct, StructCrc);
+	UHTMakefile.AddGeneratedCodeCRC(&SourceFile, Struct, StructCrc);
 
-			const TCHAR* UEnumObjectFlags = bIsDynamic ? TEXT("RF_Public|RF_Transient") : TEXT("RF_Public|RF_Transient|RF_MarkAsNative");
-			GeneratedEnumRegisterFunctionText.Logf(TEXT("\t\t\tReturnEnum = new(EC_InternalUseOnlyConstructor, Outer, TEXT(\"%s\"), %s) UEnum(FObjectInitializer());\r\n"), *FNativeClassHeaderGenerator::GetOverriddenName(Enum), UEnumObjectFlags);
-			GeneratedEnumRegisterFunctionText.Logf(TEXT("\t\t\tTArray<TPair<FName, int64>> EnumNames;\r\n"));
-			for (int32 Index = 0; Index < Enum->NumEnums(); Index++)
-			{
-				const TCHAR* OverridenNameMetaDatakey = TEXT("OverrideName");
-				const FString KeyName = Enum->HasMetaData(OverridenNameMetaDatakey, Index) ? Enum->GetMetaData(OverridenNameMetaDatakey, Index) : Enum->GetNameByIndex(Index).ToString();
-				GeneratedEnumRegisterFunctionText.Logf(TEXT("\t\t\tEnumNames.Emplace(TEXT(\"%s\"), %lld);\r\n"), *KeyName, Enum->GetValueByIndex(Index));
-			}
+	Out.Log(GeneratedStructRegisterFunctionText);
+	Out.Logf(TEXT("\tuint32 %s() { return %uU; }\r\n"), *CRCFuncName, StructCrc);
 
-			FString EnumTypeStr;
-			switch (Enum->GetCppForm())
-			{
-			case UEnum::ECppForm::Regular:    EnumTypeStr = TEXT("UEnum::ECppForm::Regular");    break;
-			case UEnum::ECppForm::Namespaced: EnumTypeStr = TEXT("UEnum::ECppForm::Namespaced"); break;
-			case UEnum::ECppForm::EnumClass:  EnumTypeStr = TEXT("UEnum::ECppForm::EnumClass");  break;
-			}
-			const FString ParamAddMaxKeyIfMissing = FClass::IsDynamic(Enum) ? TEXT(", false") : TEXT("");
-			GeneratedEnumRegisterFunctionText.Logf(TEXT("\t\t\tReturnEnum->SetEnums(EnumNames, %s%s);\r\n"), *EnumTypeStr, *ParamAddMaxKeyIfMissing);
-			GeneratedEnumRegisterFunctionText.Logf(TEXT("\t\t\tReturnEnum->CppType = TEXT(\"%s\");\r\n"), *Enum->CppType);
+	//CallSingletons.Logf(TEXT("\t\t\t\tOuterClass->LinkChild(%s); // %u\r\n"), *SingletonName, StructCrc);
+}
 
-			const FString& EnumDisplayNameFn = Enum->GetMetaData(TEXT("EnumDisplayNameFn"));
-			if( !EnumDisplayNameFn.IsEmpty() )
-			{
-				GeneratedEnumRegisterFunctionText.Logf(TEXT("\t\t\tReturnEnum->SetEnumDisplayNameFn(&%s);\r\n"), *EnumDisplayNameFn);
-			}
+void FNativeClassHeaderGenerator::ExportGeneratedEnumInitCode(FOutputDevice& Out, FOutputDevice& OutDeclarations, const FUnrealSourceFile& SourceFile, UEnum* Enum)
+{
+	const bool bIsDynamic = FClass::IsDynamic(Enum);
+	const FString FriendApiString = GetAPIString();
+	const FString StaticConstructionString = GetSingletonName(Enum);
 
-			FString Meta = GetMetaDataCodeForObject(Enum, TEXT("ReturnEnum"), TEXT("\t\t\t"));
-			if (Meta.Len())
-			{
-				GeneratedEnumRegisterFunctionText.Logf(TEXT("#if WITH_METADATA\r\n"));
-				GeneratedEnumRegisterFunctionText.Logf(TEXT("\t\t\tUMetaData* MetaData = ReturnEnum->GetOutermost()->GetMetaData();\r\n"));
-				GeneratedEnumRegisterFunctionText.Log(*Meta);
-				GeneratedEnumRegisterFunctionText.Logf(TEXT("#endif\r\n"));
-			}
+	FString SingletonName = StaticConstructionString.Replace(TEXT("()"), TEXT(""), ESearchCase::CaseSensitive); // function address
+	FString PackageSingletonName = GetPackageSingletonName(CastChecked<UPackage>(Enum->GetOuter()));
+	if (!bIsDynamic)
+	{
+		PackageSingletonName = GetPackageSingletonName(CastChecked<UPackage>(Enum->GetOuter()));
+	}
+	else
+	{
+		PackageSingletonName = FClass::GetTypePackageName(Enum);
+	}
+
+	Out.Logf(TEXT("static UEnum* %s_StaticEnum()\r\n"), *Enum->GetName());
+	Out.Logf(TEXT("{\r\n"));
+	
+	if (!bIsDynamic)
+	{
+		Out.Logf(TEXT("\textern %sclass UPackage* %s;\r\n"), *FriendApiString, *PackageSingletonName);
+		Out.Logf(TEXT("\tstatic UEnum* Singleton = nullptr;\r\n"));
+	}
+	else
+	{
+		Out.Logf(TEXT("\tclass UPackage* EnumPackage = FindOrConstructDynamicTypePackage(TEXT(\"%s\"));\r\n"), *PackageSingletonName);
+		Out.Logf(TEXT("\tclass UEnum* Singleton = Cast<UEnum>(StaticFindObjectFast(UEnum::StaticClass(), EnumPackage, TEXT(\"%s\")));\r\n"), *FNativeClassHeaderGenerator::GetOverriddenName(Enum));
+	}
+	Out.Logf(TEXT("\tif (!Singleton)\r\n"));
+	Out.Logf(TEXT("\t{\r\n"));
+	Out.Logf(TEXT("\t\textern %sclass UEnum* %s;\r\n"), *FriendApiString, *StaticConstructionString);
+	if (!bIsDynamic)
+	{
+		Out.Logf(TEXT("\t\tSingleton = GetStaticEnum(%s, %s, TEXT(\"%s\"));\r\n"), *SingletonName, *PackageSingletonName, *Enum->GetName());
+	}
+	else
+	{
+		Out.Logf(TEXT("\t\tSingleton = GetStaticEnum(%s, EnumPackage, TEXT(\"%s\"));\r\n"), *SingletonName, *FNativeClassHeaderGenerator::GetOverriddenName(Enum));
+	}
 
-			GeneratedEnumRegisterFunctionText.Logf(TEXT("\t\t}\r\n"));
-			GeneratedEnumRegisterFunctionText.Logf(TEXT("\t\treturn ReturnEnum;\r\n"));
-			GeneratedEnumRegisterFunctionText.Logf(TEXT("\t}\r\n"));
+	Out.Logf(TEXT("\t}\r\n"));
+	Out.Logf(TEXT("\treturn Singleton;\r\n"));
+	Out.Logf(TEXT("}\r\n"));
 
-			auto& GeneratedFunctionText = GetGeneratedFunctionTextDevice();
-			GeneratedFunctionText += GeneratedEnumRegisterFunctionText;
+	const FString EnumNameCpp = Enum->GetName(); //UserDefinedEnum should already have a valid cpp name.
+	Out.Logf(TEXT("static FCompiledInDeferEnum Z_CompiledInDeferEnum_UEnum_%s(%s_StaticEnum, TEXT(\"%s\"), TEXT(\"%s\"), %s, %s, %s);\r\n"),
+		*EnumNameCpp, *EnumNameCpp, *Enum->GetOutermost()->GetName(), *FNativeClassHeaderGenerator::GetOverriddenName(Enum),
+		bIsDynamic ? TEXT("true") : TEXT("false"),
+		bIsDynamic ? *AsTEXT(FClass::GetTypePackageName(Enum)) : TEXT("nullptr"),
+		bIsDynamic ? *AsTEXT(FNativeClassHeaderGenerator::GetOverriddenPathName(Enum)) : TEXT("nullptr"));
 
-			uint32 EnumCrc = GenerateTextCRC(*GeneratedEnumRegisterFunctionText);
-			GGeneratedCodeCRCs.Add(Enum, EnumCrc);
-			UHTMakefile.AddGeneratedCodeCRC(&SourceFile, Enum, EnumCrc);
-			GeneratedFunctionText.Logf(TEXT("\tuint32 %s() { return %uU; }\r\n"), *CRCFuncName, EnumCrc);
-			// CallSingletons.Logf(TEXT("\t\t\t\tOuterClass->LinkChild(%s); // %u\r\n"), *EnumSingletonName, EnumCrc);
-		}
+	const FString EnumSingletonName = GetSingletonName(Enum);
+	OutDeclarations.Log(FTypeSingletonCache::Get(Enum).GetExternDecl());
+
+	FUHTStringBuilder GeneratedEnumRegisterFunctionText;
+
+	FString CRCFuncName = FString::Printf(TEXT("Get_%s_CRC"), *SingletonName.Replace(TEXT("()"), TEXT(""), ESearchCase::CaseSensitive));
+
+	GeneratedEnumRegisterFunctionText.Logf(TEXT("\tUEnum* %s\r\n"), *EnumSingletonName);
+	GeneratedEnumRegisterFunctionText.Logf(TEXT("\t{\r\n"));
+	// Enums can either have a UClass or UPackage as outer (if declared in non-UClass header).
+	if (Enum->GetOuter()->IsA(UStruct::StaticClass()))
+	{
+		GeneratedEnumRegisterFunctionText.Logf(TEXT("\t\tUClass* Outer=%s;\r\n"), *GetSingletonName(CastChecked<UStruct>(Enum->GetOuter())));
 	}
-}
+	else if (!bIsDynamic)
+	{
+		GeneratedEnumRegisterFunctionText.Logf(TEXT("\t\tUPackage* Outer=%s;\r\n"), *GetPackageSingletonName(CastChecked<UPackage>(Enum->GetOuter())));
+	}
+	else
+	{
+		GeneratedEnumRegisterFunctionText.Logf(TEXT("\t\tUPackage* Outer = FindOrConstructDynamicTypePackage(TEXT(\"%s\"));\r\n"), *PackageSingletonName);
+	}
+	GeneratedEnumRegisterFunctionText.Logf(TEXT("\t\textern uint32 %s();\r\n"), *CRCFuncName);
+	if (!bIsDynamic)
+	{
+		GeneratedEnumRegisterFunctionText.Logf(TEXT("\t\tstatic UEnum* ReturnEnum = FindExistingEnumIfHotReloadOrDynamic(Outer, TEXT(\"%s\"), 0, %s(), false);\r\n"), *Enum->GetName(), *CRCFuncName);
+	}
+	else
+	{
+		GeneratedEnumRegisterFunctionText.Logf(TEXT("\t\tUEnum* ReturnEnum = FindExistingEnumIfHotReloadOrDynamic(Outer, TEXT(\"%s\"), 0, %s(), true);\r\n"), *FNativeClassHeaderGenerator::GetOverriddenName(Enum), *CRCFuncName);
+	}
+	GeneratedEnumRegisterFunctionText.Logf(TEXT("\t\tif (!ReturnEnum)\r\n"));
+	GeneratedEnumRegisterFunctionText.Logf(TEXT("\t\t{\r\n"));
 
-void FNativeClassHeaderGenerator::ExportMirrorsForNoexportStructs(const TArray<UScriptStruct*>& NativeStructs, int32 TextIndent, FUHTStringBuilder& HeaderOutput)
-{
-	// reverse the order.
-	for (int32 i = NativeStructs.Num() - 1; i >= 0; --i)
+	const TCHAR* UEnumObjectFlags = bIsDynamic ? TEXT("RF_Public|RF_Transient") : TEXT("RF_Public|RF_Transient|RF_MarkAsNative");
+	GeneratedEnumRegisterFunctionText.Logf(TEXT("\t\t\tReturnEnum = new(EC_InternalUseOnlyConstructor, Outer, TEXT(\"%s\"), %s) UEnum(FObjectInitializer());\r\n"), *FNativeClassHeaderGenerator::GetOverriddenName(Enum), UEnumObjectFlags);
+	GeneratedEnumRegisterFunctionText.Logf(TEXT("\t\t\tTArray<TPair<FName, int64>> EnumNames;\r\n"));
+	for (int32 Index = 0; Index < Enum->NumEnums(); Index++)
 	{
-		UScriptStruct* Struct = NativeStructs[i];
+		const TCHAR* OverridenNameMetaDatakey = TEXT("OverrideName");
+		const FString KeyName = Enum->HasMetaData(OverridenNameMetaDatakey, Index) ? Enum->GetMetaData(OverridenNameMetaDatakey, Index) : Enum->GetNameByIndex(Index).ToString();
+		GeneratedEnumRegisterFunctionText.Logf(TEXT("\t\t\tEnumNames.Emplace(TEXT(\"%s\"), %lld);\r\n"), *KeyName, Enum->GetValueByIndex(Index));
+	}
 
-		// Export struct.
-		const TCHAR* StructName = NameLookupCPP.GetNameCPP(Struct);
-		HeaderOutput.Logf(TEXT("%sstruct %s"), FCString::Tab(TextIndent), StructName);
-		if (Struct->GetSuperStruct() != NULL)
-		{
-			HeaderOutput.Logf(TEXT(" : public %s"), NameLookupCPP.GetNameCPP(Struct->GetSuperStruct()));
-		}
-		HeaderOutput.Logf(TEXT("\r\n%s{\r\n"), FCString::Tab(TextIndent));
+	FString EnumTypeStr;
+	switch (Enum->GetCppForm())
+	{
+		case UEnum::ECppForm::Regular:    EnumTypeStr = TEXT("UEnum::ECppForm::Regular");    break;
+		case UEnum::ECppForm::Namespaced: EnumTypeStr = TEXT("UEnum::ECppForm::Namespaced"); break;
+		case UEnum::ECppForm::EnumClass:  EnumTypeStr = TEXT("UEnum::ECppForm::EnumClass");  break;
+	}
+	const FString ParamAddMaxKeyIfMissing = FClass::IsDynamic(Enum) ? TEXT(", false") : TEXT("");
+	GeneratedEnumRegisterFunctionText.Logf(TEXT("\t\t\tReturnEnum->SetEnums(EnumNames, %s%s);\r\n"), *EnumTypeStr, *ParamAddMaxKeyIfMissing);
+	GeneratedEnumRegisterFunctionText.Logf(TEXT("\t\t\tReturnEnum->CppType = TEXT(\"%s\");\r\n"), *Enum->CppType);
 
-		// Export the struct's CPP properties.
-		ExportProperties(Struct, TextIndent, /*bAccessSpecifiers=*/ false, &HeaderOutput);
+	const FString& EnumDisplayNameFn = Enum->GetMetaData(TEXT("EnumDisplayNameFn"));
+	if( !EnumDisplayNameFn.IsEmpty() )
+	{
+		GeneratedEnumRegisterFunctionText.Logf(TEXT("\t\t\tReturnEnum->SetEnumDisplayNameFn(&%s);\r\n"), *EnumDisplayNameFn);
+	}
 
-		HeaderOutput.Logf(TEXT("%s};\r\n\r\n"), FCString::Tab(TextIndent));
+	FString Meta = GetMetaDataCodeForObject(Enum, TEXT("ReturnEnum"), TEXT("\t\t\t"));
+	if (Meta.Len())
+	{
+		GeneratedEnumRegisterFunctionText.Logf(TEXT("#if WITH_METADATA\r\n"));
+		GeneratedEnumRegisterFunctionText.Logf(TEXT("\t\t\tUMetaData* MetaData = ReturnEnum->GetOutermost()->GetMetaData();\r\n"));
+		GeneratedEnumRegisterFunctionText.Log(*Meta);
+		GeneratedEnumRegisterFunctionText.Logf(TEXT("#endif\r\n"));
 	}
+
+	GeneratedEnumRegisterFunctionText.Logf(TEXT("\t\t}\r\n"));
+	GeneratedEnumRegisterFunctionText.Logf(TEXT("\t\treturn ReturnEnum;\r\n"));
+	GeneratedEnumRegisterFunctionText.Logf(TEXT("\t}\r\n"));
+
+	Out.Log(GeneratedEnumRegisterFunctionText);
+
+	uint32 EnumCrc = GenerateTextCRC(*GeneratedEnumRegisterFunctionText);
+	GGeneratedCodeCRCs.Add(Enum, EnumCrc);
+	UHTMakefile.AddGeneratedCodeCRC(&SourceFile, Enum, EnumCrc);
+	Out.Logf(TEXT("\tuint32 %s() { return %uU; }\r\n"), *CRCFuncName, EnumCrc);
+	// CallSingletons.Logf(TEXT("\t\t\t\tOuterClass->LinkChild(%s); // %u\r\n"), *EnumSingletonName, EnumCrc);
+}
+
+void FNativeClassHeaderGenerator::ExportMirrorsForNoexportStruct(FOutputDevice& Out, UScriptStruct* Struct, int32 TextIndent)
+{
+	// Export struct.
+	const TCHAR* StructName = NameLookupCPP.GetNameCPP(Struct);
+	Out.Logf(TEXT("%sstruct %s"), FCString::Tab(TextIndent), StructName);
+	if (Struct->GetSuperStruct() != NULL)
+	{
+		Out.Logf(TEXT(" : public %s"), NameLookupCPP.GetNameCPP(Struct->GetSuperStruct()));
+	}
+	Out.Logf(TEXT("\r\n%s{\r\n"), FCString::Tab(TextIndent));
+
+	// Export the struct's CPP properties.
+	ExportProperties(Out, Struct, TextIndent);
+
+	Out.Logf(TEXT("%s};\r\n\r\n"), FCString::Tab(TextIndent));
 }
 
 bool FNativeClassHeaderGenerator::WillExportEventParms( UFunction* Function )
@@ -3324,7 +3039,7 @@ bool FNativeClassHeaderGenerator::WillExportEventParms( UFunction* Function )
   return It && (It->PropertyFlags&CPF_Parm);
 }
 
-void WriteEventFunctionPrologue(FUHTStringBuilder& Output, int32 Indent, const FParmsAndReturnProperties& Parameters, UObject* FunctionOuter, const TCHAR* FunctionName)
+void WriteEventFunctionPrologue(FOutputDevice& Output, int32 Indent, const FParmsAndReturnProperties& Parameters, UObject* FunctionOuter, const TCHAR* FunctionName)
 {
 	// now the body - first we need to declare a struct which will hold the parameters for the event/delegate call
 	Output.Logf(TEXT("\r\n%s{\r\n"), FCString::Tab(Indent));
@@ -3360,7 +3075,7 @@ void WriteEventFunctionPrologue(FUHTStringBuilder& Output, int32 Indent, const F
 	}
 }
 
-void WriteEventFunctionEpilogue(FUHTStringBuilder& Output, int32 Indent, const FParmsAndReturnProperties& Parameters, const TCHAR* FunctionName)
+void WriteEventFunctionEpilogue(FOutputDevice& Output, int32 Indent, const FParmsAndReturnProperties& Parameters, const TCHAR* FunctionName)
 {
 	// Out parm copying.
 	for (auto It = Parameters.Parms.CreateConstIterator(); It; ++It)
@@ -3391,507 +3106,100 @@ void WriteEventFunctionEpilogue(FUHTStringBuilder& Output, int32 Indent, const F
 	Output.Logf(TEXT("%s}\r\n"), FCString::Tab(Indent));
 }
 
-void FNativeClassHeaderGenerator::ExportDelegateDefinitions(FUnrealSourceFile& SourceFile, const TArray<UDelegateFunction*>& DelegateFunctions, const bool bWrapperImplementationsOnly)
-{
-	FUHTStringBuilder HeaderOutput;
-
-	for ( int32 i = DelegateFunctions.Num() - 1; i >= 0 ; i-- )
-	{
-		UFunction* Function = DelegateFunctions[i];
-		FUHTStringBuilder DelegateOutput;
-		check( Function->HasAnyFunctionFlags( FUNC_Delegate ) );
-
-		if (bWrapperImplementationsOnly)
-		{
-			// Export parameters structs for all delegates.  We'll need these to declare our delegate execution function.
-			ExportEventParm(Function, DelegateOutput);
-		}
-
-		const bool bIsMulticastDelegate = Function->HasAnyFunctionFlags( FUNC_MulticastDelegate );
-
-		// Unmangle the function name
-		const FString DelegateName = Function->GetName().LeftChop( FString( HEADER_GENERATED_DELEGATE_SIGNATURE_SUFFIX ).Len() );
-
-		auto* CompilerInfo = FFunctionData::FindForFunction(Function);
-
-		FFuncInfo FunctionData = CompilerInfo->GetFunctionData();
-
-		if( bWrapperImplementationsOnly )
-		{
-			// Always export delegate wrapper functions as inline
-			FunctionData.FunctionExportFlags |= FUNCEXPORT_Inline;
-		}
-
-		// Add class name to beginning of function, to avoid collisions with other classes with the same delegate name in this scope
-		FString Delegate(TEXT("delegate"));
-		check(FunctionData.MarshallAndCallName.StartsWith(Delegate));
-		FString ShortName = *FunctionData.MarshallAndCallName + Delegate.Len();
-		FunctionData.MarshallAndCallName = FString::Printf( TEXT( "F%s_DelegateWrapper" ), *ShortName );
-
-		// Setup delegate parameter
-		const FString ExtraParam( FString::Printf( TEXT( "const %s& %s" ),
-			bIsMulticastDelegate ? TEXT( "FMulticastScriptDelegate" ) : TEXT( "FScriptDelegate" ),
-			*DelegateName ) );
-
-		DelegateOutput.Log(TEXT("static "));
-
-		// export the line that looks like: int32 Main(const FString& Parms)
-		ExportNativeFunctionHeader(FunctionData, DelegateOutput, EExportFunctionType::Event, EExportFunctionHeaderStyle::Declaration, *ExtraParam);
-
-		if( !bWrapperImplementationsOnly )
-		{
-			// Only exporting function prototype
-			DelegateOutput.Logf(TEXT(";\r\n"));
-
-			ExportFunction(SourceFile, Function, &SourceFile.GetScope().Get(), false);
-		}
-		else
-		{
-			auto Parameters = GetFunctionParmsAndReturn(FunctionData.FunctionReference);
-
-			WriteEventFunctionPrologue(DelegateOutput, 0, Parameters, Function->GetOuter(), *DelegateName);
-			{
-				const TCHAR* DelegateType = bIsMulticastDelegate ? TEXT( "ProcessMulticastDelegate" ) : TEXT( "ProcessDelegate" );
-				const TCHAR* DelegateArg  = Parameters.HasParms() ? TEXT("&Parms") : TEXT("NULL");
-				DelegateOutput.Logf(TEXT("\t%s.%s<UObject>(%s);\r\n"), *DelegateName, DelegateType, DelegateArg);
-			}
-			WriteEventFunctionEpilogue(DelegateOutput, 0, Parameters, *DelegateName);
-		}
-
-		if (bWrapperImplementationsOnly)
-		{
-			FString MacroName = SourceFile.GetGeneratedMacroName(FunctionData.MacroLine, TEXT("_DELEGATE"));
-			WriteMacro(HeaderOutput, MacroName, DelegateOutput);
-		}
-	}
-
-	if (HeaderOutput.Len())
-	{
-		GeneratedHeaderText.Log(*HeaderOutput);
-		GeneratedHeaderText.Log(TEXT("\r\n\r\n"));
-	}
-}
-
-/**
- * Export a single .proto declaration, recursing as necessary for sub declarations
- *
- * @param Out output device
- * @param MessageName name of the message in the declaration
- * @param Properties array of parameters in the function definition
- * @param PropertyFlags flags to filter property array against
- * @param Ident starting indentation level
- */
-void ExportProtoDeclaration(FOutputDevice& Out, const FString& MessageName, TFieldIterator<UProperty>& Properties, uint64 PropertyFlags, int32 Indent)
+void FNativeClassHeaderGenerator::ExportDelegateDeclaration(FOutputDevice& Out, FOutputDevice& OutDeclarations, const FUnrealSourceFile& SourceFile, UFunction* Function)
 {
-	Out.Logf(TEXT("%smessage CMsg%sMessage\r\n"), FCString::Tab(Indent), *MessageName);
-	Out.Logf(TEXT("%s{\r\n"), FCString::Tab(Indent));
+	static const TCHAR DelegateStr[] = TEXT("delegate");
 
-	static TMap<FString, FString> ToProtoTypeMappings;
-	static bool bInitMapping = false;
-	if (!bInitMapping)
-	{
-		// Explicit type mappings
-		ToProtoTypeMappings.Add(TEXT("FString"), TEXT("string"));
-		ToProtoTypeMappings.Add(TEXT("int32"), TEXT("int32"));
-		ToProtoTypeMappings.Add(TEXT("int64"), TEXT("int64"));
-		ToProtoTypeMappings.Add(TEXT("uint8"), TEXT("bytes"));
-		ToProtoTypeMappings.Add(TEXT("bool"), TEXT("bool"));
-		ToProtoTypeMappings.Add(TEXT("double"), TEXT("double"));
-		ToProtoTypeMappings.Add(TEXT("float"), TEXT("float"));
-		bInitMapping = true;
-	}
+	check(Function->HasAnyFunctionFlags(FUNC_Delegate));
 
-	int32 FieldIdx = 1;
-	for(; Properties && (Properties->PropertyFlags & PropertyFlags); ++Properties)
-	{
-		UProperty* Property = *Properties;
-		UClass* PropClass = Property->GetClass();
+	const bool bIsMulticastDelegate = Function->HasAnyFunctionFlags( FUNC_MulticastDelegate );
 
-		// Skip out and return paramaters
-		if ((Property->PropertyFlags & CPF_RepSkip) || (Property->PropertyFlags & CPF_ReturnParm))
-		{
-			continue;
-		}
+	// Unmangle the function name
+	const FString DelegateName = Function->GetName().LeftChop( FString( HEADER_GENERATED_DELEGATE_SIGNATURE_SUFFIX ).Len() );
 
-		// export the property type text (e.g. FString; int32; TArray, etc.)
-		FString TypeText, ExtendedTypeText;
-		TypeText = Property->GetCPPType(&ExtendedTypeText, CPPF_None);
-		if (PropClass != UInterfaceProperty::StaticClass() && PropClass != UObjectProperty::StaticClass())
-		{
-			bool bIsRepeated = false;
-			if (UArrayProperty* ArrayProperty = Cast<UArrayProperty>(Property))
-			{
-				UClass* InnerPropClass = ArrayProperty->Inner->GetClass();
-				if (InnerPropClass != UInterfaceProperty::StaticClass() && InnerPropClass != UObjectProperty::StaticClass())
-				{
-					FString InnerExtendedTypeText;
-					FString InnerTypeText = ArrayProperty->Inner->GetCPPType(&InnerExtendedTypeText, CPPF_None);
-					TypeText = InnerTypeText;
-					ExtendedTypeText = InnerExtendedTypeText;
-					Property = ArrayProperty->Inner;
-					bIsRepeated = true;
-				}
-				else
-				{
-					FError::Throwf(TEXT("ExportProtoDeclaration - Unhandled property type '%s': %s"), *PropClass->GetName(), *Property->GetPathName());
-				}
-			}
-			else if (Property->IsA<UMapProperty>())
-			{
-				FError::Throwf(TEXT("ExportProtoDeclaration - Map properties not yet supported '%s': %s"), *PropClass->GetName(), *Property->GetPathName());
-			}
-			else if (Property->IsA<USetProperty>())
-			{
-				FError::Throwf(TEXT("ExportProtoDeclaration - Set properties not yet supported '%s': %s"), *PropClass->GetName(), *Property->GetPathName());
-			}
-			else if(Property->ArrayDim != 1)
-			{
-				bIsRepeated = true;
-			}
+	const FFunctionData* CompilerInfo = FFunctionData::FindForFunction(Function);
 
-			FString VariableTypeName = FString::Printf(TEXT("%s%s"), *TypeText, *ExtendedTypeText);
-			FString* ProtoTypeName = NULL;
+	FFuncInfo FunctionData = CompilerInfo->GetFunctionData();
 
-			UStructProperty* StructProperty = Cast<UStructProperty>(Property);
-			if (StructProperty != NULL)
-			{
-				TFieldIterator<UProperty> StructIt(StructProperty->Struct);
-				ExportProtoDeclaration(Out, VariableTypeName, StructIt, CPF_AllFlags, Indent + 1);
-				VariableTypeName = FString::Printf(TEXT("CMsg%sMessage"), *VariableTypeName);
-				ProtoTypeName = &VariableTypeName;
-			}
-			else
-			{
-				ProtoTypeName = ToProtoTypeMappings.Find(VariableTypeName);
-			}
+	// Add class name to beginning of function, to avoid collisions with other classes with the same delegate name in this scope
+	check(FunctionData.MarshallAndCallName.StartsWith(DelegateStr));
+	FString ShortName = *FunctionData.MarshallAndCallName + ARRAY_COUNT(DelegateStr) - 1;
+	FunctionData.MarshallAndCallName = FString::Printf( TEXT( "F%s_DelegateWrapper" ), *ShortName );
 
-			Out.Log(FCString::Tab(Indent + 1));
-			Out.Log(bIsRepeated ? TEXT("repeated ") : TEXT("optional "));
+	// Setup delegate parameter
+	const FString ExtraParam = FString::Printf(
+		TEXT( "const %s& %s" ),
+		bIsMulticastDelegate ? TEXT( "FMulticastScriptDelegate" ) : TEXT( "FScriptDelegate" ),
+		*DelegateName
+	);
 
-			if (ProtoTypeName != NULL)
-			{
-				Out.Logf( TEXT("%s %s = %d;\r\n"), **ProtoTypeName, *Property->GetNameCPP(), FieldIdx);
-				FieldIdx++;
-			}
-			else
-			{
-				FError::Throwf(TEXT("ExportProtoDeclaration - Unhandled property mapping '%s': %s"), *PropClass->GetName(), *Property->GetPathName());
-			}
-		}
-		else
-		{
-			FError::Throwf(TEXT("ExportProtoDeclaration - Unhandled property type '%s': %s"), *PropClass->GetName(), *Property->GetPathName());
-		}
-	}
-
-	Out.Logf(TEXT("%s}\r\n"), FCString::Tab(Indent));
-}
-
-/**
- * Generate a .proto message declaration for any functions marked as requiring one
- * 
- * @param InCallbackFunctions array of functions for consideration to generate .proto definitions
- * @param Indent starting indentation level
- * @param Output optional output redirect
- */
-void FNativeClassHeaderGenerator::ExportProtoMessage(const TArray<UFunction*>& InCallbackFunctions, int32 Indent, FUHTStringBuilder* Output)
-{
-	// Parms struct definitions.
-	FUHTStringBuilder HeaderOutput;
-
-	TArray<UFunction*> CallbackFunctions = InCallbackFunctions;
-	CallbackFunctions.Sort();
-
-	for (int32 Index = 0; Index < CallbackFunctions.Num(); Index++)
-	{
-		UFunction* Function = CallbackFunctions[Index];
-		auto* CompilerInfo = FFunctionData::FindForFunction(Function);
-
-		const FFuncInfo& FunctionData = CompilerInfo->GetFunctionData();
-		if (FunctionData.FunctionExportFlags & FUNCEXPORT_NeedsProto)
-		{
-			if (WillExportEventParms(Function) && !Function->HasAnyFunctionFlags(FUNC_Delegate))
-			{
-				FString FunctionName = Function->GetName();
-				TFieldIterator<UProperty> CommentIt(Function);
-				FString ParameterList;
-				for(; CommentIt && (CommentIt->PropertyFlags & CPF_Parm); ++CommentIt)
-				{
-					UProperty* Param = *CommentIt;
-					ForwardDeclarations.Add(Param);
-					FString TypeText, ExtendedTypeText;
-					TypeText = Param->GetCPPType(&ExtendedTypeText, CPPF_None);
-					FString ParamName = FString::Printf(TEXT("%s%s %s"), *TypeText, *ExtendedTypeText, *Param->GetName());
-
-					// add this property to the parameter list string
-					if (ParameterList.Len())
-					{
-						ParameterList += TCHAR(',');
-					}
-
-					ParameterList += ParamName;
-				}
-
-				HeaderOutput.Logf(TEXT("// %s%s(%s)\r\n"), FCString::Tab(Indent), *FunctionName, *ParameterList);
-
-				TFieldIterator<UProperty> ParamIt(Function);
-				ExportProtoDeclaration(HeaderOutput, FunctionName, ParamIt, CPF_Parm, Indent);
-			}
-		}
-	}
-
-	if (!Output)
-	{
-		GeneratedProtoText.Log(*HeaderOutput);
-	}
-	else
-	{
-		Output->Log(HeaderOutput);
-	}
-}
+	FUHTStringBuilder DelegateOutput;
+	DelegateOutput.Log(TEXT("static "));
 
-// Java uses different coding standards for capitalization
-static FString FixJavaName(const FString &StringIn)
-{
-	FString FixedString = StringIn;
+	// export the line that looks like: int32 Main(const FString& Parms)
+	ExportNativeFunctionHeader(DelegateOutput, ForwardDeclarations, FunctionData, EExportFunctionType::Event, EExportFunctionHeaderStyle::Declaration, *ExtraParam, *GetAPIString());
 
-	FixedString[0] = FChar::ToLower(FixedString[0]); // java classes/variable start lower case
-	FixedString.ReplaceInline(TEXT("ID"), TEXT("Id"), ESearchCase::CaseSensitive); // Id is standard instead of ID, some of our fnames use ID
+	// Only exporting function prototype
+	DelegateOutput.Logf(TEXT(";\r\n"));
 
-	return FixedString;
+	OutDeclarations.Log(FTypeSingletonCache::Get(Function).GetExternDecl());
+	ExportFunction(Out, SourceFile, Function, false);
 }
 
-/**
- * Export a single .java declaration, recursing as necessary for sub declarations
- *
- * @param Out output device
- * @param MessageName name of the message in the declaration
- * @param Properties array of parameters in the function definition
- * @param PropertyFlags flags to filter property array against
- * @param Ident starting indentation level
- */
-void ExportMCPDeclaration(FOutputDevice& Out, const FString& MessageName, TFieldIterator<UProperty>& Properties, uint64 PropertyFlags, int32 Indent)
+void FNativeClassHeaderGenerator::ExportDelegateDefinition(FOutputDevice& Out, const FUnrealSourceFile& SourceFile, UFunction* Function)
 {
-	Out.Logf(TEXT("%spublic class %sCommand extends ProfileCommand\r\n"), FCString::Tab(Indent), *MessageName);
-	Out.Logf(TEXT("%s{\r\n"), FCString::Tab(Indent));
+	static const TCHAR DelegateStr[] = TEXT("delegate");
 
-	static TMap<FString, FString> ToMCPTypeMappings;
-	static TMap<FString, FString> ToAnnotationMappings;
-	static bool bInitMapping = false;
-	if (!bInitMapping)
-	{
-		// Explicit type mappings
-		ToMCPTypeMappings.Add(TEXT("FString"), TEXT("String"));
-		ToMCPTypeMappings.Add(TEXT("int32"), TEXT("int"));
-		ToMCPTypeMappings.Add(TEXT("int64"), TEXT("int"));
-		ToMCPTypeMappings.Add(TEXT("uint8"), TEXT("byte"));
-		ToMCPTypeMappings.Add(TEXT("bool"), TEXT("boolean"));
-		ToMCPTypeMappings.Add(TEXT("double"), TEXT("double"));
-		ToMCPTypeMappings.Add(TEXT("float"), TEXT("float"));
-		ToMCPTypeMappings.Add(TEXT("byte"), TEXT("byte"));
+	check(Function->HasAnyFunctionFlags(FUNC_Delegate));
 
-		ToAnnotationMappings.Add(TEXT("FString"), TEXT("@NotBlankOrNull"));
+	// Export parameters structs for all delegates.  We'll need these to declare our delegate execution function.
+	FUHTStringBuilder DelegateOutput;
+	ExportEventParm(DelegateOutput, ForwardDeclarations, Function, /*Indent=*/ 0, /*bOutputConstructor=*/ true, EExportingState::Normal);
 
-		bInitMapping = true;
-	}
+	const bool bIsMulticastDelegate = Function->HasAnyFunctionFlags( FUNC_MulticastDelegate );
 
-	FString ConstructorParams;
-	FString ConstructorText;
+	// Unmangle the function name
+	const FString DelegateName = Function->GetName().LeftChop( FString( HEADER_GENERATED_DELEGATE_SIGNATURE_SUFFIX ).Len() );
 
-	for(; Properties && (Properties->PropertyFlags & PropertyFlags); ++Properties)
-	{
-		UProperty* Property = *Properties;
-		UClass* PropClass = Property->GetClass();
+	const FFunctionData* CompilerInfo = FFunctionData::FindForFunction(Function);
 
-		// Skip out and return paramaters
-		if ((Property->PropertyFlags & CPF_RepSkip) || (Property->PropertyFlags & CPF_ReturnParm))
-		{
-			continue;
-		}
+	FFuncInfo FunctionData = CompilerInfo->GetFunctionData();
 
-		// export the property type text (e.g. FString; int32; TArray, etc.)
-		FString TypeText, ExtendedTypeText;
-		TypeText = Property->GetCPPType(&ExtendedTypeText, CPPF_None);
-		if (PropClass != UInterfaceProperty::StaticClass() && PropClass != UObjectProperty::StaticClass())
-		{
-			// TODO Implement arrays
-			UArrayProperty* ArrayProperty = Cast<UArrayProperty>(Property);
-			if (ArrayProperty != NULL)
-			{
-				// skip array generation for Java, this should result in a List<TYPE> declaration, but we can do this by hand for now.
-				continue;
-			}
-			// TODO Implement maps
-			UMapProperty* MapProperty = Cast<UMapProperty>(Property);
-			if (MapProperty != NULL)
-			{
-				continue;
-			}
-			// TODO Implement sets
-			USetProperty* SetProperty = Cast<USetProperty>(Property);
-			if (SetProperty != NULL)
-			{
-				continue;
-			}
-			/*bool bIsRepeated = false;
-			UArrayProperty* ArrayProperty = Cast<UArrayProperty>(Property);
-			if (ArrayProperty != NULL)
-			{
-				UClass* InnerPropClass = ArrayProperty->Inner->GetClass();
-				if (InnerPropClass != UInterfaceProperty::StaticClass() && InnerPropClass != UObjectProperty::StaticClass())
-				{
-					FString InnerExtendedTypeText;
-					FString InnerTypeText = ArrayProperty->Inner->GetCPPType(&InnerExtendedTypeText, CPPF_None);
-					TypeText = InnerTypeText;
-					ExtendedTypeText = InnerExtendedTypeText;
-					Property = ArrayProperty->Inner;
-					bIsRepeated = true;
-				}
-				else
-				{
-					FError::Throwf(TEXT("ExportMCPDeclaration - Unhandled property type '%s': %s"), *PropClass->GetName(), *Property->GetPathName());
-				}
-			}
-			else if(Property->ArrayDim != 1)
-			{
-				bIsRepeated = true;
-			}*/
+	// Always export delegate wrapper functions as inline
+	FunctionData.FunctionExportFlags |= FUNCEXPORT_Inline;
 
-			FString VariableTypeName = FString::Printf(TEXT("%s%s"), *TypeText, *ExtendedTypeText);
-			FString PropertyName = FixJavaName(Property->GetNameCPP());
-			FString* MCPTypeName = NULL;
-			FString* AnnotationName = NULL;
+	// Add class name to beginning of function, to avoid collisions with other classes with the same delegate name in this scope
+	check(FunctionData.MarshallAndCallName.StartsWith(DelegateStr));
+	FString ShortName = *FunctionData.MarshallAndCallName + ARRAY_COUNT(DelegateStr) - 1;
+	FunctionData.MarshallAndCallName = FString::Printf( TEXT( "F%s_DelegateWrapper" ), *ShortName );
 
+	// Setup delegate parameter
+	const FString ExtraParam = FString::Printf(
+		TEXT( "const %s& %s" ),
+		bIsMulticastDelegate ? TEXT( "FMulticastScriptDelegate" ) : TEXT( "FScriptDelegate" ),
+		*DelegateName
+	);
 
-			UStructProperty* StructProperty = Cast<UStructProperty>(Property);
-			if (StructProperty != NULL)
-			{
-				// TODO Implement structs
-/*				TFieldIterator<UProperty> StructIt(StructProperty->Struct);
-				ExportMCPDeclaration(Out, VariableTypeName, StructIt, CPF_AllFlags, Indent + 1);
-				VariableTypeName = FString::Printf(TEXT("CMsg%sMessage"), *VariableTypeName);
-				MCPTypeName = &VariableTypeName;*/
-				continue;
-			}
-			else
-			{
-				UByteProperty* ByteProperty = Cast<UByteProperty>(Property);
-				if ((ByteProperty && ByteProperty->Enum) || Cast<UEnumProperty>(Property))
-				{
-					// treat enums like strings because that's how they'll be exported in JSON
-					MCPTypeName = ToMCPTypeMappings.Find(TEXT("FString"));
-					AnnotationName = ToAnnotationMappings.Find(TEXT("FString"));
-				}
-				else
-				{
-					MCPTypeName = ToMCPTypeMappings.Find(VariableTypeName);
-					AnnotationName = ToAnnotationMappings.Find(VariableTypeName);
-				}
-			}
-
-			if (AnnotationName != NULL && !AnnotationName->IsEmpty())
-			{
-				Out.Log(FCString::Tab(Indent + 1));
-				Out.Logf(TEXT("%s\r\n"), **AnnotationName);
-			}
-
-			if (MCPTypeName != NULL)
-			{
-				Out.Log(FCString::Tab(Indent + 1));
-				Out.Logf(TEXT("private %s %s;\r\n"), **MCPTypeName, *PropertyName);
-
-				ConstructorParams += FString::Printf(TEXT(", %s %s"), **MCPTypeName, *PropertyName);
-				ConstructorText += FString::Printf(TEXT("%sthis.%s = %s;\r\n"), FCString::Tab(Indent + 2), *PropertyName, *PropertyName);
-			}
-			else
-			{
-				FError::Throwf(TEXT("ExportMCPDeclaration - Unhandled property mapping '%s' (%s): %s"), *PropClass->GetName(), *VariableTypeName, *Property->GetPathName());
-			}
-		}
-		else
-		{
-			FError::Throwf(TEXT("ExportMCPDeclaration - Unhandled property type '%s': %s"), *PropClass->GetName(), *Property->GetPathName());
-		}
-	}
-
-	Out.Logf(TEXT("\r\n%spublic %sCommand(String epicId, String profileId%s)\r\n"), FCString::Tab(Indent + 1), *MessageName, *ConstructorParams);
-	Out.Logf(TEXT("%s{\r\n"), FCString::Tab(Indent + 1));
-	Out.Logf(TEXT("%ssuper(epicId, profileId);\r\n"), FCString::Tab(Indent + 2));
-	Out.Logf(TEXT("%s"), *ConstructorText);
-	Out.Logf(TEXT("%s}\r\n"), FCString::Tab(Indent + 1));
-
-	Out.Logf(TEXT("\r\n%s@Override\r\n"), FCString::Tab(Indent + 1));
-	Out.Logf(TEXT("%sprotected void execute(@Name(\"profile\") @NotNull ProfileEx profile)\r\n"), FCString::Tab(Indent + 1));
-	Out.Logf(TEXT("%s{\r\n"), FCString::Tab(Indent + 1));
-	Out.Logf(TEXT("%s}\r\n"), FCString::Tab(Indent + 1));
-
-	Out.Logf(TEXT("%s}\r\n"), FCString::Tab(Indent));
-}
-
-/**
- * Generate a .MCP message declaration for any functions marked as requiring one
- * 
- * @param InCallbackFunctions array of functions for consideration to generate .proto definitions
- * @param Indent starting indentation level
- * @param Output optional output redirect
- */
-void FNativeClassHeaderGenerator::ExportMCPMessage(const TArray<UFunction*>& InCallbackFunctions, FClassMetaData* ClassData, int32 Indent, FUHTStringBuilder* Output)
-{
-	// Parms struct definitions.
-	FUHTStringBuilder HeaderOutput;
-
-	TArray<UFunction*> CallbackFunctions = InCallbackFunctions;
-	CallbackFunctions.Sort();
-
-	for (int32 Index = 0; Index < CallbackFunctions.Num(); Index++)
-	{
-		UFunction* Function = CallbackFunctions[Index];
-		auto* CompilerInfo = FFunctionData::FindForFunction(Function);
-
-		const FFuncInfo& FunctionData = CompilerInfo->GetFunctionData();
-		if (FunctionData.FunctionExportFlags & FUNCEXPORT_NeedsMCP)
-		{
-			if (WillExportEventParms(Function) && !Function->HasAnyFunctionFlags(FUNC_Delegate))
-			{
-				FString FunctionName = Function->GetName();
-				TFieldIterator<UProperty> CommentIt(Function);
-				FString ParameterList;
-				for(; CommentIt && (CommentIt->PropertyFlags & CPF_Parm); ++CommentIt)
-				{
-					UProperty* Param = *CommentIt;
-					ForwardDeclarations.Add(Param);
-					FString TypeText, ExtendedTypeText;
-					TypeText = Param->GetCPPType(&ExtendedTypeText, CPPF_None);
-					FString ParamName = FString::Printf(TEXT("%s%s %s"), *TypeText, *ExtendedTypeText, *Param->GetName());
-
-					// add this property to the parameter list string
-					if (ParameterList.Len())
-					{
-						ParameterList += TCHAR(',');
-					}
-
-					ParameterList += ParamName;
-				}
+	DelegateOutput.Log(TEXT("static "));
 
-				HeaderOutput.Logf(TEXT("// %s%s(%s)\r\n"), FCString::Tab(Indent), *FunctionName, *ParameterList);
+	// export the line that looks like: int32 Main(const FString& Parms)
+	ExportNativeFunctionHeader(DelegateOutput, ForwardDeclarations, FunctionData, EExportFunctionType::Event, EExportFunctionHeaderStyle::Declaration, *ExtraParam, *GetAPIString());
 
-				TFieldIterator<UProperty> ParamIt(Function);
-				ExportMCPDeclaration(HeaderOutput, FunctionName, ParamIt, CPF_Parm, Indent);
-			}
-		}
-	}
+	FParmsAndReturnProperties Parameters = GetFunctionParmsAndReturn(FunctionData.FunctionReference);
 
-	if (!Output)
-	{
-		GeneratedMCPText.Log(*HeaderOutput);
-	}
-	else
+	WriteEventFunctionPrologue(DelegateOutput, 0, Parameters, Function->GetOuter(), *DelegateName);
 	{
-		Output->Log(HeaderOutput);
+		const TCHAR* DelegateType = bIsMulticastDelegate ? TEXT( "ProcessMulticastDelegate" ) : TEXT( "ProcessDelegate" );
+		const TCHAR* DelegateArg  = Parameters.HasParms() ? TEXT("&Parms") : TEXT("NULL");
+		DelegateOutput.Logf(TEXT("\t%s.%s<UObject>(%s);\r\n"), *DelegateName, DelegateType, DelegateArg);
 	}
+	WriteEventFunctionEpilogue(DelegateOutput, 0, Parameters, *DelegateName);
+
+	FString MacroName = SourceFile.GetGeneratedMacroName(FunctionData.MacroLine, TEXT("_DELEGATE"));
+	WriteMacro(Out, MacroName, DelegateOutput);
 }
 
-void FNativeClassHeaderGenerator::ExportEventParm(UFunction* Function, FUHTStringBuilder& HeaderOutput, int32 Indent, bool bOutputConstructor)
+void FNativeClassHeaderGenerator::ExportEventParm(FUHTStringBuilder& Out, TSet<FString>& PropertyFwd, UFunction* Function, int32 Indent, bool bOutputConstructor, EExportingState ExportingState)
 {
 	if (!WillExportEventParms(Function))
 	{
@@ -3905,13 +3213,17 @@ void FNativeClassHeaderGenerator::ExportEventParm(UFunction* Function, FUHTStrin
 	}
 
 	FString EventParmStructName = GetEventStructParamsName(Function->GetOuter(), *FunctionName);
-	HeaderOutput.Logf(TEXT("%sstruct %s\r\n"), FCString::Tab(Indent), *EventParmStructName);
-	HeaderOutput.Logf(TEXT("%s{\r\n"), FCString::Tab(Indent));
+	Out.Logf(TEXT("%sstruct %s\r\n"), FCString::Tab(Indent), *EventParmStructName);
+	Out.Logf(TEXT("%s{\r\n"), FCString::Tab(Indent));
 
-	for (TFieldIterator<UProperty> It(Function); It && (It->PropertyFlags&CPF_Parm); ++It)
+	for (UProperty* Prop : TFieldRange<UProperty>(Function))
 	{
-		UProperty* Prop = *It;
-		ForwardDeclarations.Add(Prop);
+		if (!(Prop->PropertyFlags & CPF_Parm))
+		{
+			continue;
+		}
+
+		PropertyFwd.Add(Prop->GetCPPTypeForwardDeclaration());
 
 		FUHTStringBuilder PropertyText;
 		PropertyText.Log(FCString::Tab(Indent + 1));
@@ -3937,10 +3249,10 @@ void FNativeClassHeaderGenerator::ExportEventParm(UFunction* Function, FUHTStrin
 
 		const FString* Dim = GArrayDimensions.Find(Prop);
 		Prop->ExportCppDeclaration(PropertyText, EExportedDeclaration::Local, Dim ? **Dim : NULL);
-		ApplyAlternatePropertyExportText(Prop, PropertyText);
+		ApplyAlternatePropertyExportText(Prop, PropertyText, ExportingState);
 
 		PropertyText.Log(TEXT(";\r\n"));
-		HeaderOutput += *PropertyText;
+		Out += *PropertyText;
 
 	}
 	// constructor must initialize the return property if it needs it
@@ -3972,31 +3284,14 @@ void FNativeClassHeaderGenerator::ExportEventParm(UFunction* Function, FUHTStrin
 		if (bNeedsOutput)
 		{
 			check(Prop->ArrayDim == 1); // can't return arrays
-			HeaderOutput.Logf(TEXT("\r\n%s/** Constructor, initializes return property only **/\r\n"), FCString::Tab(Indent + 1));
-			HeaderOutput.Logf(TEXT("%s%s()\r\n"), FCString::Tab(Indent + 1), *EventParmStructName);
-			HeaderOutput.Logf(TEXT("%s%s %s(%s)\r\n"), FCString::Tab(Indent + 2), TEXT(":"), *Prop->GetName(), *GetNullParameterValue(Prop, false, true));
-			HeaderOutput.Logf(TEXT("%s{\r\n"), FCString::Tab(Indent + 1));
-			HeaderOutput.Logf(TEXT("%s}\r\n"), FCString::Tab(Indent + 1));
+			Out.Logf(TEXT("\r\n%s/** Constructor, initializes return property only **/\r\n"), FCString::Tab(Indent + 1));
+			Out.Logf(TEXT("%s%s()\r\n"), FCString::Tab(Indent + 1), *EventParmStructName);
+			Out.Logf(TEXT("%s%s %s(%s)\r\n"), FCString::Tab(Indent + 2), TEXT(":"), *Prop->GetName(), *GetNullParameterValue(Prop, false, true));
+			Out.Logf(TEXT("%s{\r\n"), FCString::Tab(Indent + 1));
+			Out.Logf(TEXT("%s}\r\n"), FCString::Tab(Indent + 1));
 		}
 	}
-	HeaderOutput.Logf(TEXT("%s};\r\n"), FCString::Tab(Indent));
-}
-
-void FNativeClassHeaderGenerator::ExportEventParms(FScope& Scope, const TArray<UFunction*>& InCallbackFunctions, FUHTStringBuilder& Output, int32 Indent, bool bOutputConstructor)
-{
-	// Parms struct definitions.
-	FUHTStringBuilder HeaderOutput;
-
-	TArray<UFunction*> CallbackFunctions = InCallbackFunctions;
-	CallbackFunctions.Sort();
-
-	for ( int32 Index = 0; Index < CallbackFunctions.Num(); Index++ )
-	{
-		UFunction* Function = CallbackFunctions[Index];
-		ExportEventParm(Function, HeaderOutput, Indent, bOutputConstructor);
-	}
-
-	Output.Log(HeaderOutput);
+	Out.Logf(TEXT("%s};\r\n"), FCString::Tab(Indent));
 }
 
 /**
@@ -4089,14 +3384,14 @@ FString FNativeClassHeaderGenerator::GetNullParameterValue( UProperty* Prop, boo
 
 FString FNativeClassHeaderGenerator::GetFunctionReturnString(UFunction* Function)
 {
-	if (auto Return = Function->GetReturnProperty())
+	if (UProperty* Return = Function->GetReturnProperty())
 	{
 		FString ExtendedReturnType;
-		ForwardDeclarations.Add(Return);
+		ForwardDeclarations.Add(Return->GetCPPTypeForwardDeclaration());
 		FString ReturnType = Return->GetCPPType(&ExtendedReturnType, CPPF_ArgumentOrReturnValue);
 		FUHTStringBuilder ReplacementText;
 		ReplacementText += ReturnType;
-		ApplyAlternatePropertyExportText(Return, ReplacementText);
+		ApplyAlternatePropertyExportText(Return, ReplacementText, EExportingState::Normal);
 		return ReplacementText + ExtendedReturnType;
 	}
 
@@ -4147,7 +3442,7 @@ void GetLineAndColumnFromPositionInFile(const FString& File, int32 Position, int
 	}
 }
 
-bool FNativeClassHeaderGenerator::IsMissingVirtualSpecifier(const FString& SourceFile, int32 FunctionNamePosition) const
+bool FNativeClassHeaderGenerator::IsMissingVirtualSpecifier(const FString& SourceFile, int32 FunctionNamePosition)
 {
 	auto IsEndOfSearchChar = [](TCHAR C) { return (C == TEXT('}')) || (C == TEXT('{')) || (C == TEXT(';')); };
 
@@ -4252,7 +3547,15 @@ void FNativeClassHeaderGenerator::CheckRPCFunctions(const FFuncInfo& FunctionDat
 	}
 }
 
-void FNativeClassHeaderGenerator::ExportNativeFunctionHeader(const FFuncInfo& FunctionData, FUHTStringBuilder& HeaderOutput, EExportFunctionType::Type FunctionType, EExportFunctionHeaderStyle::Type FunctionHeaderStyle, const TCHAR* ExtraParam)
+void FNativeClassHeaderGenerator::ExportNativeFunctionHeader(
+	FOutputDevice&                   Out,
+	TSet<FString>&                   OutFwdDecls,
+	const FFuncInfo&                 FunctionData,
+	EExportFunctionType::Type        FunctionType,
+	EExportFunctionHeaderStyle::Type FunctionHeaderStyle,
+	const TCHAR*                     ExtraParam,
+	const TCHAR*                     APIString
+)
 {
 	UFunction* Function = FunctionData.FunctionReference;
 
@@ -4262,7 +3565,7 @@ void FNativeClassHeaderGenerator::ExportNativeFunctionHeader(const FFuncInfo& Fu
 	
 	if (!bIsDelegate)
 	{
-		HeaderOutput += TEXT("\t");
+		Out.Log(TEXT("\t"));
 	}
 
 	if (FunctionHeaderStyle == EExportFunctionHeaderStyle::Declaration)
@@ -4275,46 +3578,46 @@ void FNativeClassHeaderGenerator::ExportNativeFunctionHeader(const FFuncInfo& Fu
 			!Function->GetOwnerClass()->HasAnyClassFlags(CLASS_RequiredAPI) &&
 			(FunctionData.FunctionExportFlags & FUNCEXPORT_RequiredAPI))
 		{
-			HeaderOutput.Log(GetAPIString());
+			Out.Log(APIString);
 		}
 
 		if(FunctionType == EExportFunctionType::Interface)
 		{
-			HeaderOutput.Log(TEXT("static "));
+			Out.Log(TEXT("static "));
 		}
 		else if (bIsK2Override)
 		{
-			HeaderOutput.Log(TEXT("virtual "));
+			Out.Log(TEXT("virtual "));
 		}
 		// if the owning class is an interface class
 		else if ( bIsInterface )
 		{
-			HeaderOutput.Log(TEXT("virtual "));
+			Out.Log(TEXT("virtual "));
 		}
 		// this is not an event, the function is not a static function and the function is not marked final
 		else if ( FunctionType != EExportFunctionType::Event && !Function->HasAnyFunctionFlags(FUNC_Static) && !(FunctionData.FunctionExportFlags & FUNCEXPORT_Final) )
 		{
-			HeaderOutput.Log(TEXT("virtual "));
+			Out.Log(TEXT("virtual "));
 		}
 		else if( FunctionData.FunctionExportFlags & FUNCEXPORT_Inline )
 		{
-			HeaderOutput.Log(TEXT("inline "));
+			Out.Log(TEXT("inline "));
 		}
 	}
 
-	if (auto Return = Function->GetReturnProperty())
+	if (UProperty* Return = Function->GetReturnProperty())
 	{
 		FString ExtendedReturnType;
 		FString ReturnType = Return->GetCPPType(&ExtendedReturnType, (FunctionHeaderStyle == EExportFunctionHeaderStyle::Definition && (FunctionType != EExportFunctionType::Interface) ? CPPF_Implementation : 0) | CPPF_ArgumentOrReturnValue);
-		ForwardDeclarations.Add(Return);
+		OutFwdDecls.Add(Return->GetCPPTypeForwardDeclaration());
 		FUHTStringBuilder ReplacementText;
 		ReplacementText += ReturnType;
-		ApplyAlternatePropertyExportText(Return, ReplacementText);
-		HeaderOutput.Logf(TEXT("%s%s"), *ReplacementText, *ExtendedReturnType);
+		ApplyAlternatePropertyExportText(Return, ReplacementText, EExportingState::Normal);
+		Out.Logf(TEXT("%s%s"), *ReplacementText, *ExtendedReturnType);
 	}
 	else
 	{
-		HeaderOutput.Log( TEXT("void") );
+		Out.Log( TEXT("void") );
 	}
 
 	FString FunctionName;
@@ -4336,48 +3639,52 @@ void FNativeClassHeaderGenerator::ExportNativeFunctionHeader(const FFuncInfo& Fu
 		FunctionName += FunctionData.CppImplName;
 	}
 
-	HeaderOutput.Logf( TEXT(" %s("), *FunctionName);
+	Out.Logf(TEXT(" %s("), *FunctionName);
 
 	int32 ParmCount=0;
 
 	// Emit extra parameter if we have one
 	if( ExtraParam )
 	{
-		HeaderOutput += ExtraParam;
+		Out.Logf(TEXT("%s"), ExtraParam);
 		++ParmCount;
 	}
 
-	for( TFieldIterator<UProperty> It(Function); It && (It->PropertyFlags&(CPF_Parm|CPF_ReturnParm))==CPF_Parm; ++It )
+	for (UProperty* Property : TFieldRange<UProperty>(Function))
 	{
-		UProperty* Property = *It;
-		ForwardDeclarations.Add(Property);
+		if ((Property->PropertyFlags & (CPF_Parm | CPF_ReturnParm)) != CPF_Parm)
+		{
+			continue;
+		}
+
+		OutFwdDecls.Add(Property->GetCPPTypeForwardDeclaration());
 
 		if( ParmCount++ )
 		{
-			HeaderOutput.Log(TEXT(", "));
+			Out.Log(TEXT(", "));
 		}
 
 		FUHTStringBuilder PropertyText;
 
 		const FString* Dim = GArrayDimensions.Find(Property);
 		Property->ExportCppDeclaration( PropertyText, EExportedDeclaration::Parameter, Dim ? **Dim : NULL );
-		ApplyAlternatePropertyExportText(Property, PropertyText);
+		ApplyAlternatePropertyExportText(Property, PropertyText, EExportingState::Normal);
 
-		HeaderOutput += PropertyText;
+		Out.Logf(TEXT("%s"), *PropertyText);
 	}
 
-	HeaderOutput.Log( TEXT(")") );
+	Out.Log( TEXT(")") );
 	if (FunctionType != EExportFunctionType::Interface)
 	{
 		if (!bIsDelegate && Function->HasAllFunctionFlags(FUNC_Const))
 		{
-			HeaderOutput.Log( TEXT(" const") );
+			Out.Log( TEXT(" const") );
 		}
 
 		if (bIsInterface && FunctionHeaderStyle == EExportFunctionHeaderStyle::Declaration)
 		{
 			// all methods in interface classes are pure virtuals
-			HeaderOutput.Log(TEXT("=0"));
+			Out.Log(TEXT("=0"));
 		}
 	}
 }
@@ -4398,7 +3705,7 @@ void FNativeClassHeaderGenerator::ExportFunctionThunk(FUHTStringBuilder& RPCWrap
 	for (int32 ParameterIndex = 0; ParameterIndex < Parameters.Num(); ParameterIndex++)
 	{
 		UProperty* Param = Parameters[ParameterIndex];
-		ForwardDeclarations.Add(Param);
+		ForwardDeclarations.Add(Param->GetCPPTypeForwardDeclaration());
 
 		FString EvalBaseText = TEXT("P_GET_");	// e.g. P_GET_STR
 		FString EvalModifierText;				// e.g. _REF
@@ -4437,7 +3744,7 @@ void FNativeClassHeaderGenerator::ExportFunctionThunk(FUHTStringBuilder& RPCWrap
 		FUHTStringBuilder ReplacementText;
 		ReplacementText += TypeText;
 
-		ApplyAlternatePropertyExportText(Param, ReplacementText);
+		ApplyAlternatePropertyExportText(Param, ReplacementText, EExportingState::Normal);
 		TypeText = ReplacementText;
 
 		FString DefaultValueText;
@@ -4537,29 +3844,29 @@ void FNativeClassHeaderGenerator::ExportFunctionThunk(FUHTStringBuilder& RPCWrap
 	RPCWrappers += TEXT("\t\tP_FINISH;") LINE_TERMINATOR;
 	RPCWrappers += TEXT("\t\tP_NATIVE_BEGIN;") LINE_TERMINATOR;
 
-	auto ClassRange = ClassDefinitionRange();
+	ClassDefinitionRange ClassRange;
 	if (ClassDefinitionRanges.Contains(Function->GetOwnerClass()))
 	{
 		ClassRange = ClassDefinitionRanges[Function->GetOwnerClass()];
 		ClassRange.Validate();
 	}
 
-	auto ClassStart = ClassRange.Start;
-	auto ClassEnd = ClassRange.End;
-	auto ClassDefinition = FString(ClassEnd - ClassStart, ClassStart);
-	auto ClassName = Function->GetOwnerClass()->GetName();
+	const TCHAR* ClassStart = ClassRange.Start;
+	const TCHAR* ClassEnd   = ClassRange.End;
+	FString      ClassName  = Function->GetOwnerClass()->GetName();
+
+	FString ClassDefinition(ClassEnd - ClassStart, ClassStart);
 
-	auto FunctionName = Function->GetName();
 	bool bHasImplementation = HasIdentifierExactMatch(ClassDefinition, FunctionData.CppImplName);
 	bool bHasValidate = HasIdentifierExactMatch(ClassDefinition, FunctionData.CppValidationImplName);
 
-	auto bShouldEnableImplementationDeprecation =
+	bool bShouldEnableImplementationDeprecation =
 		// Enable deprecation warnings only if GENERATED_BODY is used inside class or interface (not GENERATED_UCLASS_BODY etc.)
 		ClassRange.bHasGeneratedBody
 		// and implementation function is called, but not the one declared by user
-		&& (FunctionData.CppImplName != FunctionName && !bHasImplementation);
+		&& (FunctionData.CppImplName != Function->GetName() && !bHasImplementation);
 
-	auto bShouldEnableValidateDeprecation =
+	bool bShouldEnableValidateDeprecation =
 		// Enable deprecation warnings only if GENERATED_BODY is used inside class or interface (not GENERATED_UCLASS_BODY etc.)
 		ClassRange.bHasGeneratedBody
 		// and validation function is called
@@ -4567,7 +3874,7 @@ void FNativeClassHeaderGenerator::ExportFunctionThunk(FUHTStringBuilder& RPCWrap
 
 	//Emit warning here if necessary
 	FUHTStringBuilder FunctionDeclaration;
-	ExportNativeFunctionHeader(FunctionData, FunctionDeclaration, EExportFunctionType::Function, EExportFunctionHeaderStyle::Declaration);
+	ExportNativeFunctionHeader(FunctionDeclaration, ForwardDeclarations, FunctionData, EExportFunctionType::Function, EExportFunctionHeaderStyle::Declaration, nullptr, *GetAPIString());
 	FunctionDeclaration.Trim();
 
 	// Call the validate function if there is one
@@ -4582,15 +3889,16 @@ void FNativeClassHeaderGenerator::ExportFunctionThunk(FUHTStringBuilder& RPCWrap
 
 	// write out the return value
 	RPCWrappers.Log(TEXT("\t\t"));
-	if (Return != NULL)
+	if (Return)
 	{
-		FString ReturnType, ReturnExtendedType;
-		ForwardDeclarations.Add(Return);
-		ReturnType = Return->GetCPPType(&ReturnExtendedType);
+		ForwardDeclarations.Add(Return->GetCPPTypeForwardDeclaration());
+
 		FUHTStringBuilder ReplacementText;
-		ReplacementText += ReturnType;
-		ApplyAlternatePropertyExportText(Return, ReplacementText);
-		ReturnType = ReplacementText;
+		FString ReturnExtendedType;
+		ReplacementText += Return->GetCPPType(&ReturnExtendedType);
+		ApplyAlternatePropertyExportText(Return, ReplacementText, EExportingState::Normal);
+
+		FString ReturnType = ReplacementText;
 		RPCWrappers.Logf(TEXT("*(%s%s*)") TEXT(PREPROCESSOR_TO_STRING(RESULT_PARAM)) TEXT("="), *ReturnType, *ReturnExtendedType);
 	}
 
@@ -4611,9 +3919,9 @@ FString FNativeClassHeaderGenerator::GetFunctionParameterString(UFunction* Funct
 	FString ParameterList;
 	FUHTStringBuilder PropertyText;
 
-	for (auto Property : TFieldRange<UProperty>(Function))
+	for (UProperty* Property : TFieldRange<UProperty>(Function))
 	{
-		ForwardDeclarations.Add(Property);
+		ForwardDeclarations.Add(Property->GetCPPTypeForwardDeclaration());
 
 		if ((Property->PropertyFlags & (CPF_Parm | CPF_ReturnParm)) != CPF_Parm)
 		{
@@ -4627,7 +3935,7 @@ FString FNativeClassHeaderGenerator::GetFunctionParameterString(UFunction* Funct
 
 		auto Dim = GArrayDimensions.Find(Property);
 		Property->ExportCppDeclaration(PropertyText, EExportedDeclaration::Parameter, Dim ? **Dim : NULL, 0, true);
-		ApplyAlternatePropertyExportText(Property, PropertyText);
+		ApplyAlternatePropertyExportText(Property, PropertyText, EExportingState::Normal);
 
 		ParameterList += PropertyText;
 		PropertyText.Reset();
@@ -4636,15 +3944,15 @@ FString FNativeClassHeaderGenerator::GetFunctionParameterString(UFunction* Funct
 	return ParameterList;
 }
 
-void FNativeClassHeaderGenerator::ExportNativeFunctions(FUnrealSourceFile& SourceFile, UClass* Class, FClassMetaData* ClassData)
+void FNativeClassHeaderGenerator::ExportNativeFunctions(FOutputDevice& OutGeneratedHeaderText, FOutputDevice& OutMacroCalls, FOutputDevice& OutNoPureDeclsMacroCalls, const FUnrealSourceFile& SourceFile, UClass* Class, FClassMetaData* ClassData)
 {
 	FUHTStringBuilder RPCWrappers;
 	FUHTStringBuilder AutogeneratedBlueprintFunctionDeclarations;
 	FUHTStringBuilder AutogeneratedBlueprintFunctionDeclarationsOnlyNotDeclared;
 
-	auto ClassName = Class->GetName();
+	FString ClassName = Class->GetName();
 
-	auto ClassRange = ClassDefinitionRange();
+	ClassDefinitionRange ClassRange;
 	if (ClassDefinitionRanges.Contains(Class))
 	{
 		ClassRange = ClassDefinitionRanges[Class];
@@ -4660,7 +3968,7 @@ void FNativeClassHeaderGenerator::ExportNativeFunctions(FUnrealSourceFile& Sourc
 			continue;
 		}
 
-		auto* CompilerInfo = FFunctionData::FindForFunction(Function);
+		FFunctionData* CompilerInfo = FFunctionData::FindForFunction(Function);
 
 		const FFuncInfo& FunctionData = CompilerInfo->GetFunctionData();
 
@@ -4670,36 +3978,35 @@ void FNativeClassHeaderGenerator::ExportNativeFunctions(FUnrealSourceFile& Sourc
 			continue;
 		}
 
-		FUHTStringBuilder& DestinationForDecl = RPCWrappers;
-
 		// Should we emit these to RPC wrappers or just ignore them?
 		const bool bWillBeProgrammerTyped = FunctionData.CppImplName == Function->GetName();
 
 		if (!bWillBeProgrammerTyped)
 		{
-			auto ClassStart = ClassRange.Start;
-			auto ClassEnd = ClassRange.End;
-			auto ClassDefinition = FString(ClassEnd - ClassStart, ClassStart);
+			const TCHAR* ClassStart = ClassRange.Start;
+			const TCHAR* ClassEnd   = ClassRange.End;
+			FString ClassDefinition(ClassEnd - ClassStart, ClassStart);
 
-			auto FunctionName = Function->GetName();
+			FString FunctionName = Function->GetName();
 			int32 ClassDefinitionStartPosition = ClassStart - *SourceFile.GetContent();
 
 			int32 ImplementationPosition = FindIdentifierExactMatch(ClassDefinition, FunctionData.CppImplName);
-			if (ImplementationPosition != INDEX_NONE)
+			bool bHasImplementation = ImplementationPosition != INDEX_NONE;
+			if (bHasImplementation)
 			{
 				ImplementationPosition += ClassDefinitionStartPosition;
 			}
+
 			int32 ValidatePosition = FindIdentifierExactMatch(ClassDefinition, FunctionData.CppValidationImplName);
-			if (ValidatePosition != INDEX_NONE)
+			bool bHasValidate = ValidatePosition != INDEX_NONE;
+			if (bHasValidate)
 			{
 				ValidatePosition += ClassDefinitionStartPosition;
 			}
-			bool bHasImplementation = ImplementationPosition != INDEX_NONE;
-			bool bHasValidate = ValidatePosition != INDEX_NONE;
 
 			//Emit warning here if necessary
 			FUHTStringBuilder FunctionDeclaration;
-			ExportNativeFunctionHeader(FunctionData, FunctionDeclaration, EExportFunctionType::Function, EExportFunctionHeaderStyle::Declaration);
+			ExportNativeFunctionHeader(FunctionDeclaration, ForwardDeclarations, FunctionData, EExportFunctionType::Function, EExportFunctionHeaderStyle::Declaration, nullptr, *GetAPIString());
 			FunctionDeclaration.Log(TEXT(";\r\n"));
 
 			// Declare validation function if needed
@@ -4707,7 +4014,7 @@ void FNativeClassHeaderGenerator::ExportNativeFunctions(FUnrealSourceFile& Sourc
 			{
 				FString ParameterList = GetFunctionParameterString(Function);
 
-				auto Virtual = (!FunctionData.FunctionReference->HasAnyFunctionFlags(FUNC_Static) && !(FunctionData.FunctionExportFlags & FUNCEXPORT_Final)) ? TEXT("virtual") : TEXT("");
+				const TCHAR* Virtual = (!FunctionData.FunctionReference->HasAnyFunctionFlags(FUNC_Static) && !(FunctionData.FunctionExportFlags & FUNCEXPORT_Final)) ? TEXT("virtual") : TEXT("");
 				FStringOutputDevice ValidDecl;
 				ValidDecl.Logf(TEXT("\t%s bool %s(%s);\r\n"), Virtual, *FunctionData.CppValidationImplName, *ParameterList);
 				AutogeneratedBlueprintFunctionDeclarations.Log(*ValidDecl);
@@ -4718,27 +4025,24 @@ void FNativeClassHeaderGenerator::ExportNativeFunctions(FUnrealSourceFile& Sourc
 			}
 
 			AutogeneratedBlueprintFunctionDeclarations.Log(*FunctionDeclaration);
-			if ((FunctionData.CppImplName != FunctionName && !bHasImplementation))
+			if (!bHasImplementation && FunctionData.CppImplName != FunctionName)
 			{
 				AutogeneratedBlueprintFunctionDeclarationsOnlyNotDeclared.Log(*FunctionDeclaration);
 			}
 
 			// Versions that skip function autodeclaration throw an error when a function is missing.
-			if ((SourceFile.GetGeneratedCodeVersionForStruct(Class) > EGeneratedCodeVersion::V1) && ClassRange.bHasGeneratedBody)
+			if (ClassRange.bHasGeneratedBody && (SourceFile.GetGeneratedCodeVersionForStruct(Class) > EGeneratedCodeVersion::V1))
 			{
-				auto Name = Class->HasAnyClassFlags(CLASS_Interface) ? *(FString(TEXT("I")) + ClassName) : NameLookupCPP.GetNameCPP(Class);
-				CheckRPCFunctions(FunctionData, Name, ImplementationPosition, ValidatePosition, SourceFile);
+				FString Name = Class->HasAnyClassFlags(CLASS_Interface) ? TEXT("I") + ClassName : FString(NameLookupCPP.GetNameCPP(Class));
+				CheckRPCFunctions(FunctionData, *Name, ImplementationPosition, ValidatePosition, SourceFile);
 			}
 		}
 
-		if (bMultiLineUFUNCTION)
-		{
-			RPCWrappers.Log(TEXT("\r\n"));
-		}
+		RPCWrappers.Log(TEXT("\r\n"));
 
 		// if this function was originally declared in a base class, and it isn't a static function,
 		// only the C++ function header will be exported
-		if (!ShouldExportFunction(Function))
+		if (!ShouldExportUFunction(Function))
 		{
 			continue;
 		}
@@ -4747,27 +4051,27 @@ void FNativeClassHeaderGenerator::ExportNativeFunctions(FUnrealSourceFile& Sourc
 		RPCWrappers.Logf(TEXT("\tDECLARE_FUNCTION(%s)"), *FunctionData.UnMarshallAndCallName);
 		RPCWrappers += LINE_TERMINATOR TEXT("\t{") LINE_TERMINATOR;
 
-		auto Parameters = GetFunctionParmsAndReturn(FunctionData.FunctionReference);
+		FParmsAndReturnProperties Parameters = GetFunctionParmsAndReturn(FunctionData.FunctionReference);
 		ExportFunctionThunk(RPCWrappers, Function, FunctionData, Parameters.Parms, Parameters.Return, AutogeneratedBlueprintFunctionDeclarationsOnlyNotDeclared);
 
 		RPCWrappers += TEXT("\t}") LINE_TERMINATOR;
 	}
 
 	FString MacroName = SourceFile.GetGeneratedMacroName(ClassData, TEXT("_RPC_WRAPPERS"));
-	WriteMacro(GeneratedHeaderText, MacroName, AutogeneratedBlueprintFunctionDeclarations + RPCWrappers);
-	InClassMacroCalls.Logf(TEXT("\t%s\r\n"), *MacroName);
+	WriteMacro(OutGeneratedHeaderText, MacroName, AutogeneratedBlueprintFunctionDeclarations + RPCWrappers);
+	OutMacroCalls.Logf(TEXT("\t%s\r\n"), *MacroName);
 
 	// Put static checks before RPCWrappers to get proper messages from static asserts before compiler errors.
 	FString NoPureDeclsMacroName = SourceFile.GetGeneratedMacroName(ClassData, TEXT("_RPC_WRAPPERS_NO_PURE_DECLS"));
 	if (SourceFile.GetGeneratedCodeVersionForStruct(Class) > EGeneratedCodeVersion::V1)
 	{
-		WriteMacro(GeneratedHeaderText, NoPureDeclsMacroName, RPCWrappers);
+		WriteMacro(OutGeneratedHeaderText, NoPureDeclsMacroName, RPCWrappers);
 	}
 	else
 	{
-		WriteMacro(GeneratedHeaderText, NoPureDeclsMacroName, AutogeneratedBlueprintFunctionDeclarationsOnlyNotDeclared + RPCWrappers);
+		WriteMacro(OutGeneratedHeaderText, NoPureDeclsMacroName, AutogeneratedBlueprintFunctionDeclarationsOnlyNotDeclared + RPCWrappers);
 	}
-	InClassNoPureDeclsMacroCalls.Logf(TEXT("\t%s\r\n"), *NoPureDeclsMacroName);
+	OutNoPureDeclsMacroCalls.Logf(TEXT("\t%s\r\n"), *NoPureDeclsMacroName);
 }
 
 /**
@@ -4775,58 +4079,32 @@ void FNativeClassHeaderGenerator::ExportNativeFunctions(FUnrealSourceFile& Sourc
  *
  * @param	CallbackFunctions	the functions to export
  */
-TArray<UFunction*> FNativeClassHeaderGenerator::ExportCallbackFunctions(FUnrealSourceFile& SourceFile, UClass* Class, FClassMetaData* ClassData)
+void FNativeClassHeaderGenerator::ExportCallbackFunctions(
+	FOutputDevice&            OutGeneratedHeaderText,
+	FOutputDevice&            OutCpp,
+	TMap<FName, FName>&       OutReferencedNames,
+	TSet<FString>&            OutFwdDecls,
+	const TArray<UFunction*>& CallbackFunctions,
+	const TCHAR*              CallbackWrappersMacroName,
+	EExportCallbackType       ExportCallbackType,
+	const TCHAR*              API,
+	const TCHAR*              APIString
+)
 {
-	TArray<UFunction*> CallbackFunctions;
-
-	for (UFunction* Function : TFieldRange<UFunction>(Class, EFieldIteratorFlags::ExcludeSuper))
-	{
-		if ((Function->FunctionFlags & FUNC_Event) && Function->GetSuperFunction() == nullptr)
-		{
-			CallbackFunctions.Add(Function);
-		}
-	}
-
 	FUHTStringBuilder RPCWrappers;
-	FUHTStringBuilder RPCWrappersCPP;
-
-	if (CallbackFunctions.Num() == 0)
-	{
-		// Early out.
-		return CallbackFunctions;
-	}
-
-	CallbackFunctions.Sort();
-
-	FUHTStringBuilder UClassMacroContent;
-
-	// export parameters structs for all events and delegates
-	ExportEventParms(SourceFile.GetScope().Get(), CallbackFunctions, UClassMacroContent, /*Indent=*/ 1, /*bOutputConstructor=*/ true);
-
-	FString MacroName = SourceFile.GetGeneratedMacroName(ClassData, TEXT("_EVENT_PARMS"));
-	WriteMacro(GeneratedHeaderText, MacroName, UClassMacroContent);
-	PrologMacroCalls.Logf(TEXT("\t%s\r\n"), *MacroName);
-
-	// export .proto files for any net service functions
-	ExportProtoMessage(CallbackFunctions);
-
-	// export .java files for any net service functions
-	ExportMCPMessage(CallbackFunctions, ClassData);
-
-	for (int32 CallbackIndex = 0; CallbackIndex < CallbackFunctions.Num(); CallbackIndex++)
+	for (UFunction* Function : CallbackFunctions)
 	{
-		UFunction* Function = CallbackFunctions[CallbackIndex];
-
 		// Never expecting to export delegate functions this way
 		check(!Function->HasAnyFunctionFlags(FUNC_Delegate));
 
-		auto* CompilerInfo = FFunctionData::FindForFunction(Function);
+		FFunctionData* CompilerInfo = FFunctionData::FindForFunction(Function);
 
 		// cache the TCHAR* for a few strings we'll use a lot here
-		FString FunctionName = Function->GetName();
+		FString FunctionName     = Function->GetName();
+		FString FunctionNameName = FString::Printf(TEXT("%s_%s"), API, *FunctionName);
 
-		GeneratedHeaderText.Logf(TEXT("extern %s FName %s_%s;") LINE_TERMINATOR, *GetAPIString(), *API, *FunctionName);
-		ReferencedNames.Add(Function->GetFName(), GetOverriddenFName(Function));
+		OutGeneratedHeaderText.Logf(TEXT("extern %s FName %s;") LINE_TERMINATOR, APIString, *FunctionNameName);
+		OutReferencedNames.Add(Function->GetFName(), GetOverriddenFName(Function));
 
 		const FFuncInfo& FunctionData = CompilerInfo->GetFunctionData();
 
@@ -4842,68 +4120,51 @@ TArray<UFunction*> FNativeClassHeaderGenerator::ExportCallbackFunctions(FUnrealS
 		if (!bWillBeProgrammerTyped)
 		{
 			// export the line that looks like: int32 Main(const FString& Parms)
-			ExportNativeFunctionHeader(FunctionData, RPCWrappers, EExportFunctionType::Event, EExportFunctionHeaderStyle::Declaration);
+			ExportNativeFunctionHeader(RPCWrappers, OutFwdDecls, FunctionData, EExportFunctionType::Event, EExportFunctionHeaderStyle::Declaration, nullptr, APIString);
 
 			RPCWrappers.Log(TEXT(";\r\n"));
-
-			if (bMultiLineUFUNCTION)
-			{
-				RPCWrappers.Log(TEXT("\r\n"));
-			}
+			RPCWrappers.Log(TEXT("\r\n"));
 		}
 
 		// Emit the thunk implementation
-		ExportNativeFunctionHeader(FunctionData, RPCWrappersCPP, EExportFunctionType::Event, EExportFunctionHeaderStyle::Definition);
+		ExportNativeFunctionHeader(OutCpp, OutFwdDecls, FunctionData, EExportFunctionType::Event, EExportFunctionHeaderStyle::Definition, nullptr, APIString);
 
-		auto Parameters = GetFunctionParmsAndReturn(FunctionData.FunctionReference);
+		FParmsAndReturnProperties Parameters = GetFunctionParmsAndReturn(FunctionData.FunctionReference);
 
-		if (!(Class->ClassFlags & CLASS_Interface))
+		if (ExportCallbackType != EExportCallbackType::Interface)
 		{
-			WriteEventFunctionPrologue(RPCWrappersCPP, /*Indent=*/ 1, Parameters, Function->GetOuter(), *FunctionName);
+			WriteEventFunctionPrologue(OutCpp, /*Indent=*/ 1, Parameters, Function->GetOuter(), *FunctionName);
 			{
 				// Cast away const just in case, because ProcessEvent isn't const
-				RPCWrappersCPP.Logf(
-					TEXT("\t\t%sProcessEvent(FindFunctionChecked(%s_%s),%s);\r\n"),
+				OutCpp.Logf(
+					TEXT("\t\t%sProcessEvent(FindFunctionChecked(%s),%s);\r\n"),
 					(Function->HasAllFunctionFlags(FUNC_Const)) ? *FString::Printf(TEXT("const_cast<%s*>(this)->"), NameLookupCPP.GetNameCPP(Cast<UClass>(Function->GetOuter()))) : TEXT(""),
-					*API,
-					*FunctionName,
+					*FunctionNameName,
 					Parameters.HasParms() ? TEXT("&Parms") : TEXT("NULL")
-					);
+				);
 			}
-			WriteEventFunctionEpilogue(RPCWrappersCPP, /*Indent=*/ 1, Parameters, *FunctionName);
+			WriteEventFunctionEpilogue(OutCpp, /*Indent=*/ 1, Parameters, *FunctionName);
 		}
 		else
 		{
-			RPCWrappersCPP += LINE_TERMINATOR;
-			RPCWrappersCPP += TEXT("\t{") LINE_TERMINATOR;
+			OutCpp.Log(LINE_TERMINATOR);
+			OutCpp.Log(TEXT("\t{") LINE_TERMINATOR);
 
 			// assert if this is ever called directly
-			RPCWrappersCPP.Logf(TEXT("\t\tcheck(0 && \"Do not directly call Event functions in Interfaces. Call Execute_%s instead.\");") LINE_TERMINATOR, *FunctionName);
+			OutCpp.Logf(TEXT("\t\tcheck(0 && \"Do not directly call Event functions in Interfaces. Call Execute_%s instead.\");") LINE_TERMINATOR, *FunctionName);
 
 			// satisfy compiler if it's expecting a return value
 			if (Parameters.Return)
 			{
 				FString EventParmStructName = GetEventStructParamsName(Function->GetOuter(), *FunctionName);
-				RPCWrappersCPP.Logf(TEXT("\t\t%s Parms;") LINE_TERMINATOR, *EventParmStructName);
-				RPCWrappersCPP.Logf(TEXT("\t\treturn Parms.ReturnValue;") LINE_TERMINATOR);
+				OutCpp.Logf(TEXT("\t\t%s Parms;") LINE_TERMINATOR, *EventParmStructName);
+				OutCpp.Log(TEXT("\t\treturn Parms.ReturnValue;") LINE_TERMINATOR);
 			}
-			RPCWrappersCPP += TEXT("\t}") LINE_TERMINATOR;
+			OutCpp.Log(TEXT("\t}") LINE_TERMINATOR);
 		}
 	}
 
-	if (!Class->HasAnyClassFlags(CLASS_NoExport))
-	{
-		GeneratedPackageCPP.Log(*RPCWrappersCPP);
-	}
-	// else drop the implementation on the floor
-
-	FString CallbackWrappersMacroName = SourceFile.GetGeneratedMacroName(ClassData, TEXT("_CALLBACK_WRAPPERS"));
-	WriteMacro(GeneratedHeaderText, CallbackWrappersMacroName, RPCWrappers);
-
-	InClassMacroCalls.Logf(TEXT("\t%s\r\n"), *CallbackWrappersMacroName);
-	InClassNoPureDeclsMacroCalls.Logf(TEXT("\t%s\r\n"), *CallbackWrappersMacroName);
-
-	return CallbackFunctions;
+	WriteMacro(OutGeneratedHeaderText, CallbackWrappersMacroName, RPCWrappers);
 }
 
 
@@ -4915,7 +4176,7 @@ TArray<UFunction*> FNativeClassHeaderGenerator::ExportCallbackFunctions(FUnrealS
  * @param	Prop			the property that is being exported
  * @param	PropertyText	the string containing the text exported from ExportCppDeclaration
  */
-void FNativeClassHeaderGenerator::ApplyAlternatePropertyExportText(UProperty* Prop, FUHTStringBuilder& PropertyText)
+void FNativeClassHeaderGenerator::ApplyAlternatePropertyExportText(UProperty* Prop, FUHTStringBuilder& PropertyText, EExportingState ExportingState)
 {
 	UArrayProperty* ArrayProperty = Cast<UArrayProperty>(Prop);
 	UProperty* InnerProperty = ArrayProperty ? ArrayProperty->Inner : nullptr;
@@ -4934,26 +4195,24 @@ void FNativeClassHeaderGenerator::ApplyAlternatePropertyExportText(UProperty* Pr
 		return;
 	}
 
-	if (!bIsExportingForOffsetDeterminationOnly)
-	{
-		return;
-	}
-
-	UDelegateProperty* DelegateProperty = Cast<UDelegateProperty>(Prop);
-	UMulticastDelegateProperty* MulticastDelegateProperty = Cast<UMulticastDelegateProperty>(Prop);
-	if (DelegateProperty || MulticastDelegateProperty)
+	if (ExportingState == EExportingState::TypeEraseDelegates)
 	{
-		FString Original = Prop->GetCPPType();
-		FString PlaceholderOfSameSizeAndAlignemnt;
-		if (DelegateProperty)
-		{
-			PlaceholderOfSameSizeAndAlignemnt = TEXT("FScriptDelegate");
-		}
-		else
+		UDelegateProperty* DelegateProperty = Cast<UDelegateProperty>(Prop);
+		UMulticastDelegateProperty* MulticastDelegateProperty = Cast<UMulticastDelegateProperty>(Prop);
+		if (DelegateProperty || MulticastDelegateProperty)
 		{
-			PlaceholderOfSameSizeAndAlignemnt = TEXT("FMulticastScriptDelegate");
+			FString Original = Prop->GetCPPType();
+			FString PlaceholderOfSameSizeAndAlignemnt;
+			if (DelegateProperty)
+			{
+				PlaceholderOfSameSizeAndAlignemnt = TEXT("FScriptDelegate");
+			}
+			else
+			{
+				PlaceholderOfSameSizeAndAlignemnt = TEXT("FMulticastScriptDelegate");
+			}
+			PropertyText.ReplaceInline(*Original, *PlaceholderOfSameSizeAndAlignemnt, ESearchCase::CaseSensitive);
 		}
-		PropertyText.ReplaceInline(*Original, *PlaceholderOfSameSizeAndAlignemnt, ESearchCase::CaseSensitive);
 	}
 }
 
@@ -4963,7 +4222,7 @@ void FNativeClassHeaderGenerator::ApplyAlternatePropertyExportText(UProperty* Pr
  * @param	HeaderDependencyMap		a mapping of header filenames to a list of header filenames that must be processed before that one.
  * @param	SortedHeaderFilenames	[out] receives the sorted list of header filenames.
  */
-bool FNativeClassHeaderGenerator::SortHeaderDependencyMap(const TMap<const FString*, HeaderDependents>& HeaderDependencyMap, TArray<const FString*>& SortedHeaderFilenames) const
+bool FNativeClassHeaderGenerator::SortHeaderDependencyMap(const TMap<const FString*, HeaderDependents>& HeaderDependencyMap, TArray<const FString*>& SortedHeaderFilenames)
 {
 	SortedHeaderFilenames.Empty(HeaderDependencyMap.Num());
 
@@ -5063,12 +4322,61 @@ FUHTStringBuilder& FNativeClassHeaderGenerator::GetGeneratedFunctionTextDevice()
 		}
 	} MaxLinesPerCpp;
 
-	if ((GeneratedFunctionBodyTextSplit.Num() == 0) || (GeneratedFunctionBodyTextSplit.Num() == 1 && GeneratedFunctionBodyTextSplit[0]->GetLineCount() > MaxLinesPerCpp.FirstValue) || (GeneratedFunctionBodyTextSplit[GeneratedFunctionBodyTextSplit.Num() - 1]->GetLineCount() > MaxLinesPerCpp.OtherValue))
+	if (GeneratedFunctionBodyTextSplit.Num() == 0 || (GeneratedFunctionBodyTextSplit.Num() == 1 && GeneratedFunctionBodyTextSplit[0]->GetLineCount() > MaxLinesPerCpp.FirstValue) || (GeneratedFunctionBodyTextSplit.Last()->GetLineCount() > MaxLinesPerCpp.OtherValue))
+	{
+		GeneratedFunctionBodyTextSplit.Emplace();
+	}
+
+	return GeneratedFunctionBodyTextSplit.Last().Get();
+}
+
+void GetSourceFilesInDependencyOrderRecursive(TArray<FUnrealSourceFile*>& OutTest, const UPackage* Package, FUnrealSourceFile* SourceFile, TSet<const FUnrealSourceFile*>& VisitedSet, bool bCheckDependenciesOnly, const TSet<FUnrealSourceFile*>& Ignore)
+{
+	// Check if the Class has already been exported, after we've checked for circular header dependencies.
+	if (OutTest.Contains(SourceFile) || Ignore.Contains(SourceFile))
+	{
+		return;
+	}
+
+	// Check for circular dependencies.
+	if (VisitedSet.Contains(SourceFile))
+	{
+		UE_LOG(LogCompile, Error, TEXT("Circular dependency detected for filename %s!"), *SourceFile->GetFilename());
+		return;
+	}
+
+	// Check for circular header dependencies between export classes.
+	bCheckDependenciesOnly = bCheckDependenciesOnly || SourceFile->GetPackage() != Package;
+
+	VisitedSet.Add(SourceFile);
+	for (FHeaderProvider& Include : SourceFile->GetIncludes())
+	{
+		if (FUnrealSourceFile* IncludeFile = Include.Resolve())
+		{
+			GetSourceFilesInDependencyOrderRecursive(OutTest, Package, IncludeFile, VisitedSet, bCheckDependenciesOnly, Ignore);
+		}
+	}
+	VisitedSet.Remove(SourceFile);
+
+	if (!bCheckDependenciesOnly)
+	{
+		OutTest.Add(SourceFile);
+	}
+}
+
+TArray<FUnrealSourceFile*> GetSourceFilesInDependencyOrder(const UPackage* Package, const TArray<FUnrealSourceFile*>& SourceFiles, const TSet<FUnrealSourceFile*>& Ignore)
+{
+	TArray<FUnrealSourceFile*> Result;
+	TSet<const FUnrealSourceFile*>	VisitedSet;
+	for (FUnrealSourceFile* SourceFile : SourceFiles)
 	{
-		GeneratedFunctionBodyTextSplit.Add( TUniqueObj<FUHTStringBuilderLineCounter>() );
+		if (SourceFile->GetPackage() == Package)
+		{
+			GetSourceFilesInDependencyOrderRecursive(Result, Package, SourceFile, VisitedSet, false, Ignore);
+		}
 	}
 
-	return GeneratedFunctionBodyTextSplit[GeneratedFunctionBodyTextSplit.Num() - 1].Get();
+	return Result;
 }
 
 // Constructor.
@@ -5076,146 +4384,153 @@ FNativeClassHeaderGenerator::FNativeClassHeaderGenerator(
 	const UPackage* InPackage,
 	const TArray<FUnrealSourceFile*>& SourceFiles,
 	FClasses& AllClasses,
-	bool InAllowSaveExportedHeaders
-#if WITH_HOT_RELOAD_CTORS
-	, bool bInExportVTableConstructors
-#endif // WITH_HOT_RELOAD_CTORS
-	, FUHTMakefile& InUHTMakefile
+	bool InAllowSaveExportedHeaders,
+	FUHTMakefile& InUHTMakefile
 )
-	: API                                   (FPackageName::GetShortName(InPackage).ToUpper())
-	, Package                               (InPackage)
-	, bIsExportingForOffsetDeterminationOnly(false)
-	, bAllowSaveExportedHeaders             (InAllowSaveExportedHeaders)
-	, bFailIfGeneratedCodeChanges           (false)
-#if WITH_HOT_RELOAD_CTORS
-	, bExportVTableConstructors				(bInExportVTableConstructors)
-#endif // WITH_HOT_RELOAD_CTORS
-	, UHTMakefile(InUHTMakefile)
+	: API                        (FPackageName::GetShortName(InPackage).ToUpper())
+	, Package                    (InPackage)
+	, bAllowSaveExportedHeaders  (InAllowSaveExportedHeaders)
+	, bFailIfGeneratedCodeChanges(FParse::Param(FCommandLine::Get(), TEXT("FailIfGeneratedCodeChanges")))
+	, UHTMakefile                (InUHTMakefile)
 {
 	const FString PackageName = FPackageName::GetShortName(Package);
 
-	GeneratedCPPFilenameBase =  PackageName + TEXT(".generated");
-	GeneratedProtoFilenameBase = PackageName + TEXT(".generated");
-	GeneratedMCPFilenameBase = PackageName + TEXT(".generated");
-	bFailIfGeneratedCodeChanges = FParse::Param(FCommandLine::Get(), TEXT("FailIfGeneratedCodeChanges"));
-
-	TempHeaderPaths.Reset();
-	PackageHeaderPaths.Reset();
-
-	// Reset header generation output strings
-	GeneratedPackageCPP.Reset();
-	GeneratedProtoText.Reset();
-	GeneratedMCPText.Reset();
-	CrossModuleGeneratedFunctionDeclarations.Reset();
-	UniqueCrossModuleReferences.Empty(UniqueCrossModuleReferences.Num());
-	GeneratedFunctionDeclarations.Reset();
-	GeneratedFunctionBodyTextSplit.Reset();
-	ListOfPublicClassesUObjectHeaderModuleIncludes.Reset();
-
+	bool bWriteClassesH = false;
 	const bool bPackageHasAnyExportClasses = AllClasses.GetClassesInPackage(Package).ContainsByPredicate([](FClass* Class)
 	{
 		return Class->HasAnyClassFlags(CLASS_Native) && !Class->HasAnyClassFlags(CLASS_NoExport | CLASS_Intrinsic);
 	});
 	if (bPackageHasAnyExportClasses)
 	{
-		FString PkgDir;
-		FString GeneratedIncludeDirectory;
-		if (FindPackageLocation(*PackageName, PkgDir, GeneratedIncludeDirectory) == false)
-		{
-			UE_LOG(LogCompile, Error, TEXT("Failed to find path for package %s"), *PackageName);
-		}
-		FString ClassesHeaderName = PackageName + TEXT("Classes.h");
-		ClassesHeaderPath = GeneratedIncludeDirectory / ClassesHeaderName;
-
-		int32 ClassCount = 0;
-		for (auto* SourceFile : SourceFiles)
+		for (FUnrealSourceFile* SourceFile : SourceFiles)
 		{
-			for (auto* Class : SourceFile->GetDefinedClasses())
+			TArray<UClass*> DefinedClasses = SourceFile->GetDefinedClasses();
+			for (UClass* Class : DefinedClasses)
 			{
-				if (Class->HasAnyClassFlags(CLASS_Native))
+				if (!Class->HasAnyClassFlags(CLASS_Native))
 				{
-					if (GTypeDefinitionInfoMap.Contains(Class) && !Class->HasAnyClassFlags(CLASS_NoExport))
-					{
-						ClassCount++;
-						Class->UnMark(OBJECTMARK_TagImp);
-						Class->Mark(OBJECTMARK_TagExp);
-					}
+					Class->UnMark(EObjectMark(OBJECTMARK_TagImp | OBJECTMARK_TagExp));
 				}
-				else
+				else if (GTypeDefinitionInfoMap.Contains(Class) && !Class->HasAnyClassFlags(CLASS_NoExport))
 				{
-					Class->UnMark(EObjectMark(OBJECTMARK_TagImp | OBJECTMARK_TagExp));
+					bWriteClassesH = true;
+					Class->UnMark(OBJECTMARK_TagImp);
+					Class->Mark(OBJECTMARK_TagExp);
 				}
 			}
 		}
+	}
 
-		if (ClassCount != 0)
-		{
-			ClassesHeaders.Add(ClassesHeaderName);
+	auto OutputToGeneratedTextDevice = [this]() -> FUHTStringBuilder&
+	{
+		return GetGeneratedFunctionTextDevice();
+	};
+
+	// Export an include line for each header
+	TArray<FUnrealSourceFile*> PublicHeaderGroupIncludes;
+	FUHTStringBuilder GeneratedFunctionDeclarations;
 
-			ListOfPublicHeaderGroupIncludes.Reset();
-			ListOfAllUObjectHeaderIncludes.Reset();
-			OriginalHeader.Reset();
-			PreHeaderText.Reset();
+	TArray<FUnrealSourceFile*> Exported;
+	{
+		// Get source files and ignore them next time round
+		static TSet<FUnrealSourceFile*> ExportedSourceFiles;
+		Exported = GetSourceFilesInDependencyOrder(Package, SourceFiles, ExportedSourceFiles);
+		ExportedSourceFiles.Append(Exported);
+	}
 
-			// Load the original header file into memory
-			FFileHelper::LoadFileToString(OriginalHeader, *ClassesHeaderPath);
+	for (FUnrealSourceFile* SourceFile : Exported)
+	{
+		FUHTStringBuilder GeneratedHeaderText;
 
-			// Write the classes and enums header prefixes.
+		NameLookupCPP.SetCurrentSourceFile(SourceFile);
+		UHTMakefile.AddToHeaderOrder(SourceFile);
 
-			PreHeaderText.Logf(
-				TEXT("// Copyright 1998-2017 Epic Games, Inc. All Rights Reserved.\r\n")
-				TEXT("/*===========================================================================\r\n")
-				TEXT("\tC++ class boilerplate exported from UnrealHeaderTool.\r\n")
-				TEXT("\tThis is automatically generated by the tools.\r\n")
-				TEXT("\tDO NOT modify this manually! Edit the corresponding .h files instead!\r\n")
-				TEXT("===========================================================================*/\r\n")
-				TEXT("#pragma once\r\n")
-				TEXT("\r\n")
-				);
+		TArray<UEnum*>             Enums;
+		TArray<UScriptStruct*>     Structs;
+		TArray<UDelegateFunction*> DelegateFunctions;
+		SourceFile->GetScope()->SplitTypesIntoArrays(Enums, Structs, DelegateFunctions);
 
-			// if a global auto-include file exists, generate a line to have that file included
-			FString GlobalAutoIncludeFilename = PackageName + TEXT("GlobalIncludes.h");
-			const FString StandardHeaderFileLocation = PkgDir / TEXT("Public");
-			if (IFileManager::Get().FileSize(*(StandardHeaderFileLocation / GlobalAutoIncludeFilename)) > 0)
-			{
-				PreHeaderText.Logf(TEXT("#include \"%s\"\r\n\r\n"), *GlobalAutoIncludeFilename);
-			}
+		// Reverse the containers as they come out in the reverse order of declaration
+		Algo::Reverse(Enums);
+		Algo::Reverse(Structs);
+		Algo::Reverse(DelegateFunctions);
 
-			// Export an include line for each header
-			for (auto* SourceFile : SourceFiles)
-			{
-				ExportSourceFileHeader(AllClasses, SourceFile);
-			}
+		ExportSourceFileHeader(GeneratedHeaderText, OutputToGeneratedTextDevice, GeneratedFunctionDeclarations, *SourceFile, Enums, Structs, DelegateFunctions);
 
-			// build the full header file out of its pieces
-			const FString FullClassesHeader = FString::Printf(
-				TEXT("%s\r\n%s"),
-				*PreHeaderText,
-				*GetListOfPublicHeaderGroupIncludesString(InPackage)
-				);
+		const FString PkgName = FPackageName::GetShortName(Package);
+
+		FString PkgDir;
+		FString GeneratedIncludeDirectory;
+		if (!FindPackageLocation(*PkgName, PkgDir, GeneratedIncludeDirectory))
+		{
+			UE_LOG(LogCompile, Error, TEXT("Failed to find path for package %s"), *PkgName);
+		}
+
+		const FString ClassHeaderPath = GeneratedIncludeDirectory / FPaths::GetBaseFilename(SourceFile->GetFilename()) + TEXT(".generated.h");
 
-			// Save the classes header if it has changed.
-			SaveHeaderIfChanged(*ClassesHeaderPath, *FullClassesHeader);
+		bool bHasChanged = WriteHeader(*ClassHeaderPath, GeneratedHeaderText, ForwardDeclarations);
+
+		SourceFile->SetGeneratedFilename(ClassHeaderPath);
+		SourceFile->SetHasChanged(bHasChanged);
+
+		ForwardDeclarations.Reset();
+
+		if (GPublicSourceFileSet.Contains(SourceFile))
+		{
+			PublicHeaderGroupIncludes.AddUnique(SourceFile);
 		}
 	}
-	else
+
+	if (bWriteClassesH)
 	{
-		// If no headers are generated, check if there's any no export classes that need to have the inl file generated.
-		for (auto* SourceFile : SourceFiles)
+		// Write the classes and enums header prefixes.
+
+		FString PkgDir;
+		FString GeneratedIncludeDirectory;
+		if (!FindPackageLocation(*PackageName, PkgDir, GeneratedIncludeDirectory))
+		{
+			UE_LOG(LogCompile, Error, TEXT("Failed to find path for package %s"), *PackageName);
+		}
+
+		FUHTStringBuilder ClassesHText;
+		ClassesHText.Log(HeaderCopyright);
+		ClassesHText.Log(TEXT("#pragma once\r\n"));
+		ClassesHText.Log(TEXT("\r\n"));
+		ClassesHText.Log(TEXT("\r\n"));
+
+		// Fill with the rest source files from this package.
+		for (FUnrealSourceFile* SourceFile : GPublicSourceFileSet)
+		{
+			if (SourceFile->GetPackage() == InPackage && !PublicHeaderGroupIncludes.Contains(SourceFile))
+			{
+				PublicHeaderGroupIncludes.Add(SourceFile);
+			}
+		}
+
+		for (FUnrealSourceFile* SourceFile : PublicHeaderGroupIncludes)
 		{
-			ExportSourceFileHeader(AllClasses, SourceFile);
+			ClassesHText.Logf(TEXT("#include \"%s\"") LINE_TERMINATOR, *GetBuildPath(*SourceFile));
 		}
+
+		ClassesHText.Log(LINE_TERMINATOR);
+
+		// Save the classes header if it has changed.
+		FString ClassesHeaderPath = GeneratedIncludeDirectory / (PackageName + TEXT("Classes.h"));
+		SaveHeaderIfChanged(*ClassesHeaderPath, *ClassesHText);
 	}
 
 	// now export the names for the functions in this package
 	// notice we always export this file (as opposed to only exporting if we have any marked names)
 	// because there would be no way to know when the file was created otherwise
 	// Export .generated.cpp
-	ExportGeneratedCPP();
+	UE_LOG(LogCompile, Log, TEXT("Autogenerating boilerplate cpp: %s.generated.cpp"), *PackageName);
 
-	ExportGeneratedProto();
+	if (GeneratedFunctionDeclarations.Len() || UniqueCrossModuleReferences.Num() > 0)
+	{
+		ExportGeneratedPackageInitCode(OutputToGeneratedTextDevice(), GeneratedFunctionDeclarations, Package);
+	}
 
-	ExportGeneratedMCP();
+	ExportGeneratedCPP(*GeneratedFunctionDeclarations, Exported);
 
 	// Export all changed headers from their temp files to the .h files
 	ExportUpdatedHeaders(PackageName);
@@ -5229,7 +4544,7 @@ void FNativeClassHeaderGenerator::DeleteUnusedGeneratedHeaders()
 	TSet<FString> AllIntermediateFolders;
 	TSet<FString> PackageHeaderPathSet(PackageHeaderPaths);
 
-	for (const auto& PackageHeader : PackageHeaderPaths)
+	for (const FString& PackageHeader : PackageHeaderPaths)
 	{
 		const FString IntermediatePath = FPaths::GetPath(PackageHeader);
 
@@ -5243,7 +4558,7 @@ void FNativeClassHeaderGenerator::DeleteUnusedGeneratedHeaders()
 		TArray<FString> AllHeaders;
 		IFileManager::Get().FindFiles( AllHeaders, *(IntermediatePath / TEXT("*.generated.h")), true, false );
 
-		for (const auto& Header : AllHeaders)
+		for (const FString& Header : AllHeaders)
 		{
 			const FString HeaderPath = IntermediatePath / Header;
 
@@ -5414,14 +4729,12 @@ FString FNativeClassHeaderGenerator::GenerateTempHeaderName( FString CurrentFile
 */
 void FNativeClassHeaderGenerator::ExportUpdatedHeaders(FString PackageName)
 {
-	for (auto It = TempHeaderPaths.CreateConstIterator(); It; ++It)
+	for (const FString& TmpFilename : TempHeaderPaths)
 	{
-		const FString& TmpFilename = *It;
-
 		FString Filename = GenerateTempHeaderName( TmpFilename, true );
 		if (!IFileManager::Get().Move(*Filename, *TmpFilename, true, true))
 		{
-			UE_LOG(LogCompile, Error, TEXT("%s"), *FString::Printf(TEXT("Error exporting %s: couldn't write file '%s'"),*PackageName,*Filename));
+			UE_LOG(LogCompile, Error, TEXT("Error exporting %s: couldn't write file '%s'"), *PackageName, *Filename);
 		}
 		else
 		{
@@ -5430,113 +4743,28 @@ void FNativeClassHeaderGenerator::ExportUpdatedHeaders(FString PackageName)
 	}
 }
 
-/**
- * Exports protobuffer definitions from boilerplate that was generated for a package.
- * They are exported to a file using the name <PackageName>.generated.proto
- */
-void FNativeClassHeaderGenerator::ExportGeneratedProto()
-{
-	if (GeneratedProtoText.Len())
-	{
-		UE_LOG(LogCompile, Log,  TEXT("Autogenerating boilerplate proto: %s.proto"), *GeneratedProtoFilenameBase );
-
-		FUHTStringBuilder ProtoPreamble;
-
-		ProtoPreamble.Logf(
-			TEXT("// Copyright 1998-2017 Epic Games, Inc. All Rights Reserved.")					LINE_TERMINATOR
-			TEXT("/*===========================================================================")	LINE_TERMINATOR
-			TEXT("  Purpose: The file defines our Google Protocol Buffers which are used in over ") LINE_TERMINATOR
-			TEXT("  the wire messages between servers as well as between clients and servers.")		LINE_TERMINATOR
-			TEXT("  This is automatically generated by UnrealHeaderTool.")							LINE_TERMINATOR
-			TEXT("  DO NOT modify this manually! Edit the corresponding .h files instead!")			LINE_TERMINATOR
-			TEXT("===========================================================================*/")	LINE_TERMINATOR
-																									LINE_TERMINATOR
-			TEXT("// We care more about speed than code size")										LINE_TERMINATOR
-			TEXT("option optimize_for = SPEED;")													LINE_TERMINATOR
-			TEXT("// We don't use the service generation functionality")							LINE_TERMINATOR
-			TEXT("option cc_generic_services = false;")												LINE_TERMINATOR
-																									LINE_TERMINATOR
-			);
-
-		FString PkgName = FPackageName::GetShortName(Package);
-		FString PkgDir;
-		FString GeneratedIncludeDirectory;
-		if (FindPackageLocation(*PkgName, PkgDir, GeneratedIncludeDirectory) == false)
-		{
-			UE_LOG(LogCompile, Error, TEXT("Failed to find path for package %s"), *PkgName);
-		}
-
-		FString HeaderPath = GeneratedIncludeDirectory / GeneratedProtoFilenameBase + TEXT(".proto");
-		SaveHeaderIfChanged(*HeaderPath, *(ProtoPreamble + GeneratedProtoText));
-	}
-}
-
-/**
- * Exports MCPbuffer definitions from boilerplate that was generated for a package.
- * They are exported to a file using the name <PackageName>.generated.MCP
- */
-void FNativeClassHeaderGenerator::ExportGeneratedMCP()
-{
-	if (GeneratedMCPText.Len())
-	{
-		UE_LOG(LogCompile, Log,  TEXT("Autogenerating boilerplate MCP: %s.java"), *GeneratedMCPFilenameBase );
-
-		FUHTStringBuilder MCPPreamble;
-
-		MCPPreamble.Logf(
-			TEXT("// Copyright 1998-2017 Epic Games, Inc. All Rights Reserved.")					LINE_TERMINATOR
-			TEXT("/*===========================================================================")	LINE_TERMINATOR
-			TEXT("  Purpose: The file defines java heaers for MCP rpc messages. ")					LINE_TERMINATOR
-			TEXT("  DO NOT modify this manually! Edit the corresponding .h files instead!")			LINE_TERMINATOR
-			TEXT("===========================================================================*/")	LINE_TERMINATOR
-			);
-
-		FString PkgName = FPackageName::GetShortName(Package);
-		FString PkgDir;
-		FString GeneratedIncludeDirectory;
-		if (FindPackageLocation(*PkgName, PkgDir, GeneratedIncludeDirectory) == false)
-		{
-			UE_LOG(LogCompile, Error, TEXT("Failed to find path for package %s"), *PkgName);
-		}
-
-		FString HeaderPath = GeneratedIncludeDirectory / GeneratedMCPFilenameBase + TEXT(".java");
-		SaveHeaderIfChanged(*HeaderPath, *(MCPPreamble + GeneratedMCPText));
-	}
-}
-
 /**
  * Exports C++ definitions for boilerplate that was generated for a package.
  * They are exported to a file using the name <PackageName>.generated.cpp
  * @param ReferencedNames list of function names to export.
  */
-void FNativeClassHeaderGenerator::ExportGeneratedCPP()
+void FNativeClassHeaderGenerator::ExportGeneratedCPP(const TCHAR* Declarations, const TArray<FUnrealSourceFile*>& Exported)
 {
-	FUHTStringBuilder GeneratedCPPPreamble;
-	FUHTStringBuilder GeneratedCPPClassesIncludes;
-	FUHTStringBuilder GeneratedCPPEpilogue;
-	FUHTStringBuilder GeneratedCPPText;
-
-	TArray<FUHTStringBuilder> GeneratedCPPFiles;
+	static const TCHAR EnableDeprecationWarnings [] = TEXT("PRAGMA_ENABLE_DEPRECATION_WARNINGS") LINE_TERMINATOR;
+	static const TCHAR DisableDeprecationWarnings[] = TEXT("PRAGMA_DISABLE_DEPRECATION_WARNINGS") LINE_TERMINATOR;
+	static const TCHAR DisableWarning4883        [] = TEXT("#ifdef _MSC_VER") LINE_TERMINATOR TEXT("#pragma warning (push)") LINE_TERMINATOR TEXT("#pragma warning (disable : 4883)") LINE_TERMINATOR TEXT("#endif") LINE_TERMINATOR;
+	static const TCHAR EnableWarning4883         [] = TEXT("#ifdef _MSC_VER") LINE_TERMINATOR TEXT("#pragma warning (pop)") LINE_TERMINATOR TEXT("#endif") LINE_TERMINATOR;
 
-	UE_LOG(LogCompile, Log,  TEXT("Autogenerating boilerplate cpp: %s.cpp"), *GeneratedCPPFilenameBase );
-
-	GeneratedCPPPreamble.Logf(
-		TEXT("// Copyright 1998-2017 Epic Games, Inc. All Rights Reserved.")					LINE_TERMINATOR
-		TEXT("/*===========================================================================")	LINE_TERMINATOR
-		TEXT("\tBoilerplate C++ definitions for a single module.")								LINE_TERMINATOR
-		TEXT("\tThis is automatically generated by UnrealHeaderTool.")							LINE_TERMINATOR
-		TEXT("\tDO NOT modify this manually! Edit the corresponding .h files instead!")			LINE_TERMINATOR
-		TEXT("===========================================================================*/")	LINE_TERMINATOR
-																								LINE_TERMINATOR
-		);
+	const FManifestModule* ModuleInfo = GPackageToManifestModuleMap.FindChecked(Package);
 
 	// Write out an include for the generated header boilerplate code
+	FUHTStringBuilder GeneratedCPPPreamble;
+	GeneratedCPPPreamble.Log(HeaderCopyright);
+	GeneratedCPPPreamble.Log(LINE_TERMINATOR);
 	//GeneratedCPPPreamble.Logf(TEXT("#pragma once") LINE_TERMINATOR);
 	GeneratedCPPPreamble.Logf(TEXT("#include \"GeneratedCppIncludes.h\"") LINE_TERMINATOR);
 
 	FString ModulePCHInclude;
-
-	const auto* ModuleInfo = GPackageToManifestModuleMap.FindChecked(Package);
 	if (ModuleInfo->PCH.Len())
 	{
 		FString PCH = ModuleInfo->PCH;
@@ -5544,47 +4772,45 @@ void FNativeClassHeaderGenerator::ExportGeneratedCPP()
 		ModulePCHInclude = FString::Printf(TEXT("#include \"%s\"") LINE_TERMINATOR, *PCH);
 	}
 
+	FUHTStringBuilder Includes;
+	for (const FUnrealSourceFile* SourceFile : Exported)
+	{
+		FString NewFileName = SourceFile->GetFilename();
+		ConvertToBuildIncludePath(Package, NewFileName);
+
+		FString IncludeStr = FString::Printf(
+			TEXT("#ifndef %s")			LINE_TERMINATOR
+			TEXT("\t#include \"%s\"")	LINE_TERMINATOR
+			TEXT("#endif")				LINE_TERMINATOR,
+			*SourceFile->GetFileDefineName(), *NewFileName);
+
+		Includes.Log(*IncludeStr);
+	}
+
 	// Write out the ordered class dependencies into a single header that we can easily include
 	FString DepHeaderPathname = ModuleInfo->GeneratedCPPFilenameBase + TEXT(".dep.h");
-	SaveHeaderIfChanged(*DepHeaderPathname, *(GeneratedCPPPreamble + ListOfPublicClassesUObjectHeaderModuleIncludes));
+	SaveHeaderIfChanged(*DepHeaderPathname, *(GeneratedCPPPreamble + Includes));
 
 	// Write out our include to the .dep.h file
+	FUHTStringBuilder GeneratedCPPClassesIncludes;
 	GeneratedCPPClassesIncludes.Logf(TEXT("#include \"%s\"") LINE_TERMINATOR, *FPaths::GetCleanFilename(DepHeaderPathname));
-
 	GeneratedCPPClassesIncludes.Logf(TEXT("PRAGMA_DISABLE_OPTIMIZATION") LINE_TERMINATOR);
-	{
-		// Autogenerate names (alphabetically sorted).
-		ReferencedNames.KeySort( TLess<FName>() );
-
-		for (auto& PairIt : ReferencedNames)
-		{
-			GeneratedCPPText.Logf(TEXT("FName %s_%s = FName(TEXT(\"%s\"));") LINE_TERMINATOR, *API, *PairIt.Key.ToString(), *PairIt.Value.ToString());
-		}
-	}
-
-	GeneratedCPPText.Log(*GeneratedPackageCPP);
-
-	GeneratedCPPEpilogue.Logf(
-		LINE_TERMINATOR
-		);
 
-	FString EnableDeprecationWarnings = FString(TEXT("PRAGMA_ENABLE_DEPRECATION_WARNINGS") LINE_TERMINATOR);
-	FString DisableDeprecationWarnings = FString(TEXT("PRAGMA_DISABLE_DEPRECATION_WARNINGS") LINE_TERMINATOR);
-	
-	const bool bNativizaedAssetsModule = (ModuleInfo->Name == TEXT("NativizedAssets"));
-	const FString DisableWarning4883(TEXT("#ifdef _MSC_VER") LINE_TERMINATOR TEXT("#pragma warning (push)") LINE_TERMINATOR TEXT("#pragma warning (disable : 4883)") LINE_TERMINATOR TEXT("#endif") LINE_TERMINATOR);
-	const FString EnableWarning4883(TEXT("#ifdef _MSC_VER") LINE_TERMINATOR TEXT("#pragma warning (pop)") LINE_TERMINATOR TEXT("#endif") LINE_TERMINATOR);
-
-	FString PkgName = FPackageName::GetShortName(Package);
-	FString PkgDir;
-	if (GeneratedFunctionDeclarations.Len() || CrossModuleGeneratedFunctionDeclarations.Len())
+	// Autogenerate names (alphabetically sorted).
+	FUHTStringBuilder GeneratedCPPText;
+	ReferencedNames.KeySort(TLess<>());
+	for (const TPair<FName, FName>& Pair : ReferencedNames)
 	{
-		ExportGeneratedPackageInitCode(Package);
+		GeneratedCPPText.Logf(TEXT("FName %s_%s = FName(TEXT(\"%s\"));") LINE_TERMINATOR, *API, *Pair.Key.ToString(), *Pair.Value.ToString());
 	}
 
-	TArray<FString> NumberedHeaderNames;
+	FUHTStringBuilder GeneratedCPPEpilogue;
+	GeneratedCPPEpilogue.Logf(LINE_TERMINATOR);
+
+	const bool bNativizedAssetsModule = ModuleInfo->Name == TEXT("NativizedAssets");
 
 	// Generate each of the .generated.cpp files
+	TArray<FString> NumberedHeaderNames;
 	for( int32 FileIdx=0;FileIdx<GeneratedFunctionBodyTextSplit.Num();FileIdx++ )
 	{
 		FUHTStringBuilder FileText;
@@ -5594,16 +4820,19 @@ void FNativeClassHeaderGenerator::ExportGeneratedCPP()
 			FileText = GeneratedCPPText;
 		}
 
-		if (GeneratedFunctionDeclarations.Len() || CrossModuleGeneratedFunctionDeclarations.Len())
+		if (*Declarations || UniqueCrossModuleReferences.Num() > 0)
 		{
 			FileText.Logf(TEXT("#if USE_COMPILED_IN_NATIVES\r\n"));
-			if (CrossModuleGeneratedFunctionDeclarations.Len())
+			if (UniqueCrossModuleReferences.Num() > 0)
 			{
 				FileText.Logf(TEXT("// Cross Module References\r\n"));
-				FileText.Log(CrossModuleGeneratedFunctionDeclarations);
+				for (const FString& Ref : UniqueCrossModuleReferences)
+				{
+					FileText.Log(*Ref);
+				}
 				FileText.Logf(TEXT("\r\n"));
 			}
-			FileText.Log(GeneratedFunctionDeclarations);
+			FileText.Log(Declarations);
 			FileText.Log(GeneratedFunctionBodyTextSplit[FileIdx].Get());
 			FileText.Logf(TEXT("#endif\r\n"));
 		}
@@ -5613,13 +4842,13 @@ void FNativeClassHeaderGenerator::ExportGeneratedCPP()
 		SaveHeaderIfChanged(*CppPath, *(GeneratedCPPPreamble
 			+ ModulePCHInclude
 			+ GeneratedCPPClassesIncludes
-			+ (bNativizaedAssetsModule ? DisableWarning4883 : FString())
+			+ (bNativizedAssetsModule ? DisableWarning4883 : TEXT(""))
 			+ DisableDeprecationWarnings
 			+ GeneratedLinkerFixupFunction
 			+ FileText 
 			+ GeneratedCPPEpilogue 
 			+ EnableDeprecationWarnings
-			+ (bNativizaedAssetsModule ? EnableWarning4883 : FString())
+			+ (bNativizedAssetsModule ? EnableWarning4883 : TEXT(""))
 			));
 
 		if (GeneratedFunctionBodyTextSplit.Num() > 1)
@@ -5862,9 +5091,8 @@ ECompilationResult::Type PreparseModules(FUHTMakefile& UHTMakefile, const FStrin
 
 					if (CurrentlyProcessing == PublicClassesHeaders)
 					{
-						for (auto* Class : UnrealSourceFile->GetDefinedClasses())
+						for (auto* Class : DefinedClasses)
 						{
-							GPublicClassSet.Add(Class);
 							UHTMakefile.AddPublicClassSetEntry(UnrealSourceFilePtr, Class);
 						}
 
@@ -6086,7 +5314,7 @@ ECompilationResult::Type UnrealHeaderTool_Main(const FString& ModuleInfoFilename
 			TArray<IScriptGeneratorPluginInterface*> ScriptPlugins;
 			// Can only export scripts for game targets
 			if (GManifest.IsGameTarget)
-			{				
+			{
 				GetScriptPlugins(ScriptPlugins);
 			}
 
@@ -6103,28 +5331,7 @@ ECompilationResult::Type UnrealHeaderTool_Main(const FString& ModuleInfoFilename
 					FClasses AllClasses(Package);
 					AllClasses.Validate();
 
-#if WITH_HOT_RELOAD_CTORS
-					static struct FUseVTableConstructorsCache
-					{
-						FUseVTableConstructorsCache()
-						{
-							bUseVTableConstructors = false;
-							GConfig->GetBool(TEXT("Core.System"), TEXT("UseVTableConstructors"), bUseVTableConstructors, GEngineIni);
-						}
-
-						bool bUseVTableConstructors;
-					} UseVTableConstructorsCache;
-
-					bool bModuleIsGame =
-						   Module.ModuleType == EBuildModuleType::GameDeveloper
-						|| Module.ModuleType == EBuildModuleType::GameEditor
-						|| Module.ModuleType == EBuildModuleType::GameRuntime
-						|| Module.ModuleType == EBuildModuleType::GameThirdParty;
-
-					Result = FHeaderParser::ParseAllHeadersInside(AllClasses, GWarn, Package, Module, ScriptPlugins, bModuleIsGame || UseVTableConstructorsCache.bUseVTableConstructors, UHTMakefile);
-#else // WITH_HOT_RELOAD_CTORS
 					Result = FHeaderParser::ParseAllHeadersInside(AllClasses, GWarn, Package, Module, ScriptPlugins, UHTMakefile);
-#endif // WITH_HOT_RELOAD_CTORS
 					if (Result != ECompilationResult::Succeeded)
 					{
 						++NumFailures;
@@ -6317,7 +5524,7 @@ TSharedRef<FUnrealSourceFile> PerformInitialParseOnHeader(UPackage* InParent, co
 
 	for (auto& DependsOnElement : DependsOn)
 	{
-		UnrealSourceFile->GetIncludes().Add(DependsOnElement);
+		UnrealSourceFile->GetIncludes().AddUnique(DependsOnElement);
 	}
 
 	return UnrealSourceFile;
