Copying //UE4/Dev-Core to //UE4/Dev-Main (Source: //UE4/Dev-Core @ 3345728)

diff --git a/Engine/Source/Programs/UnrealHeaderTool/Private/HeaderParser.cpp b/Engine/Source/Programs/UnrealHeaderTool/Private/HeaderParser.cpp
--- a/Engine/Source/Programs/UnrealHeaderTool/Private/HeaderParser.cpp
+++ b/Engine/Source/Programs/UnrealHeaderTool/Private/HeaderParser.cpp
@@ -206,63 +206,48 @@ namespace
 	 */
 	void ParseNetServiceIdentifiers(FFuncInfo& FuncInfo, const TArray<FString>& Identifiers)
 	{
-		FString IdTag         (TEXT("Id="));
-		FString ResponseIdTag (TEXT("ResponseId="));
-		FString MCPTag        (TEXT("MCP"));
-		FString ProtobufferTag(TEXT("Protobuffer"));
-		for (auto& Identifier : Identifiers)
+		static const TCHAR IdTag        [] = TEXT("Id");
+		static const TCHAR ResponseIdTag[] = TEXT("ResponseId");
+
+		for (const FString& Identifier : Identifiers)
 		{
-			if (Identifier == ProtobufferTag)
-			{
-				FuncInfo.FunctionExportFlags |= FUNCEXPORT_NeedsProto;
-			}
-			else if (Identifier == MCPTag)
-			{
-				FuncInfo.FunctionExportFlags |= FUNCEXPORT_NeedsMCP;
-			}
-			else if (Identifier.StartsWith(IdTag))
+			const TCHAR* IdentifierPtr = *Identifier;
+
+			if (const TCHAR* Equals = FCString::Strchr(IdentifierPtr, TEXT('=')))
 			{
-				int32 TempInt = FCString::Atoi(*Identifier.Mid(IdTag.Len()));
-				if (TempInt <= 0 || TempInt > MAX_uint16)
+				// It's a tag with an argument
+
+				if (FCString::Strnicmp(IdentifierPtr, IdTag, ARRAY_COUNT(IdTag) - 1) == 0)
 				{
-					FError::Throwf(TEXT("Invalid network identifier %s for function"), *Identifier);
+					int32 TempInt = FCString::Atoi(Equals + 1);
+					if (TempInt <= 0 || TempInt > MAX_uint16)
+					{
+						FError::Throwf(TEXT("Invalid network identifier %s for function"), IdentifierPtr);
+					}
+					FuncInfo.RPCId = TempInt;
 				}
-				FuncInfo.RPCId = TempInt;
-			}
-			else if (Identifier.StartsWith(ResponseIdTag))
-			{
-				int32 TempInt = FCString::Atoi(*Identifier.Mid(ResponseIdTag.Len()));
-				if (TempInt <= 0 || TempInt > MAX_uint16)
+				else if (FCString::Strnicmp(IdentifierPtr, ResponseIdTag, ARRAY_COUNT(ResponseIdTag) - 1) == 0)
 				{
-					FError::Throwf(TEXT("Invalid network identifier %s for function"), *Identifier);
+					int32 TempInt = FCString::Atoi(Equals + 1);
+					if (TempInt <= 0 || TempInt > MAX_uint16)
+					{
+						FError::Throwf(TEXT("Invalid network identifier %s for function"), IdentifierPtr);
+					}
+					FuncInfo.RPCResponseId = TempInt;
 				}
-				FuncInfo.RPCResponseId = TempInt;
 			}
 			else
 			{
-				FError::Throwf(TEXT("Invalid network identifier %s for function"), *Identifier);
-			}
-		}
-		if (FuncInfo.FunctionExportFlags & FUNCEXPORT_NeedsProto)
-		{
-			if (FuncInfo.RPCId == 0)
-			{
-				FError::Throwf(TEXT("net service function does not have an RPCId."));
-			}
-			if (FuncInfo.RPCId == FuncInfo.RPCResponseId)
-			{
-				FError::Throwf(TEXT("Net service RPCId and ResponseRPCId cannot be the same."));
-			}
-			if ((FuncInfo.FunctionFlags & FUNC_NetResponse) && FuncInfo.RPCResponseId > 0)
-			{
-				FError::Throwf(TEXT("Net service response functions cannot have a ResponseId."));
+				// Assume it's an endpoint name
+
+				if (FuncInfo.EndpointName.Len())
+				{
+					FError::Throwf(TEXT("Function should not specify multiple endpoints - '%s' found but already using '%s'"), *Identifier);
+				}
+
+				FuncInfo.EndpointName = Identifier;
 			}
 		}
-	
-		if (!(FuncInfo.FunctionExportFlags & FUNCEXPORT_NeedsProto) && !(FuncInfo.FunctionExportFlags & FUNCEXPORT_NeedsMCP))
-		{
-			FError::Throwf(TEXT("net service function needs to specify at least one provider type."));
-		}
 	}
 
 	/**
@@ -408,6 +393,11 @@ namespace
 					FuncInfo.FunctionExportFlags |= FUNCEXPORT_CustomThunk;
 
 					ParseNetServiceIdentifiers(FuncInfo, Specifier.Values);
+
+					if (FuncInfo.EndpointName.Len() == 0)
+					{
+						FError::Throwf(TEXT("ServiceRequest needs to specify an endpoint name"));
+					}
 				}
 				break;
 
@@ -423,6 +413,11 @@ namespace
 					FuncInfo.FunctionFlags |= FUNC_NetResponse;
 
 					ParseNetServiceIdentifiers(FuncInfo, Specifier.Values);
+
+					if (FuncInfo.EndpointName.Len() == 0)
+					{
+						FError::Throwf(TEXT("ServiceResponse needs to specify an endpoint name"));
+					}
 				}
 				break;
 
@@ -2310,6 +2305,14 @@ UScriptStruct* FHeaderParser::CompileStructDeclaration(FClasses& AllClasses)
 				}
 			}
 		}
+		else if (Token.Matches(TEXT("#")) && MatchIdentifier(TEXT("pragma")))
+		{
+			// skip it and skip over the text, it is not recorded or processed
+			TCHAR c;
+			while (!IsEOL(c = GetChar()))
+			{
+			}
+		}
 		else if (ProbablyAnUnknownObjectLikeMacro(*this, Token))
 		{
 			// skip it
@@ -7272,12 +7275,9 @@ TArray<FUnrealSourceFile*> GetSourceFilesWithInheritanceOrdering(UPackage* Curre
 void FHeaderParser::ExportNativeHeaders(
 	UPackage* CurrentPackage,
 	FClasses& AllClasses,
-	bool bAllowSaveExportedHeaders
-#if WITH_HOT_RELOAD_CTORS
-	, bool bExportVTableConstructors
-#endif // WITH_HOT_RELOAD_CTORS
-	, FUHTMakefile& UHTMakefile
-	, const FManifestModule& Module
+	bool bAllowSaveExportedHeaders,
+	FUHTMakefile& UHTMakefile,
+	const FManifestModule& Module
 )
 {
 	// Build a list of header filenames
@@ -7313,11 +7313,8 @@ void FHeaderParser::ExportNativeHeaders(
 				CurrentPackage,
 				SourceFiles,
 				AllClasses,
-				bAllowSaveExportedHeaders
-#if WITH_HOT_RELOAD_CTORS
-				, bExportVTableConstructors
-#endif // WITH_HOT_RELOAD_CTORS
-				, UHTMakefile
+				bAllowSaveExportedHeaders,
+				UHTMakefile
 			);
 		}
 		UHTMakefile.StopExporting();
@@ -7423,11 +7420,8 @@ ECompilationResult::Type FHeaderParser::ParseAllHeadersInside(
 	FFeedbackContext* Warn,
 	UPackage* CurrentPackage,
 	const FManifestModule& Module,
-	TArray<IScriptGeneratorPluginInterface*>& ScriptPlugins
-#if WITH_HOT_RELOAD_CTORS
-	, bool bExportVTableConstructors
-#endif // WITH_HOT_RELOAD_CTORS
-	, FUHTMakefile& UHTMakefile
+	TArray<IScriptGeneratorPluginInterface*>& ScriptPlugins,
+	FUHTMakefile& UHTMakefile
 	)
 {
 	// Disable loading of objects outside of this package (or more exactly, objects which aren't UFields, CDO, or templates)
@@ -7498,13 +7492,10 @@ ECompilationResult::Type FHeaderParser::ParseAllHeadersInside(
 				ExportNativeHeaders(
 					CurrentPackage,
 					ModuleClasses,
-					Module.SaveExportedHeaders
-#if WITH_HOT_RELOAD_CTORS
-					, bExportVTableConstructors
-#endif // WITH_HOT_RELOAD_CTORS
-					, UHTMakefile
-					, Module
-					);
+					Module.SaveExportedHeaders,
+					UHTMakefile,
+					Module
+				);
 			}
 			GHeaderCodeGenTime += ExportTime;
 
@@ -7675,12 +7666,10 @@ EBlockDirectiveType ParseCommandToBlockDirectiveType(const TCHAR** Str)
 		return EBlockDirectiveType::NotCPPBlock;
 	}
 
-#if WITH_HOT_RELOAD_CTORS
 	if (FParse::Command(Str, TEXT("WITH_HOT_RELOAD")))
 	{
 		return EBlockDirectiveType::WithHotReload;
 	}
-#endif
 
 	if (FParse::Command(Str, TEXT("WITH_EDITOR")))
 	{
@@ -8468,18 +8457,14 @@ bool FHeaderParser::TryToMatchConstructorParameterList(FToken Token)
 	check(ClassData);
 
 	bool bOICtor = false;
-#if WITH_HOT_RELOAD_CTORS
 	bool bVTCtor = false;
-#endif // WITH_HOT_RELOAD_CTORS
 
 	if (!ClassData->bDefaultConstructorDeclared && MatchSymbol(TEXT(")")))
 	{
 		ClassData->bDefaultConstructorDeclared = true;
 	}
 	else if (!ClassData->bObjectInitializerConstructorDeclared
-#if WITH_HOT_RELOAD_CTORS
 		|| !ClassData->bCustomVTableHelperConstructorDeclared
-#endif // WITH_HOT_RELOAD_CTORS
 	)
 	{
 		FToken ObjectInitializerParamParsingToken;
@@ -8494,9 +8479,7 @@ bool FHeaderParser::TryToMatchConstructorParameterList(FToken Token)
 			if (ObjectInitializerParamParsingToken.Matches(TEXT(",")) || ObjectInitializerParamParsingToken.Matches(TEXT("<")))
 			{
 				bOICtor = false;
-#if WITH_HOT_RELOAD_CTORS
 				bVTCtor = false;
-#endif // WITH_HOT_RELOAD_CTORS
 				break;
 			}
 
@@ -8531,12 +8514,10 @@ bool FHeaderParser::TryToMatchConstructorParameterList(FToken Token)
 				bOICtor = true;
 			}
 
-#if WITH_HOT_RELOAD_CTORS
 			if (ObjectInitializerParamParsingToken.Matches(TEXT("FVTableHelper")))
 			{
 				bVTCtor = true;
 			}
-#endif // WITH_HOT_RELOAD_CTORS
 		}
 
 		// Parse until finish.
@@ -8556,17 +8537,10 @@ bool FHeaderParser::TryToMatchConstructorParameterList(FToken Token)
 		}
 
 		ClassData->bObjectInitializerConstructorDeclared = ClassData->bObjectInitializerConstructorDeclared || (bOICtor && bIsRef && bIsConst);
-#if WITH_HOT_RELOAD_CTORS
 		ClassData->bCustomVTableHelperConstructorDeclared = ClassData->bCustomVTableHelperConstructorDeclared || (bVTCtor && bIsRef);
-#endif // WITH_HOT_RELOAD_CTORS
 	}
 
-	ClassData->bConstructorDeclared =
-#if WITH_HOT_RELOAD_CTORS
-		ClassData->bConstructorDeclared || !bVTCtor;
-#else // WITH_HOT_RELOAD_CTORS
-		true;
-#endif // WITH_HOT_RELOAD_CTORS
+	ClassData->bConstructorDeclared = ClassData->bConstructorDeclared || !bVTCtor;
 
 	// Optionally match semicolon.
 	if (!MatchSymbol(TEXT(";")))
diff --git a/Engine/Source/Programs/UnrealHeaderTool/Private/HeaderParser.h b/Engine/Source/Programs/UnrealHeaderTool/Private/HeaderParser.h
--- a/Engine/Source/Programs/UnrealHeaderTool/Private/HeaderParser.h
+++ b/Engine/Source/Programs/UnrealHeaderTool/Private/HeaderParser.h
@@ -170,11 +170,8 @@ public:
 		FFeedbackContext* Warn,
 		UPackage* LimitOuter,
 		const FManifestModule& Module,
-		TArray<class IScriptGeneratorPluginInterface*>& ScriptPlugins
-#if WITH_HOT_RELOAD_CTORS
-		, bool bExportVTableConstructors
-#endif // WITH_HOT_RELOAD_CTORS
-		, FUHTMakefile& UHTMakefile
+		TArray<class IScriptGeneratorPluginInterface*>& ScriptPlugins,
+		FUHTMakefile& UHTMakefile
 	);
 
 	// Performs a preliminary parse of the text in the specified buffer, pulling out:
@@ -471,12 +468,9 @@ protected:
 	static void ExportNativeHeaders(
 		UPackage* CurrentPackage,
 		FClasses& AllClasses,
-		bool bAllowSaveExportedHeaders
-#if WITH_HOT_RELOAD_CTORS
-		, bool bExportVTableConstructors
-#endif // WITH_HOT_RELOAD_CTORS
-		, FUHTMakefile& UHTMakefile
-		, const FManifestModule& Module
+		bool bAllowSaveExportedHeaders,
+		FUHTMakefile& UHTMakefile,
+		const FManifestModule& Module
 	);
 
 	// FContextSupplier interface.
diff --git a/Engine/Source/Programs/UnrealHeaderTool/Private/NativeClassExporter.h b/Engine/Source/Programs/UnrealHeaderTool/Private/NativeClassExporter.h
--- a/Engine/Source/Programs/UnrealHeaderTool/Private/NativeClassExporter.h
+++ b/Engine/Source/Programs/UnrealHeaderTool/Private/NativeClassExporter.h
@@ -52,6 +52,18 @@ class FScope;
 typedef FStringOutputDevice FUHTStringBuilder;
 typedef FStringOutputDeviceCountLines FUHTStringBuilderLineCounter;
 
+enum class EExportingState
+{
+	Normal,
+	TypeEraseDelegates
+};
+
+enum class EExportCallbackType
+{
+	Interface,
+	Class
+};
+
 struct FNativeClassHeaderGenerator
 {
 private:
@@ -65,38 +77,10 @@ private:
 		return FString::Printf(TEXT("%s_API "), *API);
 	}
 
-	FString ClassesHeaderPath;
-	/** the name of the cpp file where the implementation functions are generated, without the .cpp extension */
-	FString GeneratedCPPFilenameBase;
-	/** the name of the proto file where proto definitions are generated */
-	FString GeneratedProtoFilenameBase;
-	/** the name of the java file where mcp definitions are generated */
-	FString GeneratedMCPFilenameBase;
 	const UPackage* Package;
-	FUHTStringBuilder PreHeaderText;
-	FUHTStringBuilder EnumForeachText;
-	TArray<FUnrealSourceFile*> ListOfPublicHeaderGroupIncludes;			// This is built up each time from scratch each time for each header groups
-	FUHTStringBuilder ListOfPublicClassesUObjectHeaderModuleIncludes; // This is built up for the entire module, and for all processed headers
-	FUHTStringBuilder ListOfAllUObjectHeaderIncludes;
-	FUHTStringBuilder GeneratedPackageCPP;
-	FUHTStringBuilder GeneratedHeaderText;
-	FUHTStringBuilder GeneratedForwardDeclarations;
-    /** output generated for a .proto file */
-	FUHTStringBuilder GeneratedProtoText;
-	/** output generated for a mcp .java file */
-	FUHTStringBuilder GeneratedMCPText;
-	FUHTStringBuilder PrologMacroCalls;
-	FUHTStringBuilder InClassMacroCalls;
-	FUHTStringBuilder InClassNoPureDeclsMacroCalls;
-	FUHTStringBuilder StandardUObjectConstructorsMacroCall;
-	FUHTStringBuilder EnhancedUObjectConstructorsMacroCall;
 
 	/** Generated function implementations that belong in the cpp file, split into multiple files base on line count **/
 	TArray<TUniqueObj<FUHTStringBuilderLineCounter>> GeneratedFunctionBodyTextSplit;
-	/** Declarations of generated functions for this module**/
-	FUHTStringBuilder GeneratedFunctionDeclarations;
-	/** Declarations of generated functions for cross module**/
-	FUHTStringBuilder CrossModuleGeneratedFunctionDeclarations;
 	/** Set of already exported cross-module references, to prevent duplicates */
 	TSet<FString> UniqueCrossModuleReferences;
 
@@ -105,8 +89,6 @@ private:
 
 	/** the existing disk version of the header for this package's names */
 	FString OriginalNamesHeader;
-	/** the existing disk version of this header */
-	FString OriginalHeader;
 
 	/** Array of temp filenames that for files to overwrite headers */
 	TArray<FString> TempHeaderPaths;
@@ -114,25 +96,14 @@ private:
 	/** Array of all header filenames from the current package. */
 	TArray<FString> PackageHeaderPaths;
 
-	/** Array of all generated Classes headers for the current package */
-	TArray<FString> ClassesHeaders;
-	
-	/** if we are exporting a struct for offset determination only, replace some types with simpler ones (delegates) */
-	bool bIsExportingForOffsetDeterminationOnly;
-
 	/** If false, exported headers will not be saved to disk */
 	bool bAllowSaveExportedHeaders;
 
 	/** If true, any change in the generated headers will result in UHT failure. */
 	bool bFailIfGeneratedCodeChanges;
 
-#if WITH_HOT_RELOAD_CTORS
-	/** If true, then generate vtable constructors and it's callers. */
-	bool bExportVTableConstructors;
-#endif // WITH_HOT_RELOAD_CTORS
-
-	/** All properties that need to be forward declared. */
-	TArray<UProperty*>	ForwardDeclarations;
+	/** Forward declarations that we need. */
+	TSet<FString> ForwardDeclarations;
 
 	/** Makefile to save parsing data to. */
 	FUHTMakefile& UHTMakefile;
@@ -153,7 +124,7 @@ private:
 	 * @param	SortedHeaderFilenames	[out] receives the sorted list of header filenames.
 	 * @return	true upon success, false if there were circular dependencies which made it impossible to sort properly.
 	 */
-	bool SortHeaderDependencyMap(const TMap<const FString*, HeaderDependents>& HeaderDependencyMap, TArray<const FString*>& SortedHeaderFilenames ) const;
+	static bool SortHeaderDependencyMap(const TMap<const FString*, HeaderDependents>& HeaderDependencyMap, TArray<const FString*>& SortedHeaderFilenames );
 
 	/**
 	 * Finds to headers that are dependent on each other.
@@ -165,7 +136,7 @@ private:
 	 * @param  OutHeader2          [out] Receives the second inter-dependent header index.
 	 * @return true if an inter-dependency was found.
 	 */
-	bool FindInterDependency( TMap<const FString*, HeaderDependents>& HeaderDependencyMap, const FString* Header, const FString*& OutHeader1, const FString*& OutHeader2 );
+	static bool FindInterDependency( TMap<const FString*, HeaderDependents>& HeaderDependencyMap, const FString* Header, const FString*& OutHeader1, const FString*& OutHeader2 );
 
 	/**
 	 * Finds to headers that are dependent on each other.
@@ -177,41 +148,30 @@ private:
 	 * @param  OutHeader2          [out] Receives the second inter-dependent header index.
 	 * @return true if an inter-dependency was found.
 	 */
-	bool FindInterDependencyRecursive( TMap<const FString*, HeaderDependents>& HeaderDependencyMap, const FString* HeaderIndex, TSet<const FString*>& VisitedHeaders, const FString*& OutHeader1, const FString*& OutHeader2 );
-
-	/**
-	 * Determines whether the glue version of the specified native function
-	 * should be exported
-	 * 
-	 * @param	Function	the function to check
-	 * @return	true if the glue version of the function should be exported.
-	 */
-	bool ShouldExportFunction( UFunction* Function );
+	static bool FindInterDependencyRecursive( TMap<const FString*, HeaderDependents>& HeaderDependencyMap, const FString* HeaderIndex, TSet<const FString*>& VisitedHeaders, const FString*& OutHeader1, const FString*& OutHeader2 );
 
 	/**
 	 * Exports the struct's C++ properties to the HeaderText output device and adds special
 	 * compiler directives for GCC to pack as we expect.
 	 *
-	 * @param	Struct				UStruct to export properties
-	 * @param	TextIndent			Current text indentation
-	 * @param	Output	alternate output device
+	 * @param	Out				alternate output device
+	 * @param	Struct			UStruct to export properties
+	 * @param	TextIndent		Current text indentation
 	 */
-	void ExportProperties(UStruct* Struct, int32 TextIndent, bool bAccessSpecifiers = true, FUHTStringBuilder* Output = NULL);
+	static void ExportProperties(FOutputDevice& Out, UStruct* Struct, int32 TextIndent);
 
 	/** Return the name of the singleton function that returns the UObject for Item */
 	FString GetSingletonName(UField* Item, bool bRequiresValidObject=true);
-	/** Return the name of the singleton function that returns the UObject for Item */
-	FString GetSingletonName(FClass* Item, bool bRequiresValidObject=true);
 	
 	/** 
 	 * Returns the name (overridden if marked up) with TEXT("") or "" wrappers for use in a string literal.
 	 */
-	FString GetOverriddenNameForLiteral(const UField* Item);
+	static FString GetOverriddenNameForLiteral(const UField* Item);
 
 	/** 
 	 * Export functions used to find and call C++ or script implementation of a script function in the interface 
 	 */
-	void ExportInterfaceCallFunctions(const TArray<UFunction*>& CallbackFunctions, UClass* Class, FUHTStringBuilder& HeaderOutput);
+	void ExportInterfaceCallFunctions(FOutputDevice& OutCpp, FUHTStringBuilder& Out, const TArray<UFunction*>& CallbackFunctions, const TCHAR* ClassName);
 
 	/** 
 	 * Export UInterface boilerplate.
@@ -220,133 +180,122 @@ private:
 	 * @param Class Interface to export.
 	 * @param FriendText Friend text for this boilerplate.
 	 */
-	void ExportUInterfaceBoilerplate(FUHTStringBuilder& UInterfaceBoilerplate, FClass* Class, const FString& FriendText);
-
-	/**
-	 * Appends the header definition for an inheritance hierarchy of classes to the header.
-	 * Wrapper for ExportClassHeaderRecursive().
-	 *
-	 * @param	AllClasses			The classes being processed.
-	 * @param	SourceFile			The source file to export header for.
-	 */
-	void ExportSourceFileHeader(FClasses& AllClasses, FUnrealSourceFile* SourceFile);
+	static void ExportUInterfaceBoilerplate(FUHTStringBuilder& UInterfaceBoilerplate, FClass* Class, const FString& FriendText);
 
 	/**
 	 * After all of the dependency checking, and setup for isolating the generated code, actually export the class
 	 *
-	 * @param SourceFile Source file to export.
-	 */
-	void ExportClassesFromSourceFileInner(FUnrealSourceFile& SourceFile);
+	 * @param	OutEnums			Output device for enums declarations.
+	 * @param	OutputGetter		The function to call to get the output.
+	 * @param	SourceFile			Source file to export.
+	 */
+	void ExportClassFromSourceFileInner(
+		FOutputDevice&           OutGeneratedHeaderText,
+		FOutputDevice&           OutputGetter,
+		FOutputDevice&           OutDeclarations,
+		FClass*                  Class,
+		const FUnrealSourceFile& SourceFile
+	);
 
 	/**
 	 * After all of the dependency checking, but before actually exporting the class, set up the generated code
 	 *
-	 * @param SourceFile Source file to export.
+	 * @param	SourceFile			Source file to export.
 	 */
-	void ExportClassesFromSourceFileWrapper(FUnrealSourceFile& SourceFile);
+	bool WriteHeader(const TCHAR* Path, const FString& InBodyText, const TSet<FString>& InFwdDecl);
 
 	/**
 	 * Appends the header definition for an inheritance hierarchy of classes to the header.
 	 *
-	 * @param	AllClasses				The classes being processed.
+	 * @param	OutputGetter			The function to call to get the output.
 	 * @param	SourceFile				The source file for which to export header.
 	 * @param	VisitedSet				The set of source files visited so far. Must be empty before the first call.
 	 * @param	bCheckDependenciesOnly	Whether we should just keep checking for dependency errors, without exporting anything.
 	 */
-	void ExportSourceFileHeaderRecursive(FClasses& AllClasses, FUnrealSourceFile* SourceFile, TSet<const FUnrealSourceFile*>& VisitedSet, bool bCheckDependenciesOnly);
+	void ExportSourceFileHeader(
+		FOutputDevice&                     OutGeneratedHeaderText,
+		TFunctionRef<FUHTStringBuilder&()> OutputGetter,
+		FOutputDevice&                     OutDeclarations,
+		const FUnrealSourceFile&           SourceFile,
+		const TArray<UEnum*>&              Enums,
+		const TArray<UScriptStruct*>&      Structs,
+		const TArray<UDelegateFunction*>&  DelegateFunctions
+	);
 
 	/**
 	 * Returns a string in the format CLASS_Something|CLASS_Something which represents all class flags that are set for the specified
 	 * class which need to be exported as part of the DECLARE_CLASS macro
 	 */
-	FString GetClassFlagExportText( UClass* Class );
+	static FString GetClassFlagExportText( UClass* Class );
 
 	/**
-	 * Exports the header text for the list of enums specified
+	 * Exports the header text for the enum specified
 	 * 
-	 * @param	Enums	the enums to export
+	 * @param	Out		the output device for the mirror struct
+	 * @param	Enums	the enum to export
 	 */
-	void ExportEnums( const TArray<UEnum*>& Enums );
+	static void ExportEnum(FOutputDevice& Out, UEnum* Enum);
 
 	/**
 	 * Exports the inl text for enums declared in non-UClass headers.
 	 * 
-	 * @param	Enums	the enums to export
+	 * @param	OutputGetter	The function to call to get the output.
+	 * @param	Enum			the enum to export
 	 */
-	void ExportGeneratedEnumsInitCode(FUnrealSourceFile& SourceFile, const TArray<UEnum*>& Enums);
+	void ExportGeneratedEnumInitCode(FOutputDevice& Out, FOutputDevice& OutDeclarations, const FUnrealSourceFile& SourceFile, UEnum* Enum);
 
 	/**
-	 * Exports the macro declarations for GENERATED_BODY() for each Foo in the list of structs specified
+	 * Exports the macro declarations for GENERATED_BODY() for each Foo in the struct specified
 	 * 
-	 * @param	Structs		the structs to export
+	 * @param	Out				output device
+	 * @param	Struct			The struct to export
 	 */
-	void ExportGeneratedStructBodyMacros(FUnrealSourceFile& SourceFile, const TArray<UScriptStruct*>& NativeStructs);
+	void ExportGeneratedStructBodyMacros(FOutputDevice& OutGeneratedHeaderText, FOutputDevice& Out, FOutputDevice& OutDeclarations, const FUnrealSourceFile& SourceFile, UScriptStruct* Struct);
 
 	/**
-	 * Exports a local mirror of the specified struct; used to get offsets for noexport structs
+	 * Exports a local mirror of the specified struct; used to get offsets
 	 * 
-	 * @param	Structs		the structs to export
+	 * @param	Out			the output device for the mirror struct
+	 * @param	Structs		the struct to export
 	 * @param	TextIndent	the current indentation of the header exporter
-	 * @param	HeaderOutput	the output device for the mirror struct
 	 */
-	void ExportMirrorsForNoexportStructs(const TArray<UScriptStruct*>& NativeStructs, int32 TextIndent, FUHTStringBuilder& HeaderOutput);
+	static void ExportMirrorsForNoexportStruct(FOutputDevice& Out, UScriptStruct* Struct, int32 TextIndent);
 
-	/**
+	/**heade
 	 * Exports the parameter struct declarations for the list of functions specified
 	 * 
 	 * @param	Function	the function that (may) have parameters which need to be exported
 	 * @return	true		if the structure generated is not completely empty
 	 */
-	bool WillExportEventParms( UFunction* Function );
+	static bool WillExportEventParms( UFunction* Function );
 
 	/**
-	 * Exports C++ type definitions for delegates
+	 * Exports C++ type declarations for delegates
 	 *
+	 * @param	Out					output device
 	 * @param	SourceFile			Source file of the delegate.
 	 * @param	DelegateFunctions	the functions that have parameters which need to be exported
-	 * @param	bWrapperImplementationsOnly	 True if only function implementations for delegate wrappers should be exported
 	 */
-	void ExportDelegateDefinitions(FUnrealSourceFile& SourceFile, const TArray<UDelegateFunction*>& DelegateFunctions, const bool bWrapperImplementationsOnly);
+	void ExportDelegateDeclaration(FOutputDevice& Out, FOutputDevice& OutDeclarations, const FUnrealSourceFile& SourceFile, UFunction* Function);
 
 	/**
-	 * Exports the parameter struct declarations for the list of functions specified
-	 * 
-	 * @param	Scope				Current scope.
-	 * @param	CallbackFunctions	the functions that have parameters which need to be exported
-	 * @param	Output				ALternate output device
-	 * @param	Indent				number of spaces to put before each line
-	 * @param	bOutputConstructor	If true, output a constructor for the parm struct
+	 * Exports C++ type definitions for delegates
+	 *
+	 * @param	Out					output device
+	 * @param	SourceFile			Source file of the delegate.
+	 * @param	DelegateFunctions	the functions that have parameters which need to be exported
 	 */
-	void ExportEventParms(FScope& Scope, const TArray<UFunction*>& CallbackFunctions, FUHTStringBuilder& Output, int32 Indent = 0, bool bOutputConstructor = true);
+	void ExportDelegateDefinition(FOutputDevice& Out, const FUnrealSourceFile& SourceFile, UFunction* Function);
 
 	/**
 	 * Exports the parameter struct declarations for the given function.
 	 *
+	 * @param	Out					output device
 	 * @param	Function			the function that have parameters which need to be exported
-	 * @param	HeaderOutput		output device
 	 * @param	Indent				number of spaces to put before each line
 	 * @param	bOutputConstructor	If true, output a constructor for the param struct
 	 */
-	void ExportEventParm(UFunction* Function, FUHTStringBuilder& HeaderOutput, int32 Indent = 0, bool bOutputConstructor = true);
-
-	/**
-	 * Generate a .proto message declaration for any functions marked as requiring one
-	 * 
-	 * @param InCallbackFunctions array of functions for consideration to generate .proto definitions
-	 * @param Indent starting indentation level
-	 * @param Output optional output redirect
-	 */
-	void ExportProtoMessage(const TArray<UFunction*>& InCallbackFunctions, int32 Indent = 0, FUHTStringBuilder* Output = NULL);
-
-	/**
-	 * Generate a .java message declaration for any functions marked as requiring one
-	 * 
-	 * @param InCallbackFunctions array of functions for consideration to generate .proto definitions
-	 * @param ClassData class data
-	 * @param Indent starting indentation level
-	 * @param Output optional output redirect
-	 */
-	void ExportMCPMessage(const TArray<UFunction*>& InCallbackFunctions, FClassMetaData* ClassData, int32 Indent = 0, FUHTStringBuilder* Output = NULL);
+	static void ExportEventParm(FUHTStringBuilder& Out, TSet<FString>& PropertyFwd, UFunction* Function, int32 Indent, bool bOutputConstructor, EExportingState ExportingState);
 
 	/** 
 	* Exports the temp header files into the .h files, then deletes the temp files.
@@ -358,19 +307,7 @@ private:
 	/**
 	 * Exports the generated cpp file for all functions/events/delegates in package.
 	 */
-	void ExportGeneratedCPP();
-
-	/**
-	 * Exports protobuffer definitions from boilerplate that was generated for a package.
-	 * They are exported to a file using the name <PackageName>.generated.proto
-	 */
-	void ExportGeneratedProto();
-
-	/**
-	 * Exports mcp definitions from boilerplate that was generated for a package.
-	 * They are exported to a file using the name <PackageName>.generated.java
-	 */
-	void ExportGeneratedMCP();
+	void ExportGeneratedCPP(const TCHAR* Declarations, const TArray<FUnrealSourceFile*>& Exported);
 
 	/**
 	 * Get the intrinsic null value for this property
@@ -381,18 +318,26 @@ private:
 	 *
 	 * @return	the intrinsic null value for the property (0 for ints, TEXT("") for strings, etc.)
 	 */
-	FString GetNullParameterValue( UProperty* Prop, bool bMacroContext, bool bInitializer = false );
+	static FString GetNullParameterValue( UProperty* Prop, bool bMacroContext, bool bInitializer = false );
 
 	/**
 	 * Exports a native function prototype
 	 * 
+	 * @param	Out					Where to write the exported function prototype to.
 	 * @param	FunctionData		data representing the function to export
-	 * @param	HeaderOutput		Where to write the exported function prototype to.
 	 * @param	FunctionType		Whether to export this function prototype as an event stub, an interface or a native function stub.
 	 * @param	FunctionHeaderStyle	Whether we're outputting a declaration or definition.
 	 * @param	ExtraParam			Optional extra parameter that will be added to the declaration as the first argument
 	 */
-	void ExportNativeFunctionHeader(const FFuncInfo& FunctionData, FUHTStringBuilder& HeaderOutput, EExportFunctionType::Type FunctionType, EExportFunctionHeaderStyle::Type FunctionHeaderStyle, const TCHAR* ExtraParam = NULL);
+	static void ExportNativeFunctionHeader(
+		FOutputDevice&                   Out,
+		TSet<FString>&                   OutFwdDecls,
+		const FFuncInfo&                 FunctionData,
+		EExportFunctionType::Type        FunctionType,
+		EExportFunctionHeaderStyle::Type FunctionHeaderStyle,
+		const TCHAR*                     ExtraParam,
+		const TCHAR*                     APIString
+	);
 
 	/**
 	* Runs checks whether necessary RPC functions exist for function described by FunctionData.
@@ -412,7 +357,7 @@ private:
 	 * @param Class			class
 	 * @param ClassData		class data
 	 */
-	void ExportNativeFunctions(FUnrealSourceFile& SourceFile, UClass* Class, FClassMetaData* ClassData);
+	void ExportNativeFunctions(FOutputDevice& OutGeneratedHeaderText, FOutputDevice& OutMacroCalls, FOutputDevice& OutNoPureDeclsMacroCalls, const FUnrealSourceFile& SourceFile, UClass* Class, FClassMetaData* ClassData);
 
 	/**
 	 * Export the actual internals to a standard thunk function
@@ -427,32 +372,34 @@ private:
 	void ExportFunctionThunk(FUHTStringBuilder& RPCWrappers, UFunction* Function, const FFuncInfo& FunctionData, const TArray<UProperty*>& Parameters, UProperty* Return, FUHTStringBuilder& DeprecationWarningOutputDevice);
 
 	/** Exports the native function registration code for the given class. */
-	void ExportNatives(FClass* Class);
+	static void ExportNatives(FOutputDevice& Out, FClass* Class);
 
 	/**
 	 * Exports generated singleton functions for UObjects that used to be stored in .u files.
 	 * 
+	 * @param	Out				The destination to write to.
 	 * @param	SourceFile		The source file being processed.
 	 * @param	Class			Class to export
 	 * @param	OutFriendText	(Output parameter) Friend text
 	 */
-	void ExportNativeGeneratedInitCode(FUnrealSourceFile& SourceFile, FClass* Class, FUHTStringBuilder& OutFriendText);
+	void ExportNativeGeneratedInitCode(FOutputDevice& Out, FOutputDevice& OutDeclarations, const FUnrealSourceFile& SourceFile, FClass* Class, FUHTStringBuilder& OutFriendText);
 
 	/**
 	 * Export given function.
 	 *
-	 * @param Function Given function.
-	 * @param Scope Scope for this function.
-	 * @param bIsNoExport Is in NoExport class.
+	 * @param	Out				The destination to write to.
+	 * @param	Function		Given function.
+	 * @param	bIsNoExport		Is in NoExport class.
 	 */
-	void ExportFunction(FUnrealSourceFile& SourceFile, UFunction* Function, FScope* Scope, bool bIsNoExport);
+	void ExportFunction(FOutputDevice& Out, const FUnrealSourceFile& SourceFile, UFunction* Function, bool bIsNoExport);
 
 	/**
 	 * Exports a generated singleton function to setup the package for compiled-in classes.
 	 * 
+	 * @param	Out			The destination to write to.
 	 * @param	Package		Package to export code for.
 	**/
-	void ExportGeneratedPackageInitCode(const UPackage* Package);
+	void ExportGeneratedPackageInitCode(FOutputDevice& Out, FUHTStringBuilder& OutDeclarations, const UPackage* Package);
 
 	/**
 	 * Function to output the C++ code necessary to set up the given array of properties
@@ -462,7 +409,7 @@ private:
 	 * @param	Properties		Array of properties to export
 	 * @param	Spaces			String of spaces to use as an indent
 	 */
-	void OutputProperties(FString& Meta, FOutputDevice& OutputDevice, const FString& OuterString, TArray<UProperty*>& Properties, const TCHAR* Spaces);
+	void OutputProperties(FString& Meta, FOutputDevice& OutputDevice, const FString& OuterString, const TArray<UProperty*>& Properties, const TCHAR* Spaces);
 
 	/**
 	 * Function to output the C++ code necessary to set up a property
@@ -491,10 +438,18 @@ private:
 	 * Exports the proxy definitions for the list of enums specified
 	 * 
 	 * @param SourceFile	current source file
-	 * @param Class			current class
-	 * @param ClassData		current class data
 	 */
-	TArray<UFunction*> ExportCallbackFunctions(FUnrealSourceFile& SourceFile, UClass* Class, FClassMetaData* ClassData);
+	static void ExportCallbackFunctions(
+		FOutputDevice&            OutGeneratedHeaderText,
+		FOutputDevice&            Out,
+		TMap<FName, FName>&       ReferencedNames,
+		TSet<FString>&            OutFwdDecls,
+		const TArray<UFunction*>& CallbackFunctions,
+		const TCHAR*              CallbackWrappersMacroName,
+		EExportCallbackType       ExportCallbackType,
+		const TCHAR*              API,
+		const TCHAR*              APIString
+	);
 
 	/**
 	 * Determines if the property has alternate export text associated with it and if so replaces the text in PropertyText with the
@@ -504,7 +459,7 @@ private:
 	 * @param	Prop			the property that is being exported
 	 * @param	PropertyText	the string containing the text exported from ExportCppDeclaration
 	 */
-	void ApplyAlternatePropertyExportText(UProperty* Prop, FUHTStringBuilder& PropertyText);
+	static void ApplyAlternatePropertyExportText(UProperty* Prop, FUHTStringBuilder& PropertyText, EExportingState ExportingState);
 
 	/**
 	* Create a temp header file name from the header name
@@ -514,7 +469,7 @@ private:
 	*
 	* @return	The generated string
 	*/
-	FString GenerateTempHeaderName( FString CurrentFilename, bool bReverseOperation = false );
+	static FString GenerateTempHeaderName( FString CurrentFilename, bool bReverseOperation = false );
 
 	/**
 	 * Saves a generated header if it has changed. 
@@ -533,15 +488,13 @@ private:
 	/**
 	 * Exports macros that manages UObject constructors.
 	 * 
-	 * @param ConstructorsMacroPrefix Prefix for constructors macro.
-	 * @param Class Class for which to export macros.
-	 */
-	void ExportConstructorsMacros(const FString& ConstructorsMacroPrefix, FClass* Class);
-
-	/**
-	 * Gets list of public headers for the given package in the form of string of includes.
+	 * @param	VTableOut								The destination to write vtable helpers to.
+	 * @param	StandardUObjectConstructorsMacroCall	The destination to write standard constructor macros to.
+	 * @param	EnhancedUObjectConstructorsMacroCall	The destination to write enhanced constructor macros to.
+	 * @param	ConstructorsMacroPrefix					Prefix for constructors macro.
+	 * @param	Class									Class for which to export macros.
 	 */
-	FString GetListOfPublicHeaderGroupIncludesString(const UPackage* Package);
+	static void ExportConstructorsMacros(FOutputDevice& OutGeneratedHeaderText, FOutputDevice& VTableOut, FOutputDevice& StandardUObjectConstructorsMacroCall, FOutputDevice& EnhancedUObjectConstructorsMacroCall, const FString& ConstructorsMacroPrefix, FClass* Class, const TCHAR* APIArg);
 
 public:
 
@@ -558,11 +511,8 @@ public:
 		const UPackage* InPackage,
 		const TArray<FUnrealSourceFile*>& SourceFiles,
 		FClasses& AllClasses,
-		bool InAllowSaveExportedHeaders
-#if WITH_HOT_RELOAD_CTORS
-		, bool bExportVTableConstructors
-#endif // WITH_HOT_RELOAD_CTORS
-		, FUHTMakefile& UHTMakefile
+		bool InAllowSaveExportedHeaders,
+		FUHTMakefile& UHTMakefile
 	);
 
 	/**
@@ -588,5 +538,5 @@ public:
 	 * @param FunctionNamePosition Position of name of function in SourceFile.
 	 * @return true if function misses "virtual" specifier, false otherwise.
 	 */
-	bool IsMissingVirtualSpecifier(const FString& SourceFile, int32 FunctionNamePosition) const;
+	static bool IsMissingVirtualSpecifier(const FString& SourceFile, int32 FunctionNamePosition);
 };
diff --git a/Engine/Source/Programs/UnrealHeaderTool/Private/ParserHelper.h b/Engine/Source/Programs/UnrealHeaderTool/Private/ParserHelper.h
--- a/Engine/Source/Programs/UnrealHeaderTool/Private/ParserHelper.h
+++ b/Engine/Source/Programs/UnrealHeaderTool/Private/ParserHelper.h
@@ -37,8 +37,8 @@ enum EFunctionExportFlags
 	FUNCEXPORT_Inline			=0x00000010,	// export as an inline static C++ function
 	FUNCEXPORT_CppStatic		=0x00000020,	// Export as a real C++ static function, causing thunks to call via ClassName::FuncName instead of this->FuncName
 	FUNCEXPORT_CustomThunk		=0x00000040,	// Export no thunk function; the user will manually define a custom one
-	FUNCEXPORT_NeedsProto		=0x00000080,	// Export a .proto file from the parameters of this function
-	FUNCEXPORT_NeedsMCP			=0x00000100,	// Export a .java file from the parameters of this function
+	//							=0x00000080,
+	//							=0x00000100,
 };
 
 enum EPropertyHeaderExportFlags
@@ -1043,6 +1043,8 @@ struct FFuncInfo
 	FString		CppValidationImplName;
 	/** Name for callback-style names **/
 	FString		UnMarshallAndCallName;
+	/** Endpoint name */
+	FString		EndpointName;
 	/** Identifier for an RPC call to a platform service */
 	int16		RPCId;
 	/** Identifier for an RPC call expecting a response */
@@ -1527,9 +1529,7 @@ public:
 		, bConstructorDeclared(false)
 		, bDefaultConstructorDeclared(false)
 		, bObjectInitializerConstructorDeclared(false)
-#if WITH_HOT_RELOAD_CTORS
 		, bCustomVTableHelperConstructorDeclared(false)
-#endif // WITH_HOT_RELOAD_CTORS
 		, GeneratedBodyMacroAccessSpecifier(ACCESS_NotAnAccessSpecifier)
 	{
 	}
@@ -1752,10 +1752,8 @@ public:
 	// Is ObjectInitializer constructor (i.e. a constructor with only one parameter of type FObjectInitializer) declared?
 	bool bObjectInitializerConstructorDeclared;
 
-#if WITH_HOT_RELOAD_CTORS
 	// Is custom VTable helper constructor declared?
 	bool bCustomVTableHelperConstructorDeclared;
-#endif // WITH_HOT_RELOAD_CTORS
 
 	// GENERATED_BODY access specifier to preserve.
 	EAccessSpecifier GeneratedBodyMacroAccessSpecifier;
diff --git a/Engine/Source/Programs/UnrealHeaderTool/Private/UHTMakefile/UHTMakefile.cpp b/Engine/Source/Programs/UnrealHeaderTool/Private/UHTMakefile/UHTMakefile.cpp
--- a/Engine/Source/Programs/UnrealHeaderTool/Private/UHTMakefile/UHTMakefile.cpp
+++ b/Engine/Source/Programs/UnrealHeaderTool/Private/UHTMakefile/UHTMakefile.cpp
@@ -1174,7 +1174,6 @@ void FUHTMakefile::ResolvePublicClassSetEntry(int32 Index)
 {
 	UClass* Class = GetClassByIndex(PublicClassSetEntryProxies[Index]);
 	PublicClassSetEntries.Add(Class);
-	GPublicClassSet.Add(Class);
 }
 
 void FUHTMakefile::CreateTypeDefinitionInfoMapEntry(int32 Index)
diff --git a/Engine/Source/Programs/UnrealHeaderTool/Private/UHTMakefile/UHTMakefile.h b/Engine/Source/Programs/UnrealHeaderTool/Private/UHTMakefile/UHTMakefile.h
--- a/Engine/Source/Programs/UnrealHeaderTool/Private/UHTMakefile/UHTMakefile.h
+++ b/Engine/Source/Programs/UnrealHeaderTool/Private/UHTMakefile/UHTMakefile.h
@@ -1133,7 +1133,7 @@ public:
 	int32 GetTokenIndex(const FToken* Token) const;
 	const FToken* GetTokenByIndex(int32 Index) const;
 
-	void AddGeneratedCodeCRC(FUnrealSourceFile* SourceFile, UField* Field, uint32 CRC)
+	void AddGeneratedCodeCRC(const FUnrealSourceFile* SourceFile, UField* Field, uint32 CRC)
 	{
 		int32 Index = GeneratedCodeCRCs.Emplace(Field, CRC);
 		GetHeaderDescriptor(SourceFile).AddEntry(AddObject(ESerializedObjectType::EGeneratedCodeCRC, Index));
@@ -1325,7 +1325,7 @@ public:
 		return CurrentModule.GetHeaderDescriptor(UnrealSourceFileIndex);
 	}
 
-	FUHTMakefileHeaderDescriptor& GetHeaderDescriptor(FUnrealSourceFile* UnrealSourceFile)
+	FUHTMakefileHeaderDescriptor& GetHeaderDescriptor(const FUnrealSourceFile* UnrealSourceFile)
 	{
 		int32 UnrealSourceFileIndex = GetUnrealSourceFileIndex(UnrealSourceFile);
 		return GetHeaderDescriptorBySourceFileIndex(UnrealSourceFileIndex);
diff --git a/Engine/Source/Programs/UnrealHeaderTool/Private/UnrealSourceFile.cpp b/Engine/Source/Programs/UnrealHeaderTool/Private/UnrealSourceFile.cpp
--- a/Engine/Source/Programs/UnrealHeaderTool/Private/UnrealSourceFile.cpp
+++ b/Engine/Source/Programs/UnrealHeaderTool/Private/UnrealSourceFile.cpp
@@ -99,10 +99,9 @@ const FString& FUnrealSourceFile::GetContent() const
 	return Content;
 }
 
-EGeneratedCodeVersion FUnrealSourceFile::GetGeneratedCodeVersionForStruct(UStruct* Struct)
+EGeneratedCodeVersion FUnrealSourceFile::GetGeneratedCodeVersionForStruct(UStruct* Struct) const
 {
-	auto Version = GeneratedCodeVersions.Find(Struct);
-	if (Version)
+	if (const EGeneratedCodeVersion* Version = GeneratedCodeVersions.Find(Struct))
 	{
 		return *Version;
 	}
diff --git a/Engine/Source/Programs/UnrealHeaderTool/Private/UnrealSourceFile.h b/Engine/Source/Programs/UnrealHeaderTool/Private/UnrealSourceFile.h
--- a/Engine/Source/Programs/UnrealHeaderTool/Private/UnrealSourceFile.h
+++ b/Engine/Source/Programs/UnrealHeaderTool/Private/UnrealSourceFile.h
@@ -224,7 +224,7 @@ public:
 	/**
 	 * Gets generated code version for given UStruct.
 	 */
-	EGeneratedCodeVersion GetGeneratedCodeVersionForStruct(UStruct* Struct);
+	EGeneratedCodeVersion GetGeneratedCodeVersionForStruct(UStruct* Struct) const;
 
 	/**
 	 * Gets generated code versions.
diff --git a/Engine/Source/Programs/UnrealHeaderTool/Resources/UHTDebugging/StructOnlyHeader.h b/Engine/Source/Programs/UnrealHeaderTool/Resources/UHTDebugging/StructOnlyHeader.h
--- a/Engine/Source/Programs/UnrealHeaderTool/Resources/UHTDebugging/StructOnlyHeader.h
+++ b/Engine/Source/Programs/UnrealHeaderTool/Resources/UHTDebugging/StructOnlyHeader.h
@@ -7,5 +7,9 @@
 USTRUCT()
 struct FSomeStruct
 {
+#pragma region X
+
 	GENERATED_BODY()
+
+#pragma endregion X
 };
diff --git a/Engine/Source/Programs/UnrealVS/Properties/AssemblyInfo.cs b/Engine/Source/Programs/UnrealVS/Properties/AssemblyInfo.cs
--- a/Engine/Source/Programs/UnrealVS/Properties/AssemblyInfo.cs
+++ b/Engine/Source/Programs/UnrealVS/Properties/AssemblyInfo.cs
@@ -31,8 +31,8 @@ using System.Runtime.InteropServices;
 // You can specify all the values or you can default the Revision and Build Numbers 
 // by using the '*' as shown below:
 
-[assembly: AssemblyVersion("1.43.0.0")]
-[assembly: AssemblyFileVersion("1.43.0.0")]
+[assembly: AssemblyVersion("1.44.0.0")]
+[assembly: AssemblyFileVersion("1.44.0.0")]
 
 
 
diff --git a/Engine/Source/Programs/UnrealVS/UnrealVS.csproj b/Engine/Source/Programs/UnrealVS/UnrealVS.csproj
--- a/Engine/Source/Programs/UnrealVS/UnrealVS.csproj
+++ b/Engine/Source/Programs/UnrealVS/UnrealVS.csproj
@@ -22,8 +22,8 @@
     </SccAuxPath>
     <SccProvider>
     </SccProvider>
-    <MinimumVisualStudioVersion>15.0</MinimumVisualStudioVersion>
-    <OldToolsVersion>14.0</OldToolsVersion>
+    <MinimumVisualStudioVersion>14.0</MinimumVisualStudioVersion>
+    <OldToolsVersion>12.0</OldToolsVersion>
     <TargetFrameworkProfile />
     <FileUpgradeFlags>
     </FileUpgradeFlags>
@@ -87,6 +87,7 @@
     <Reference Include="Microsoft.VisualStudio.Shell.12.0, Version=12.0.0.0, Culture=neutral, PublicKeyToken=b03f5f7f11d50a3a, processorArchitecture=MSIL" Condition="'$(Configuration)' == 'VS12' Or '$(Configuration)' == 'VS12d'" />
     <Reference Include="Microsoft.VisualStudio.Shell.14.0, Version=14.0.0.0, Culture=neutral, PublicKeyToken=b03f5f7f11d50a3a, processorArchitecture=MSIL" Condition="'$(Configuration)' == 'VS14' Or '$(Configuration)' == 'VS14d'" />
     <Reference Include="Microsoft.VisualStudio.Shell.15.0, Version=15.0.0.0, Culture=neutral, PublicKeyToken=b03f5f7f11d50a3a, processorArchitecture=MSIL" Condition="'$(Configuration)' == 'VS15' Or '$(Configuration)' == 'VS15d'" />
+    <Reference Include="Microsoft.VisualStudio.Shell.Framework, Version=15.0.0.0, Culture=neutral, PublicKeyToken=b03f5f7f11d50a3a, processorArchitecture=MSIL" Condition="'$(Configuration)' == 'VS15' Or '$(Configuration)' == 'VS15d'" />
     <Reference Include="Microsoft.VisualStudio.Shell.Immutable.10.0" />
     <Reference Include="Microsoft.VisualStudio.Shell.Interop, Version=7.1.40304.0, Culture=neutral, PublicKeyToken=b03f5f7f11d50a3a" />
     <Reference Include="Microsoft.VisualStudio.Shell.Interop.10.0, Version=10.0.0.0, Culture=neutral, PublicKeyToken=b03f5f7f11d50a3a, processorArchitecture=MSIL">
diff --git a/Engine/Source/Programs/UnrealVS/UnrealVSPackage.cs b/Engine/Source/Programs/UnrealVS/UnrealVSPackage.cs
--- a/Engine/Source/Programs/UnrealVS/UnrealVSPackage.cs
+++ b/Engine/Source/Programs/UnrealVS/UnrealVSPackage.cs
@@ -56,7 +56,7 @@ namespace UnrealVS
 	{
 		/** Constants */
 
-		private const string VersionString = "v1.43";
+		private const string VersionString = "v1.44";
 		private const string UnrealSolutionFileNamePrefix = "UE4";
 		private const string ExtensionName = "UnrealVS";
 		private const string CommandLineOptionKey = ExtensionName + "CommandLineMRU";
diff --git a/Engine/Source/Programs/UnrealVS/UnrealVS_2017.csproj b/Engine/Source/Programs/UnrealVS/UnrealVS_2017.csproj
new file mode 100644
--- /dev/null
+++ b/Engine/Source/Programs/UnrealVS/UnrealVS_2017.csproj
@@ -0,0 +1,365 @@
+﻿<?xml version="1.0" encoding="utf-8"?>
+<Project DefaultTargets="Build" xmlns="http://schemas.microsoft.com/developer/msbuild/2003" ToolsVersion="14.0">
+  <Import Project="$(MSBuildExtensionsPath)\$(MSBuildToolsVersion)\Microsoft.Common.props" Condition="Exists('$(MSBuildExtensionsPath)\$(MSBuildToolsVersion)\Microsoft.Common.props')" />
+  <PropertyGroup>
+    <Configuration Condition=" '$(Configuration)' == '' ">Debug</Configuration>
+    <Platform Condition=" '$(Platform)' == '' ">AnyCPU</Platform>
+    <SchemaVersion>2.0</SchemaVersion>
+    <ProjectGuid>{C0EE3E24-0255-4D0E-9369-82247AD5E2DD}</ProjectGuid>
+    <ProjectTypeGuids>{82b43b9b-a64c-4715-b499-d71e9ca2bd60};{60dc8134-eba5-43b8-bcc9-bb4bc16c2548};{FAE04EC0-301F-11D3-BF4B-00C04F79EFBC}</ProjectTypeGuids>
+    <OutputType>Library</OutputType>
+    <AppDesignerFolder>Properties</AppDesignerFolder>
+    <RootNamespace>UnrealVS</RootNamespace>
+    <AssemblyName>UnrealVS</AssemblyName>
+    <SignAssembly>True</SignAssembly>
+    <AssemblyOriginatorKeyFile>Key.snk</AssemblyOriginatorKeyFile>
+    <TargetFrameworkVersion>v4.6</TargetFrameworkVersion>
+    <SccProjectName>
+    </SccProjectName>
+    <SccLocalPath>
+    </SccLocalPath>
+    <SccAuxPath>
+    </SccAuxPath>
+    <SccProvider>
+    </SccProvider>
+    <MinimumVisualStudioVersion>15.0</MinimumVisualStudioVersion>
+    <OldToolsVersion>14.0</OldToolsVersion>
+    <TargetFrameworkProfile />
+    <FileUpgradeFlags>
+    </FileUpgradeFlags>
+    <UpgradeBackupLocation>
+    </UpgradeBackupLocation>
+    <IsWebBootstrapper>false</IsWebBootstrapper>
+    <PublishUrl>publish\</PublishUrl>
+    <Install>true</Install>
+    <InstallFrom>Disk</InstallFrom>
+    <UpdateEnabled>false</UpdateEnabled>
+    <UpdateMode>Foreground</UpdateMode>
+    <UpdateInterval>7</UpdateInterval>
+    <UpdateIntervalUnits>Days</UpdateIntervalUnits>
+    <UpdatePeriodically>false</UpdatePeriodically>
+    <UpdateRequired>false</UpdateRequired>
+    <MapFileExtensions>true</MapFileExtensions>
+    <ApplicationRevision>0</ApplicationRevision>
+    <ApplicationVersion>1.0.0.%2a</ApplicationVersion>
+    <UseApplicationTrust>false</UseApplicationTrust>
+    <BootstrapperEnabled>true</BootstrapperEnabled>
+  </PropertyGroup>
+  <PropertyGroup Condition=" '$(Configuration)|$(Platform)' == 'VS12d|AnyCPU' ">
+    <DebugSymbols>true</DebugSymbols>
+    <DebugType>full</DebugType>
+    <Optimize>false</Optimize>
+    <OutputPath>bin\VS12d\</OutputPath>
+    <DefineConstants>TRACE;DEBUG;VS12</DefineConstants>
+    <ErrorReport>prompt</ErrorReport>
+    <WarningLevel>4</WarningLevel>
+    <Prefer32Bit>false</Prefer32Bit>
+    <CopyVsixExtensionFiles>False</CopyVsixExtensionFiles>
+    <DeployExtension>True</DeployExtension>
+    <CreateVsixContainer>True</CreateVsixContainer>
+  </PropertyGroup>
+  <PropertyGroup Condition=" '$(Configuration)|$(Platform)' == 'VS12|AnyCPU' ">
+    <DebugType>pdbonly</DebugType>
+    <Optimize>true</Optimize>
+    <OutputPath>..\..\..\Extras\UnrealVS\VS2013\</OutputPath>
+    <DefineConstants>TRACE;CODE_ANALYSIS;VS12</DefineConstants>
+    <ErrorReport>prompt</ErrorReport>
+    <WarningLevel>4</WarningLevel>
+    <RunCodeAnalysis>true</RunCodeAnalysis>
+    <Prefer32Bit>false</Prefer32Bit>
+  </PropertyGroup>
+  <ItemGroup>
+    <VSCTCompile Include="UnrealVS.vsct">
+      <ResourceName>Menus.ctmenu</ResourceName>
+      <SubType>Designer</SubType>
+    </VSCTCompile>
+  </ItemGroup>
+  <ItemGroup>
+    <Reference Include="envdte, Version=8.0.0.0, Culture=neutral, PublicKeyToken=b03f5f7f11d50a3a">
+      <EmbedInteropTypes>False</EmbedInteropTypes>
+    </Reference>
+    <Reference Include="envdte80, Version=8.0.0.0, Culture=neutral, PublicKeyToken=b03f5f7f11d50a3a">
+      <EmbedInteropTypes>True</EmbedInteropTypes>
+    </Reference>
+    <Reference Include="Microsoft.CSharp" />
+    <Reference Include="Microsoft.VisualStudio.OLE.Interop, Version=7.1.40304.0, Culture=neutral, PublicKeyToken=b03f5f7f11d50a3a" />
+    <Reference Include="Microsoft.VisualStudio.Shell.11.0, Version=11.0.0.0, Culture=neutral, PublicKeyToken=b03f5f7f11d50a3a, processorArchitecture=MSIL" Condition="'$(Configuration)' == 'VS11' Or '$(Configuration)' == 'VS11d'" />
+    <Reference Include="Microsoft.VisualStudio.Shell.12.0, Version=12.0.0.0, Culture=neutral, PublicKeyToken=b03f5f7f11d50a3a, processorArchitecture=MSIL" Condition="'$(Configuration)' == 'VS12' Or '$(Configuration)' == 'VS12d'" />
+    <Reference Include="Microsoft.VisualStudio.Shell.14.0, Version=14.0.0.0, Culture=neutral, PublicKeyToken=b03f5f7f11d50a3a, processorArchitecture=MSIL" Condition="'$(Configuration)' == 'VS14' Or '$(Configuration)' == 'VS14d'" />
+    <Reference Include="Microsoft.VisualStudio.Shell.15.0, Version=15.0.0.0, Culture=neutral, PublicKeyToken=b03f5f7f11d50a3a, processorArchitecture=MSIL" Condition="'$(Configuration)' == 'VS15' Or '$(Configuration)' == 'VS15d'" />
+    <Reference Include="Microsoft.VisualStudio.Shell.Framework, Version=15.0.0.0, Culture=neutral, PublicKeyToken=b03f5f7f11d50a3a, processorArchitecture=MSIL" Condition="'$(Configuration)' == 'VS15' Or '$(Configuration)' == 'VS15d'" />
+    <Reference Include="Microsoft.VisualStudio.Shell.Immutable.10.0" />
+    <Reference Include="Microsoft.VisualStudio.Shell.Interop, Version=7.1.40304.0, Culture=neutral, PublicKeyToken=b03f5f7f11d50a3a" />
+    <Reference Include="Microsoft.VisualStudio.Shell.Interop.10.0, Version=10.0.0.0, Culture=neutral, PublicKeyToken=b03f5f7f11d50a3a, processorArchitecture=MSIL">
+      <EmbedInteropTypes>True</EmbedInteropTypes>
+    </Reference>
+    <Reference Include="Microsoft.VisualStudio.Shell.Interop.11.0, Version=11.0.0.0, Culture=neutral, PublicKeyToken=b03f5f7f11d50a3a, processorArchitecture=MSIL">
+      <EmbedInteropTypes>True</EmbedInteropTypes>
+    </Reference>
+    <Reference Include="Microsoft.VisualStudio.Shell.Interop.8.0, Version=8.0.0.0, Culture=neutral, PublicKeyToken=b03f5f7f11d50a3a" />
+    <Reference Include="Microsoft.VisualStudio.Shell.Interop.9.0, Version=9.0.0.0, Culture=neutral, PublicKeyToken=b03f5f7f11d50a3a" />
+    <Reference Include="Microsoft.VisualStudio.TextManager.Interop, Version=7.1.40304.0, Culture=neutral, PublicKeyToken=b03f5f7f11d50a3a" />
+    <Reference Include="Microsoft.VisualStudio.VCProject, Version=15.0.0.0, Culture=neutral, PublicKeyToken=b03f5f7f11d50a3a" Condition="'$(Configuration)' == 'VS15' Or '$(Configuration)' == 'VS15d'">
+      <EmbedInteropTypes>True</EmbedInteropTypes>
+    </Reference>
+    <Reference Include="Microsoft.VisualStudio.VCProject, Version=14.0.0.0, Culture=neutral, PublicKeyToken=b03f5f7f11d50a3a" Condition="'$(Configuration)' == 'VS14' Or '$(Configuration)' == 'VS14d'">
+      <EmbedInteropTypes>True</EmbedInteropTypes>
+    </Reference>
+    <Reference Include="Microsoft.VisualStudio.VCProject, Version=12.0.0.0, Culture=neutral, PublicKeyToken=b03f5f7f11d50a3a" Condition="'$(Configuration)' == 'VS12' Or '$(Configuration)' == 'VS12d'">
+      <EmbedInteropTypes>True</EmbedInteropTypes>
+    </Reference>
+    <Reference Include="Microsoft.VisualStudio.VCProject, Version=11.0.0.0, Culture=neutral, PublicKeyToken=b03f5f7f11d50a3a" Condition="'$(Configuration)' == 'VS11' Or '$(Configuration)' == 'VS11d'">
+      <EmbedInteropTypes>True</EmbedInteropTypes>
+    </Reference>
+    <Reference Include="Microsoft.VisualStudio.VCProjectEngine, Version=15.0.0.0, Culture=neutral, PublicKeyToken=b03f5f7f11d50a3a" Condition="'$(Configuration)' == 'VS15' Or '$(Configuration)' == 'VS15d'">
+      <EmbedInteropTypes>True</EmbedInteropTypes>
+    </Reference>
+    <Reference Include="Microsoft.VisualStudio.VCProjectEngine, Version=14.0.0.0, Culture=neutral, PublicKeyToken=b03f5f7f11d50a3a" Condition="'$(Configuration)' == 'VS14' Or '$(Configuration)' == 'VS14d'">
+      <EmbedInteropTypes>True</EmbedInteropTypes>
+    </Reference>
+    <Reference Include="Microsoft.VisualStudio.VCProjectEngine, Version=12.0.0.0, Culture=neutral, PublicKeyToken=b03f5f7f11d50a3a" Condition="'$(Configuration)' == 'VS12' Or '$(Configuration)' == 'VS12d'">
+      <EmbedInteropTypes>True</EmbedInteropTypes>
+    </Reference>
+    <Reference Include="Microsoft.VisualStudio.VCProjectEngine, Version=11.0.0.0, Culture=neutral, PublicKeyToken=b03f5f7f11d50a3a" Condition="'$(Configuration)' == 'VS11' Or '$(Configuration)' == 'VS11d'">
+      <EmbedInteropTypes>True</EmbedInteropTypes>
+    </Reference>
+    <Reference Include="PresentationCore" />
+    <Reference Include="PresentationFramework" />
+    <Reference Include="System" />
+    <Reference Include="System.Core" />
+    <Reference Include="System.Data" />
+    <Reference Include="System.Design" />
+    <Reference Include="System.Drawing" />
+    <Reference Include="System.Windows.Forms" />
+    <Reference Include="System.Xaml" />
+    <Reference Include="System.Xml" />
+    <Reference Include="VSLangProj, Version=7.0.3300.0, Culture=neutral, PublicKeyToken=b03f5f7f11d50a3a">
+      <EmbedInteropTypes>True</EmbedInteropTypes>
+    </Reference>
+    <Reference Include="VSLangProj80, Version=8.0.0.0, Culture=neutral, PublicKeyToken=b03f5f7f11d50a3a">
+      <EmbedInteropTypes>True</EmbedInteropTypes>
+    </Reference>
+    <Reference Include="WindowsBase" />
+  </ItemGroup>
+  <ItemGroup>
+    <COMReference Include="Microsoft.VisualStudio.CommandBars">
+      <Guid>{1CBA492E-7263-47BB-87FE-639000619B15}</Guid>
+      <VersionMajor>8</VersionMajor>
+      <VersionMinor>0</VersionMinor>
+      <Lcid>0</Lcid>
+      <WrapperTool>primary</WrapperTool>
+      <Isolated>False</Isolated>
+      <EmbedInteropTypes>False</EmbedInteropTypes>
+    </COMReference>
+    <COMReference Include="stdole">
+      <Guid>{00020430-0000-0000-C000-000000000046}</Guid>
+      <VersionMajor>2</VersionMajor>
+      <VersionMinor>0</VersionMinor>
+      <Lcid>0</Lcid>
+      <WrapperTool>primary</WrapperTool>
+      <Isolated>False</Isolated>
+      <EmbedInteropTypes>False</EmbedInteropTypes>
+    </COMReference>
+  </ItemGroup>
+  <ItemGroup>
+    <Compile Include="BatchBuilder.cs" />
+    <Compile Include="BatchBuilderToolControl.xaml.cs">
+      <DependentUpon>BatchBuilderToolControl.xaml</DependentUpon>
+    </Compile>
+    <Compile Include="BatchBuilderToolWindow.cs" />
+    <Compile Include="BuildStartupProject.cs" />
+    <Compile Include="BuildStatusIcon.xaml.cs">
+      <DependentUpon>BuildStatusIcon.xaml</DependentUpon>
+    </Compile>
+    <Compile Include="CommandLineEditor.cs" />
+    <Compile Include="CompileSingleFile.cs" />
+    <Compile Include="GenerateProjectFiles.cs" />
+    <Compile Include="Guids.cs" />
+    <Compile Include="Logging.cs" />
+    <Compile Include="ProfileManager.cs">
+      <SubType>Component</SubType>
+    </Compile>
+    <Compile Include="QuickBuild.cs" />
+    <Compile Include="Resources.Designer.cs">
+      <AutoGen>True</AutoGen>
+      <DesignTime>True</DesignTime>
+      <DependentUpon>Resources.resx</DependentUpon>
+    </Compile>
+    <Compile Include="GlobalSuppressions.cs" />
+    <Compile Include="StartupProjectSelector.cs" />
+    <Compile Include="UnrealVSOptions.cs">
+      <SubType>Component</SubType>
+    </Compile>
+    <Compile Include="UnrealVSPackage.cs" />
+    <Compile Include="Properties\AssemblyInfo.cs" />
+    <Compile Include="Utils.cs" />
+    <Compile Include="VCProject.cs" />
+  </ItemGroup>
+  <ItemGroup>
+    <EmbeddedResource Include="Resources.resx">
+      <Generator>ResXFileCodeGenerator</Generator>
+      <LastGenOutput>Resources.Designer.cs</LastGenOutput>
+      <SubType>Designer</SubType>
+    </EmbeddedResource>
+    <EmbeddedResource Include="VSPackage.resx">
+      <MergeWithCTO>true</MergeWithCTO>
+      <ManifestResourceName>VSPackage</ManifestResourceName>
+      <SubType>Designer</SubType>
+    </EmbeddedResource>
+  </ItemGroup>
+  <ItemGroup>
+    <None Include="TargetManifests\VS11.source.extension.vsixmanifest">
+      <SubType>Designer</SubType>
+    </None>
+    <None Include="TargetManifests\VS12.source.extension.vsixmanifest">
+      <SubType>Designer</SubType>
+    </None>
+    <None Include="TargetManifests\VS11d.source.extension.vsixmanifest">
+      <SubType>Designer</SubType>
+    </None>
+    <None Include="TargetManifests\VS12d.source.extension.vsixmanifest">
+      <SubType>Designer</SubType>
+    </None>
+  </ItemGroup>
+  <ItemGroup>
+    <None Include="Key.snk" />
+  </ItemGroup>
+  <ItemGroup>
+    <Content Include="TargetManifests\VS14.source.extension.vsixmanifest" />
+    <Content Include="TargetManifests\VS14d.source.extension.vsixmanifest" />
+    <Content Include="TargetManifests\VS15.source.extension.vsixmanifest" />
+    <Content Include="TargetManifests\VS15d.source.extension.vsixmanifest" />
+    <None Include="Generated\source.extension.vsixmanifest" />
+    <None Include="Resources\BatchBuilder.png" />
+    <Resource Include="Resources\BatchBuilderBackground.png" />
+    <Content Include="Resources\Icon.png">
+      <CopyToOutputDirectory>Always</CopyToOutputDirectory>
+      <IncludeInVSIX>true</IncludeInVSIX>
+    </Content>
+    <Content Include="Resources\Preview.png">
+      <IncludeInVSIX>true</IncludeInVSIX>
+      <CopyToOutputDirectory>Always</CopyToOutputDirectory>
+    </Content>
+    <Resource Include="Resources\RefreshProjects.png" />
+  </ItemGroup>
+  <ItemGroup>
+    <BootstrapperPackage Include=".NETFramework,Version=v4.5">
+      <Visible>False</Visible>
+      <ProductName>Microsoft .NET Framework 4.5 %28x86 and x64%29</ProductName>
+      <Install>true</Install>
+    </BootstrapperPackage>
+    <BootstrapperPackage Include="Microsoft.Net.Client.3.5">
+      <Visible>False</Visible>
+      <ProductName>.NET Framework 3.5 SP1 Client Profile</ProductName>
+      <Install>false</Install>
+    </BootstrapperPackage>
+    <BootstrapperPackage Include="Microsoft.Net.Framework.3.5.SP1">
+      <Visible>False</Visible>
+      <ProductName>.NET Framework 3.5 SP1</ProductName>
+      <Install>false</Install>
+    </BootstrapperPackage>
+  </ItemGroup>
+  <ItemGroup>
+    <Page Include="BatchBuilderToolControl.xaml">
+      <SubType>Designer</SubType>
+      <Generator>MSBuild:Compile</Generator>
+      <Generator>MSBuild:Compile</Generator>
+      <SubType>Designer</SubType>
+      <Generator>MSBuild:Compile</Generator>
+      <SubType>Designer</SubType>
+    </Page>
+    <Page Include="BuildStatusIcon.xaml">
+      <SubType>Designer</SubType>
+      <Generator>MSBuild:Compile</Generator>
+      <Generator>MSBuild:Compile</Generator>
+      <SubType>Designer</SubType>
+      <Generator>MSBuild:Compile</Generator>
+      <SubType>Designer</SubType>
+    </Page>
+  </ItemGroup>
+  <PropertyGroup>
+    <UseCodebase>true</UseCodebase>
+  </PropertyGroup>
+  <PropertyGroup>
+    <VisualStudioVersion Condition="'$(VisualStudioVersion)' == ''">11.0</VisualStudioVersion>
+    <VSToolsPath Condition="'$(VSToolsPath)' == ''">$(MSBuildExtensionsPath32)\Microsoft\VisualStudio\v$(VisualStudioVersion)</VSToolsPath>
+  </PropertyGroup>
+  <PropertyGroup Condition="'$(Configuration)|$(Platform)' == 'VS11|AnyCPU'">
+    <DebugType>pdbonly</DebugType>
+    <Optimize>true</Optimize>
+    <OutputPath>..\..\..\Extras\UnrealVS\VS2012\</OutputPath>
+    <DefineConstants>TRACE;CODE_ANALYSIS;VS11</DefineConstants>
+    <ErrorReport>prompt</ErrorReport>
+    <WarningLevel>4</WarningLevel>
+    <RunCodeAnalysis>true</RunCodeAnalysis>
+    <Prefer32Bit>false</Prefer32Bit>
+  </PropertyGroup>
+  <PropertyGroup Condition=" '$(Configuration)|$(Platform)' == 'VS11d|AnyCPU' ">
+    <DebugSymbols>true</DebugSymbols>
+    <DebugType>full</DebugType>
+    <Optimize>false</Optimize>
+    <OutputPath>bin\VS11d\</OutputPath>
+    <DefineConstants>TRACE;DEBUG;VS11</DefineConstants>
+    <ErrorReport>prompt</ErrorReport>
+    <WarningLevel>4</WarningLevel>
+    <Prefer32Bit>false</Prefer32Bit>
+    <CopyVsixExtensionFiles>False</CopyVsixExtensionFiles>
+    <DeployExtension>True</DeployExtension>
+    <CreateVsixContainer>True</CreateVsixContainer>
+  </PropertyGroup>
+  <PropertyGroup Condition=" '$(Configuration)|$(Platform)' == 'VS14d|AnyCPU' ">
+    <DebugSymbols>true</DebugSymbols>
+    <OutputPath>bin\VS14d\</OutputPath>
+    <DefineConstants>TRACE;DEBUG;VS14</DefineConstants>
+    <DebugType>full</DebugType>
+    <PlatformTarget>AnyCPU</PlatformTarget>
+    <ErrorReport>prompt</ErrorReport>
+    <CodeAnalysisRuleSet>MinimumRecommendedRules.ruleset</CodeAnalysisRuleSet>
+  </PropertyGroup>
+  <PropertyGroup Condition=" '$(Configuration)|$(Platform)' == 'VS14|AnyCPU' ">
+    <OutputPath>..\..\..\Extras\UnrealVS\VS2015\</OutputPath>
+    <DefineConstants>TRACE;CODE_ANALYSIS;VS14</DefineConstants>
+    <Optimize>true</Optimize>
+    <DebugType>pdbonly</DebugType>
+    <PlatformTarget>AnyCPU</PlatformTarget>
+    <RunCodeAnalysis>true</RunCodeAnalysis>
+    <ErrorReport>prompt</ErrorReport>
+    <CodeAnalysisRuleSet>MinimumRecommendedRules.ruleset</CodeAnalysisRuleSet>
+  </PropertyGroup>
+  <PropertyGroup Condition="'$(Configuration)|$(Platform)' == 'VS15|AnyCPU'">
+    <OutputPath>..\..\..\Extras\UnrealVS\VS2017\</OutputPath>
+    <DefineConstants>CODE_ANALYSIS;TRACE;CODE_ANALYSIS;VS14</DefineConstants>
+    <Optimize>true</Optimize>
+    <DebugType>pdbonly</DebugType>
+    <PlatformTarget>AnyCPU</PlatformTarget>
+    <RunCodeAnalysis>true</RunCodeAnalysis>
+    <ErrorReport>prompt</ErrorReport>
+    <CodeAnalysisRuleSet>MinimumRecommendedRules.ruleset</CodeAnalysisRuleSet>
+  </PropertyGroup>
+  <PropertyGroup Condition="'$(Configuration)|$(Platform)' == 'VS15d|AnyCPU'">
+    <DebugSymbols>true</DebugSymbols>
+    <OutputPath>bin\VS15d\</OutputPath>
+    <DefineConstants>TRACE;DEBUG;VS14</DefineConstants>
+    <DebugType>full</DebugType>
+    <PlatformTarget>AnyCPU</PlatformTarget>
+    <ErrorReport>prompt</ErrorReport>
+    <CodeAnalysisRuleSet>MinimumRecommendedRules.ruleset</CodeAnalysisRuleSet>
+  </PropertyGroup>
+  <Import Project="$(MSBuildBinPath)\Microsoft.CSharp.targets" />
+  <Import Project="$(VSToolsPath)\VSSDK\Microsoft.VsSDK.targets" Condition="'$(VSToolsPath)' != ''" />
+  <Import Project="$(MSBuildExtensionsPath32)\Microsoft\VisualStudio\v11.0\VSSDK\Microsoft.VsSDK.targets" Condition="false" />
+  <PropertyGroup>
+    <PostBuildEvent>
+    </PostBuildEvent>
+  </PropertyGroup>
+  <PropertyGroup>
+    <PreBuildEvent>copy "$(ProjectDir)TargetManifests\$(ConfigurationName).source.extension.vsixmanifest" "$(ProjectDir)Generated\source.extension.vsixmanifest" /Y</PreBuildEvent>
+  </PropertyGroup>
+  <!-- To modify your build process, add your task inside one of the targets below and uncomment it. 
+       Other similar extension points exist, see Microsoft.Common.targets.
+  <Target Name="BeforeBuild">
+  </Target>
+  <Target Name="AfterBuild">
+  </Target>
+  -->
+</Project>
\ No newline at end of file
diff --git a/Engine/Source/Runtime/AIModule/Classes/AITypes.h b/Engine/Source/Runtime/AIModule/Classes/AITypes.h
--- a/Engine/Source/Runtime/AIModule/Classes/AITypes.h
+++ b/Engine/Source/Runtime/AIModule/Classes/AITypes.h
@@ -299,6 +299,11 @@ public:
 		static const FAIGenericID<TCounter> InvalidIDInstance;
 		return InvalidIDInstance;
 	}
+
+	friend FORCEINLINE uint32 GetTypeHash(const FAIGenericID& ID)
+	{
+		return GetTypeHash(ID.Index);
+	}
 };
 
 template<typename TCounterType>
diff --git a/Engine/Source/Runtime/AIModule/Classes/BehaviorTree/BehaviorTreeComponent.h b/Engine/Source/Runtime/AIModule/Classes/BehaviorTree/BehaviorTreeComponent.h
--- a/Engine/Source/Runtime/AIModule/Classes/BehaviorTree/BehaviorTreeComponent.h
+++ b/Engine/Source/Runtime/AIModule/Classes/BehaviorTree/BehaviorTreeComponent.h
@@ -98,10 +98,8 @@ protected:
 	bool TreeHasBeenStarted() const;
 
 public:
-#if WITH_HOT_RELOAD_CTORS
 	/** DO NOT USE. This constructor is for internal usage only for hot-reload purposes. */
 	UBehaviorTreeComponent(FVTableHelper& Helper);
-#endif // WITH_HOT_RELOAD_CTORS
 
 	virtual bool IsRunning() const override;
 	virtual bool IsPaused() const override;
diff --git a/Engine/Source/Runtime/Core/Private/Async/TaskGraph.cpp b/Engine/Source/Runtime/Core/Private/Async/TaskGraph.cpp
--- a/Engine/Source/Runtime/Core/Private/Async/TaskGraph.cpp
+++ b/Engine/Source/Runtime/Core/Private/Async/TaskGraph.cpp
@@ -1469,6 +1469,7 @@ public:
 			FString Name;
 			int32 Priority = ThreadIndexToPriorityIndex(ThreadIndex);
 			EThreadPriority ThreadPri;
+			uint64 Affinity = FPlatformAffinity::GetTaskGraphThreadMask();
 			if (Priority == 1)
 			{
 				Name = FString::Printf(TEXT("TaskGraphThreadHP %d"), ThreadIndex - (LastExternalThread + 1));
@@ -1478,6 +1479,11 @@ public:
 			{
 				Name = FString::Printf(TEXT("TaskGraphThreadBP %d"), ThreadIndex - (LastExternalThread + 1));
 				ThreadPri = TPri_Lowest;
+				if (PLATFORM_PS4)
+				{
+					// hack, use the audio affinity mask, since this might include the 7th core
+					Affinity = FPlatformAffinity::GetTaskGraphBackgroundTaskMask();
+				}
 			}
 			else
 			{
@@ -1485,7 +1491,7 @@ public:
 				ThreadPri = TPri_BelowNormal; // we want normal tasks below normal threads like the game thread
 			}
 			uint32 StackSize = 384 * 1024;
-			WorkerThreads[ThreadIndex].RunnableThread = FRunnableThread::Create(&Thread(ThreadIndex), *Name, StackSize, ThreadPri, FPlatformAffinity::GetTaskGraphThreadMask()); // these are below normal threads so that they sleep when the named threads are active
+			WorkerThreads[ThreadIndex].RunnableThread = FRunnableThread::Create(&Thread(ThreadIndex), *Name, StackSize, ThreadPri, Affinity); // these are below normal threads so that they sleep when the named threads are active
 			WorkerThreads[ThreadIndex].bAttached = true;
 		}
 	}
diff --git a/Engine/Source/Runtime/Core/Private/Misc/App.cpp b/Engine/Source/Runtime/Core/Private/Misc/App.cpp
--- a/Engine/Source/Runtime/Core/Private/Misc/App.cpp
+++ b/Engine/Source/Runtime/Core/Private/Misc/App.cpp
@@ -82,8 +82,7 @@ EBuildConfigurations::Type FApp::GetBuildConfiguration()
 		extern const bool GIsDebugGame;
 		return GIsDebugGame? EBuildConfigurations::DebugGame : EBuildConfigurations::Development;
 	#else
-		static const bool bUsingDebugGame = FParse::Param(FCommandLine::Get(), TEXT("debug"));
-		return bUsingDebugGame? EBuildConfigurations::DebugGame : EBuildConfigurations::Development;
+		return IsRunningDebug() ? EBuildConfigurations::DebugGame : EBuildConfigurations::Development;
 	#endif
 
 #elif UE_BUILD_SHIPPING
@@ -97,6 +96,14 @@ EBuildConfigurations::Type FApp::GetBuildConfiguration()
 #endif
 }
 
+bool FApp::IsRunningDebug()
+{
+	static FString RunConfig;
+	static const bool bHasRunConfig = FParse::Value(FCommandLine::Get(), TEXT("RunConfig="), RunConfig);
+	static const bool bRunningDebug = FParse::Param(FCommandLine::Get(), TEXT("debug"))
+	                                  || (bHasRunConfig && RunConfig.StartsWith(TEXT("Debug")));
+	return bRunningDebug;
+}
 
 FString FApp::GetBuildDate()
 {
diff --git a/Engine/Source/Runtime/Core/Private/Misc/OutputDeviceFile.cpp b/Engine/Source/Runtime/Core/Private/Misc/OutputDeviceFile.cpp
--- a/Engine/Source/Runtime/Core/Private/Misc/OutputDeviceFile.cpp
+++ b/Engine/Source/Runtime/Core/Private/Misc/OutputDeviceFile.cpp
@@ -59,6 +59,8 @@ class CORE_API FAsyncWriter : public FRunnable, public FArchive
 	FCriticalSection BufferPosCritical;
 	/** [CLIENT/WRITER THREAD] Outstanding serialize request counter. This is to make sure we flush all requests. */
 	FThreadSafeCounter SerializeRequestCounter;
+	/** [CLIENT/WRITER THREAD] Tells the writer thread, the client requested flush. */
+	FThreadSafeCounter WantsArchiveFlush;
 
 	/** [WRITER THREAD] Last time the archive was flushed. used in threaded situations to flush the underlying archive at a certain maximum rate. */
 	double LastArchiveFlushTime;
@@ -66,6 +68,15 @@ class CORE_API FAsyncWriter : public FRunnable, public FArchive
 	/** [WRITER THREAD] Archive flush interval. */
 	double ArchiveFlushIntervalSec;
 
+	/** [WRITER THREAD] Flushes the archive and reset the flush timer. */
+	void FlushArchiveAndResetTimer()
+	{
+		// This should be the one and only place where we flush because we want the flush to happen only on the 
+		// async writer thread (if threading is enabled)
+		Ar.Flush();
+		LastArchiveFlushTime = FPlatformTime::Seconds();
+	}
+
 	/** [WRITER THREAD] Serialize the contents of the ring buffer to disk */
 	void SerializeBufferToArchive()
 	{
@@ -89,19 +100,27 @@ class CORE_API FAsyncWriter : public FRunnable, public FArchive
 			// Modify the start pos. Only the worker thread modifies this value so it's ok to not guard it with a critical section.
 			BufferStartPos = ThisThreadEndPos;
 
+			// Decrement the request counter, we now know we serialized at least one request.
+			// We might have serialized more requests but it's irrelevant, the counter will go down to 0 eventually
+			SerializeRequestCounter.Decrement();
+
 			// Flush the archive periodically if running on a separate thread
 			if (Thread)
 			{
 				if ((FPlatformTime::Seconds() - LastArchiveFlushTime) > ArchiveFlushIntervalSec)
 				{
-					Ar.Flush();
-					LastArchiveFlushTime = FPlatformTime::Seconds();
+					FlushArchiveAndResetTimer();
 				}
 			}
-
-			// Decrement the request counter, we now know we serialized at least one request.
-			// We might have serialized more requests but it's irrelevant, the counter will go down to 0 eventually
-			SerializeRequestCounter.Decrement();
+			// If no threading is available or when we explicitly requested flush (see FlushBuffer), flush immediately after writing.
+			// In some rare cases we may flush twice (see above) but that's ok. We need a clear division between flushing because of the timer
+			// and force flush on demand.
+			if (WantsArchiveFlush.GetValue() > 0)
+			{
+				FlushArchiveAndResetTimer();
+				int32 FlushCount = WantsArchiveFlush.Decrement();
+				check(FlushCount >= 0);
+			}
 		}
 	}
 
@@ -217,10 +236,8 @@ public:
 	void Flush()
 	{
 		FScopeLock WriteLock(&BufferPosCritical);
+		WantsArchiveFlush.Increment();
 		FlushBuffer();
-		// At this point the serialize queue should be empty and the writer thread should no longer
-		// be accessing the Archive so we should be safe to flush it from here.
-		Ar.Flush();
 	}
 
 	//~ Begin FRunnable Interface.
@@ -238,11 +255,11 @@ public:
 			}
 			else if ((FPlatformTime::Seconds() - LastArchiveFlushTime) > ArchiveFlushIntervalSec)
 			{
-				SerializeRequestCounter.Increment();
+				FlushArchiveAndResetTimer();
 			}
 			else
 			{
-				FPlatformProcess::Sleep(0.01f);
+				FPlatformProcess::SleepNoStats(0.01f);
 			}
 		}
 		return 0;
diff --git a/Engine/Source/Runtime/Core/Public/Containers/ContainerAllocationPolicies.h b/Engine/Source/Runtime/Core/Public/Containers/ContainerAllocationPolicies.h
--- a/Engine/Source/Runtime/Core/Public/Containers/ContainerAllocationPolicies.h
+++ b/Engine/Source/Runtime/Core/Public/Containers/ContainerAllocationPolicies.h
@@ -769,6 +769,8 @@ private:
 
 	enum { NumInlineHashBuckets = (NumInlineElements + AverageNumberOfElementsPerHashBucket - 1) / AverageNumberOfElementsPerHashBucket };
 
+	static_assert(!(NumInlineHashBuckets & (NumInlineHashBuckets - 1)), "Number of inline buckets must be a power of two");
+
 public:
 
 	/** Computes the number of hash buckets to use for a given number of elements. */
diff --git a/Engine/Source/Runtime/Core/Public/Containers/ContainersFwd.h b/Engine/Source/Runtime/Core/Public/Containers/ContainersFwd.h
--- a/Engine/Source/Runtime/Core/Public/Containers/ContainersFwd.h
+++ b/Engine/Source/Runtime/Core/Public/Containers/ContainersFwd.h
@@ -9,6 +9,6 @@ class FString;
 
 template<typename T, typename Allocator = FDefaultAllocator> class TArray;
 template<typename T> class TTransArray;
-template<typename KeyType, typename ValueType, bool bInAllowDuplicateKeys> struct TDefaultMapKeyFuncs;
-template<typename KeyType, typename ValueType, typename SetAllocator = FDefaultSetAllocator, typename KeyFuncs = TDefaultMapKeyFuncs<KeyType, ValueType, false> > class TMap;
-template<typename KeyType, typename ValueType, typename SetAllocator = FDefaultSetAllocator, typename KeyFuncs = TDefaultMapKeyFuncs<KeyType, ValueType, true > > class TMultiMap;
+template<typename KeyType, typename ValueType, bool bInAllowDuplicateKeys> struct TDefaultMapHashableKeyFuncs;
+template<typename KeyType, typename ValueType, typename SetAllocator = FDefaultSetAllocator, typename KeyFuncs = TDefaultMapHashableKeyFuncs<KeyType, ValueType, false> > class TMap;
+template<typename KeyType, typename ValueType, typename SetAllocator = FDefaultSetAllocator, typename KeyFuncs = TDefaultMapHashableKeyFuncs<KeyType, ValueType, true > > class TMultiMap;
diff --git a/Engine/Source/Runtime/Core/Public/Containers/EnumAsByte.h b/Engine/Source/Runtime/Core/Public/Containers/EnumAsByte.h
--- a/Engine/Source/Runtime/Core/Public/Containers/EnumAsByte.h
+++ b/Engine/Source/Runtime/Core/Public/Containers/EnumAsByte.h
@@ -129,6 +129,12 @@ private:
 
 	/** Holds the value as a byte. **/
 	uint8 Value;
+
+
+	FORCEINLINE friend uint32 GetTypeHash(const TEnumAsByte& Enum)
+	{
+		return GetTypeHash(Enum.Value);
+	}
 };
 
 
diff --git a/Engine/Source/Runtime/Core/Public/Containers/Map.h b/Engine/Source/Runtime/Core/Public/Containers/Map.h
--- a/Engine/Source/Runtime/Core/Public/Containers/Map.h
+++ b/Engine/Source/Runtime/Core/Public/Containers/Map.h
@@ -13,6 +13,7 @@
 #include "Containers/Set.h"
 #include "Containers/Algo/Reverse.h"
 #include "Templates/Tuple.h"
+#include "HasGetTypeHash.h"
 
 #define ExchangeB(A,B) {bool T=A; A=B; B=T;}
 
@@ -90,6 +91,12 @@ struct TDefaultMapKeyFuncs : BaseKeyFuncs<TPair<KeyType,ValueType>,KeyType,bInAl
 	}
 };
 
+template<typename KeyType, typename ValueType, bool bInAllowDuplicateKeys>
+struct TDefaultMapHashableKeyFuncs : TDefaultMapKeyFuncs<KeyType, ValueType, bInAllowDuplicateKeys>
+{
+	static_assert(THasGetTypeHash<KeyType>::Value, "TMap must have a hashable KeyType unless a custom key func is provided.");
+};
+
 /** 
  * The base class of maps from keys to values.  Implemented using a TSet of key-value pairs with a custom KeyFuncs, 
  * with the same O(1) addition, removal, and finding. 
@@ -851,7 +858,7 @@ private:
 class FScriptMap;
 
 /** A TMapBase specialization that only allows a single value associated with each key.*/
-template<typename KeyType,typename ValueType,typename SetAllocator /*= FDefaultSetAllocator*/,typename KeyFuncs /*= TDefaultMapKeyFuncs<KeyType,ValueType,false>*/>
+template<typename KeyType,typename ValueType,typename SetAllocator /*= FDefaultSetAllocator*/,typename KeyFuncs /*= TDefaultMapHashableKeyFuncs<KeyType,ValueType,false>*/>
 class TMap : public TSortableMapBase<KeyType, ValueType, SetAllocator, KeyFuncs>
 {
 	friend struct TContainerTraits<TMap>;
@@ -979,7 +986,7 @@ public:
 };
 
 /** A TMapBase specialization that allows multiple values to be associated with each key. */
-template<typename KeyType,typename ValueType,typename SetAllocator /* = FDefaultSetAllocator */,typename KeyFuncs /*= TDefaultMapKeyFuncs<KeyType,ValueType,true>*/>
+template<typename KeyType,typename ValueType,typename SetAllocator /* = FDefaultSetAllocator */,typename KeyFuncs /*= TDefaultMapHashableKeyFuncs<KeyType,ValueType,true>*/>
 class TMultiMap : public TSortableMapBase<KeyType, ValueType, SetAllocator, KeyFuncs>
 {
 	friend struct TContainerTraits<TMultiMap>;
diff --git a/Engine/Source/Runtime/Core/Public/Containers/Set.h b/Engine/Source/Runtime/Core/Public/Containers/Set.h
--- a/Engine/Source/Runtime/Core/Public/Containers/Set.h
+++ b/Engine/Source/Runtime/Core/Public/Containers/Set.h
@@ -1035,7 +1035,7 @@ private:
 		if (LocalHashSize)
 		{
 			// Allocate the new hash.
-			checkSlow(!(LocalHashSize & (HashSize - 1)));
+			checkSlow(FMath::IsPowerOfTwo(HashSize));
 			Hash.ResizeAllocation(0, LocalHashSize, sizeof(FSetElementId));
 			for (int32 HashIndex = 0; HashIndex < LocalHashSize; ++HashIndex)
 			{
diff --git a/Engine/Source/Runtime/Core/Public/CoreGlobals.h b/Engine/Source/Runtime/Core/Public/CoreGlobals.h
--- a/Engine/Source/Runtime/Core/Public/CoreGlobals.h
+++ b/Engine/Source/Runtime/Core/Public/CoreGlobals.h
@@ -297,10 +297,8 @@ extern CORE_API bool GEventDrivenLoaderEnabled;
 //@todoio put this in some kind of API
 extern CORE_API bool GPakCache_AcceptPrecacheRequests;
 
-#if WITH_HOT_RELOAD_CTORS
 /** true when we are retrieving VTablePtr from UClass */
 extern CORE_API bool GIsRetrievingVTablePtr;
-#endif // WITH_HOT_RELOAD_CTORS
 
 /** Steadily increasing frame counter. */
 extern CORE_API uint64 GFrameCounter;
@@ -385,7 +383,6 @@ extern CORE_API bool GPumpingMessagesOutsideOfMainLoop;
 /** Enables various editor and HMD hacks that allow the experimental VR editor feature to work, perhaps at the expense of other systems */
 extern CORE_API bool GEnableVREditorHacks;
 
-#if WITH_HOT_RELOAD_CTORS
 /**
  * Ensures that current thread is during retrieval of vtable ptr of some
  * UClass.
@@ -394,7 +391,6 @@ extern CORE_API bool GEnableVREditorHacks;
  *		construct proper error message.
  */
 CORE_API void EnsureRetrievingVTablePtrDuringCtor(const TCHAR* CtorSignature);
-#endif // WITH_HOT_RELOAD_CTORS
 
 /** @return True if called from the game thread. */
 FORCEINLINE bool IsInGameThread()
diff --git a/Engine/Source/Runtime/Core/Public/GenericPlatform/GenericPlatformAffinity.h b/Engine/Source/Runtime/Core/Public/GenericPlatform/GenericPlatformAffinity.h
--- a/Engine/Source/Runtime/Core/Public/GenericPlatform/GenericPlatformAffinity.h
+++ b/Engine/Source/Runtime/Core/Public/GenericPlatform/GenericPlatformAffinity.h
@@ -73,6 +73,11 @@ public:
 		return 0xFFFFFFFFFFFFFFFF;
 	}
 
+	static const CORE_API uint64 GetTaskGraphBackgroundTaskMask()
+	{
+		return 0xFFFFFFFFFFFFFFFF;
+	}
+
 	// @todo what do we think about having this as a function in this class? Should be make a whole new one? 
 	// scrap it and force the priority like before?
 	static EThreadPriority GetRenderingThreadPriority()
diff --git a/Engine/Source/Runtime/Core/Public/Misc/App.h b/Engine/Source/Runtime/Core/Public/Misc/App.h
--- a/Engine/Source/Runtime/Core/Public/Misc/App.h
+++ b/Engine/Source/Runtime/Core/Public/Misc/App.h
@@ -40,6 +40,11 @@ public:
 	 */
 	static EBuildConfigurations::Type GetBuildConfiguration();
 
+	/**
+	 * Gets whether the application is running with Debug game libraries (set from commandline)
+	 */
+	static bool IsRunningDebug();
+
 	/*
 	* Gets the unique version string for this build. This string is not assumed to have any particular format other being a unique identifier for the build.
 	*
diff --git a/Engine/Source/Runtime/Core/Public/Misc/Build.h b/Engine/Source/Runtime/Core/Public/Misc/Build.h
--- a/Engine/Source/Runtime/Core/Public/Misc/Build.h
+++ b/Engine/Source/Runtime/Core/Public/Misc/Build.h
@@ -141,11 +141,6 @@
 	#define WITH_HOT_RELOAD (!IS_MONOLITHIC && !UE_BUILD_SHIPPING && !UE_BUILD_TEST && !UE_GAME && !UE_SERVER)
 #endif
 
-/**
- * Whether we want to use special hot-reload empty constructors.
- */
-#define WITH_HOT_RELOAD_CTORS 1
-
 /*----------------------------------------------------------------------------
 	Optional bridge options coming from UBT, do not modify directly!
 	If UBT doesn't set the value, it is assumed to be 0, and we set that here.
diff --git a/Engine/Source/Runtime/Core/Public/Misc/ExpressionParserTypes.inl b/Engine/Source/Runtime/Core/Public/Misc/ExpressionParserTypes.inl
--- a/Engine/Source/Runtime/Core/Public/Misc/ExpressionParserTypes.inl
+++ b/Engine/Source/Runtime/Core/Public/Misc/ExpressionParserTypes.inl
@@ -4,13 +4,13 @@
 
 #include "CoreFwd.h"
 #include "Templates/RemoveReference.h"
+#include "Templates/PointerIsConvertibleFromTo.h"
 
 #define LOCTEXT_NAMESPACE "ExpressionParser"
 
 class FExpressionNode;
 class FExpressionToken;
 struct FOperatorFunctionID;
-template <typename From, typename To> struct TPointerIsConvertibleFromTo;
 template< class T > struct TRemoveConst;
 
 namespace Impl
diff --git a/Engine/Source/Runtime/Core/Public/Misc/InlineValue.h b/Engine/Source/Runtime/Core/Public/Misc/InlineValue.h
--- a/Engine/Source/Runtime/Core/Public/Misc/InlineValue.h
+++ b/Engine/Source/Runtime/Core/Public/Misc/InlineValue.h
@@ -6,9 +6,8 @@
 #include "Misc/AssertionMacros.h"
 #include "HAL/UnrealMemory.h"
 #include "Templates/MemoryOps.h"
-
-template <typename From, typename To> struct TPointerIsConvertibleFromTo;
-template <typename T> struct TDecay;
+#include "Templates/PointerIsConvertibleFromTo.h"
+#include "Templates/Decay.h"
 
 /**
  * A container type that houses an instance of BaseType in inline memory where it is <= MaxInlineSize,
diff --git a/Engine/Source/Runtime/Core/Public/Templates/RefCounting.h b/Engine/Source/Runtime/Core/Public/Templates/RefCounting.h
--- a/Engine/Source/Runtime/Core/Public/Templates/RefCounting.h
+++ b/Engine/Source/Runtime/Core/Public/Templates/RefCounting.h
@@ -138,6 +138,11 @@ public:
 		return Reference;
 	}
 
+	FORCEINLINE friend uint32 GetTypeHash(const TRefCountPtr& InPtr)
+	{
+		return GetTypeHash(InPtr.Reference);
+	}
+
 	FORCEINLINE ReferencedType** GetInitReference()
 	{
 		*this = nullptr;
diff --git a/Engine/Source/Runtime/Core/Public/Templates/Tuple.h b/Engine/Source/Runtime/Core/Public/Templates/Tuple.h
--- a/Engine/Source/Runtime/Core/Public/Templates/Tuple.h
+++ b/Engine/Source/Runtime/Core/Public/Templates/Tuple.h
@@ -172,6 +172,36 @@ namespace UE4Tuple_Private
 		}
 	};
 
+	template <uint32 NumArgs, uint32 ArgToCompare = 0, bool Last = ArgToCompare + 1 == NumArgs>
+	struct TLessThanHelper
+	{
+		template <typename TupleType>
+		FORCEINLINE static bool Do(const TupleType& Lhs, const TupleType& Rhs)
+		{
+			return Lhs.template Get<ArgToCompare>() < Rhs.template Get<ArgToCompare>() || (!(Rhs.template Get<ArgToCompare>() < Lhs.template Get<ArgToCompare>()) && TLessThanHelper<NumArgs, ArgToCompare + 1>::Do(Lhs, Rhs));
+		}
+	};
+
+	template <uint32 NumArgs, uint32 ArgToCompare>
+	struct TLessThanHelper<NumArgs, ArgToCompare, true>
+	{
+		template <typename TupleType>
+		FORCEINLINE static bool Do(const TupleType& Lhs, const TupleType& Rhs)
+		{
+			return Lhs.template Get<ArgToCompare>() < Rhs.template Get<ArgToCompare>();
+		}
+	};
+
+	template <uint32 NumArgs>
+	struct TLessThanHelper<NumArgs, NumArgs, false>
+	{
+		template <typename TupleType>
+		FORCEINLINE static bool Do(const TupleType& Lhs, const TupleType& Rhs)
+		{
+			return false;
+		}
+	};
+
 	template <typename Indices, typename... Types>
 	struct TTupleStorage;
 
@@ -427,6 +457,26 @@ namespace UE4Tuple_Private
 		{
 			return !(Lhs == Rhs);
 		}
+
+		FORCEINLINE friend bool operator<(const TTupleImpl& Lhs, const TTupleImpl& Rhs)
+		{
+			return TLessThanHelper<sizeof...(Types)>::Do(Lhs, Rhs);
+		}
+
+		FORCEINLINE friend bool operator<=(const TTupleImpl& Lhs, const TTupleImpl& Rhs)
+		{
+			return !(Rhs < Lhs);
+		}
+
+		FORCEINLINE friend bool operator>(const TTupleImpl& Lhs, const TTupleImpl& Rhs)
+		{
+			return Rhs < Lhs;
+		}
+
+		FORCEINLINE friend bool operator>=(const TTupleImpl& Lhs, const TTupleImpl& Rhs)
+		{
+			return !(Lhs < Rhs);
+		}
 	};
 
 	#ifdef _MSC_VER
@@ -478,6 +528,26 @@ namespace UE4Tuple_Private
 			{
 				return false;
 			}
+
+			FORCEINLINE friend bool operator<(const TTupleImpl& Lhs, const TTupleImpl& Rhs)
+			{
+				return false;
+			}
+
+			FORCEINLINE friend bool operator<=(const TTupleImpl& Lhs, const TTupleImpl& Rhs)
+			{
+				return true;
+			}
+
+			FORCEINLINE friend bool operator>(const TTupleImpl& Lhs, const TTupleImpl& Rhs)
+			{
+				return false;
+			}
+
+			FORCEINLINE friend bool operator>=(const TTupleImpl& Lhs, const TTupleImpl& Rhs)
+			{
+				return true;
+			}
 		};
 
 	#endif
diff --git a/Engine/Source/Runtime/Core/Public/Templates/TypeHash.h b/Engine/Source/Runtime/Core/Public/Templates/TypeHash.h
--- a/Engine/Source/Runtime/Core/Public/Templates/TypeHash.h
+++ b/Engine/Source/Runtime/Core/Public/Templates/TypeHash.h
@@ -4,6 +4,7 @@
 
 #include "CoreTypes.h"
 #include "Templates/IsEnum.h"
+#include "Templates/Tuple.h"
 #include "Misc/Crc.h"
 
 /**
@@ -127,3 +128,37 @@ FORCEINLINE  typename TEnableIf<TIsEnum<EnumType>::Value, uint32>::Type GetTypeH
 {
 	return GetTypeHash((__underlying_type(EnumType))E);
 }
+
+namespace UE4TypeHash_Private
+{
+	template <uint32 ArgToCombine, uint32 ArgCount>
+	struct TGetTupleHashHelper
+	{
+		template <typename TupleType>
+		FORCEINLINE static uint32 Do(uint32 Hash, const TupleType& Tuple)
+		{
+			return TGetTupleHashHelper<ArgToCombine + 1, ArgCount>::Do(HashCombine(Hash, GetTypeHash(Tuple.template Get<ArgToCombine>())), Tuple);
+		}
+	};
+
+	template <uint32 ArgIndex>
+	struct TGetTupleHashHelper<ArgIndex, ArgIndex>
+	{
+		template <typename TupleType>
+		FORCEINLINE static uint32 Do(uint32 Hash, const TupleType& Tuple)
+		{
+			return Hash;
+		}
+	};
+}
+
+template <typename... Types>
+FORCEINLINE uint32 GetTypeHash(const TTuple<Types...>& Tuple)
+{
+	return UE4TypeHash_Private::TGetTupleHashHelper<1u, sizeof...(Types)>::Do(GetTypeHash(Tuple.template Get<0>()), Tuple);
+}
+
+FORCEINLINE uint32 GetTypeHash(const TTuple<>& Tuple)
+{
+	return 0;
+}
diff --git a/Engine/Source/Runtime/CoreUObject/Private/Misc/PackageName.cpp b/Engine/Source/Runtime/CoreUObject/Private/Misc/PackageName.cpp
--- a/Engine/Source/Runtime/CoreUObject/Private/Misc/PackageName.cpp
+++ b/Engine/Source/Runtime/CoreUObject/Private/Misc/PackageName.cpp
@@ -117,6 +117,7 @@ struct FLongPackagePathsSingleton
 	FString EngineRootPath;
 	FString GameRootPath;
 	FString ScriptRootPath;
+	FString MemoryRootPath;
 	FString TempRootPath;
 	TArray<FString> MountPointRootPaths;
 
@@ -154,6 +155,7 @@ struct FLongPackagePathsSingleton
 		{
 			OutRoots.Add(ConfigRootPath);
 			OutRoots.Add(ScriptRootPath);
+			OutRoots.Add(MemoryRootPath);
 			OutRoots.Add(TempRootPath);
 		}
 	}
@@ -210,6 +212,7 @@ private:
 		EngineRootPath = TEXT("/Engine/");
 		GameRootPath   = TEXT("/Game/");
 		ScriptRootPath = TEXT("/Script/");
+		MemoryRootPath = TEXT("/Memory/");
 		TempRootPath   = TEXT("/Temp/");
 
 		EngineContentPath      = FPaths::EngineContentDir();
@@ -239,6 +242,7 @@ private:
 		ContentPathToRoot.Emplace(EngineRootPath, EngineShadersPathShort);
 		ContentPathToRoot.Emplace(GameRootPath,   GameContentPath);
 		ContentPathToRoot.Emplace(ScriptRootPath, GameScriptPath);
+		ContentPathToRoot.Emplace(ScriptRootPath, GameScriptPath);
 		ContentPathToRoot.Emplace(TempRootPath,   GameSavedPath);
 		ContentPathToRoot.Emplace(GameRootPath,   GameContentPathRebased);
 		ContentPathToRoot.Emplace(ScriptRootPath, GameScriptPathRebased);
@@ -676,6 +680,11 @@ bool FPackageName::DoesPackageExist(const FString& LongPackageName, const FGuid*
 		return false;
 	}
 
+	if (IsMemoryPackage(PackageName))
+	{
+		return false;
+	}
+
 	if ( !FPackageName::IsValidLongPackageName( PackageName, true, &Reason ) )
 	{
 		UE_LOG(LogPackageName, Error, TEXT( "DoesPackageExist: DoesPackageExist FAILED: '%s' is not a standard unreal filename or a long path name. Reason: %s"), *LongPackageName, *Reason.ToString() );
@@ -1104,6 +1113,12 @@ bool FPackageName::IsScriptPackage(const FString& InPackageName)
 	return InPackageName.StartsWith(FLongPackagePathsSingleton::Get().ScriptRootPath);
 }
 
+// Are we a package that resides in memory and not on disk
+bool FPackageName::IsMemoryPackage(const FString& InPackageName)
+{
+	return InPackageName.StartsWith(FLongPackagePathsSingleton::Get().MemoryRootPath);
+}
+
 bool FPackageName::IsLocalizedPackage(const FString& InPackageName)
 {
 	// Minimum valid package name length is "/A/L10N"
diff --git a/Engine/Source/Runtime/CoreUObject/Private/Serialization/ArchiveUObject.cpp b/Engine/Source/Runtime/CoreUObject/Private/Serialization/ArchiveUObject.cpp
--- a/Engine/Source/Runtime/CoreUObject/Private/Serialization/ArchiveUObject.cpp
+++ b/Engine/Source/Runtime/CoreUObject/Private/Serialization/ArchiveUObject.cpp
@@ -141,7 +141,7 @@ FArchive& FArchiveUObject::operator<<(FWeakObjectPtr& WeakObjectPtr)
 /**
  * Serialize the given UObject* as an FString
  */
-FArchive& FObjectAndNameAsStringProxyArchive::operator<<(class UObject*& Obj)
+FArchive& FObjectAndNameAsStringProxyArchive::operator<<(UObject*& Obj)
 {
 	if (IsLoading())
 	{
@@ -165,6 +165,12 @@ FArchive& FObjectAndNameAsStringProxyArchive::operator<<(class UObject*& Obj)
 	return *this;
 }
 
+FArchive& FObjectAndNameAsStringProxyArchive::operator<<(FWeakObjectPtr& Obj)
+{
+	Obj.Serialize(*this);
+	return *this;
+}
+
 #if WITH_EDITORONLY_DATA
 void FSerializedPropertyScope::PushEditorOnlyProperty()
 {
diff --git a/Engine/Source/Runtime/CoreUObject/Private/Serialization/AsyncLoading.cpp b/Engine/Source/Runtime/CoreUObject/Private/Serialization/AsyncLoading.cpp
--- a/Engine/Source/Runtime/CoreUObject/Private/Serialization/AsyncLoading.cpp
+++ b/Engine/Source/Runtime/CoreUObject/Private/Serialization/AsyncLoading.cpp
@@ -161,6 +161,14 @@ static FAutoConsoleVariableRef CVarEventDrivenLoaderEnabled(
 	ECVF_Default
 );
 
+int32 GMaxReadyRequestsToStallMB = 30;
+static FAutoConsoleVariableRef CVar_MaxReadyRequestsToStallMB(
+	TEXT("s.MaxReadyRequestsToStallMB"),
+	GMaxReadyRequestsToStallMB,
+	TEXT("Controls the maximum amount memory for unhandled IO requests before we stall the pak precacher to let the CPU catch up (in megabytes).")
+);
+
+
 uint32 FAsyncLoadingThread::AsyncLoadingThreadID = 0;
 
 static void IsTimeLimitExceededPrint(double InTickStartTime, double CurrentTime, double LastTestTime, float InTimeLimit, const TCHAR* InLastTypeOfWorkPerformed = nullptr, UObject* InLastObjectWorkWasPerformedOn = nullptr)
@@ -410,11 +418,6 @@ FAsyncPackage* FAsyncLoadingThread::FindExistingPackageAndAddCompletionCallback(
 
 void FAsyncLoadingThread::UpdateExistingPackagePriorities(FAsyncPackage* InPackage, TAsyncLoadPriority InNewPriority, IAssetRegistryInterface* InAssetRegistry)
 {
-	if (!GEventDrivenLoaderEnabledInCookedBuilds)
-	{
-		return;
-	}
-
 	check(!IsInGameThread() || !IsMultithreaded());
 	if (GEventDrivenLoaderEnabled)
 	{
@@ -1550,14 +1553,14 @@ struct FPrecacheCallbackHandler
 
 	void CheckThottleIOState()
 	{
-		if (GPakCache_AcceptPrecacheRequests && UnprocessedMemUsed > 30 * 1024 * 1024)
+		if (GPakCache_AcceptPrecacheRequests && UnprocessedMemUsed > GMaxReadyRequestsToStallMB * 1024 * 1024)
 		{
-			UE_LOG(LogStreaming, Log, TEXT("Throttling off"));
+			UE_LOG(LogStreaming, Log, TEXT("Throttling off pak precacher to save memory while CPU catches up."));
 			GPakCache_AcceptPrecacheRequests = false;
 		}
-		else if (!GPakCache_AcceptPrecacheRequests && UnprocessedMemUsed <= 30 * 1024 * 1024)
+		else if (!GPakCache_AcceptPrecacheRequests && UnprocessedMemUsed <= GMaxReadyRequestsToStallMB * 1024 * 1024)
 		{
-			UE_LOG(LogStreaming, Log, TEXT("Throttling on"));
+			UE_LOG(LogStreaming, Log, TEXT("Resuming pak precacher."));
 			GPakCache_AcceptPrecacheRequests = true;
 		}
 	}
@@ -2329,8 +2332,11 @@ EAsyncPackageState::Type FAsyncPackage::SetupImports_Event()
 					// this is a hack because the fake export table is missing lots
 					if (bDynamicSomethingMissingFromTheFakeExportTable)
 					{
-						check(ImportLinker->ExportMap.Num() == 2); // we assume there are two elements in the fake export table and the second one is the CDO
-						LocalExportIndex = FPackageIndex::FromExport(1);
+						check(ImportLinker->ExportMap.Num() == 1 || ImportLinker->ExportMap.Num() == 2);
+						// we assume there are two elements in the fake export table and the second one is the CDO
+						// or there is just a struct without any CDO
+						const int32 DynamicExportIndex = (ImportLinker->ExportMap.Num() == 2) ? 1 : 0;
+						LocalExportIndex = FPackageIndex::FromExport(DynamicExportIndex);
 					}
 
 					Import.bImportFailed = LocalExportIndex.IsNull();
