Making sure the engine is properly cleaned up on UHT (early) exit to prevent UHT hangs.

diff --git a/Engine/Source/Programs/UnrealHeaderTool/Private/UnrealHeaderToolMain.cpp b/Engine/Source/Programs/UnrealHeaderTool/Private/UnrealHeaderToolMain.cpp
--- a/Engine/Source/Programs/UnrealHeaderTool/Private/UnrealHeaderToolMain.cpp
+++ b/Engine/Source/Programs/UnrealHeaderTool/Private/UnrealHeaderToolMain.cpp
@@ -89,6 +89,16 @@ INT32_MAIN_INT32_ARGC_TCHAR_ARGV()
 		LogCompile.SetVerbosity(ELogVerbosity::Verbose);
 	}
 
+	// Make sure the engine is properly cleaned up whenever we exit this function
+	struct FExitCleanup
+	{
+		~FExitCleanup()
+		{
+			FEngineLoop::AppPreExit();
+			FEngineLoop::AppExit();
+		}
+	} OnExitEngineCleanup;
+
 	GIsUCCMakeStandaloneHeaderGenerator = true;
 	if (GEngineLoop.PreInit(*ShortCmdLine) != 0)
 	{
@@ -115,9 +125,6 @@ INT32_MAIN_INT32_ARGC_TCHAR_ARGV()
 	extern ECompilationResult::Type UnrealHeaderTool_Main(const FString& ModuleInfoFilename);
 	ECompilationResult::Type Result = UnrealHeaderTool_Main(ModuleInfoFilename);
 
-	FEngineLoop::AppPreExit();
-	FEngineLoop::AppExit();
-
 	if (Result == ECompilationResult::Succeeded && GUHTWarningLogged && GWarn->TreatWarningsAsErrors)
 	{
 		Result = ECompilationResult::OtherCompilationError;
