Copying //UE4/Dev-Framework to //UE4/Dev-Main (Source: //UE4/Dev-Framework @ 3431384)

diff --git a/Engine/Source/Editor/UnrealEd/Private/Kismet2/StructureEditorUtils.cpp b/Engine/Source/Editor/UnrealEd/Private/Kismet2/StructureEditorUtils.cpp
--- a/Engine/Source/Editor/UnrealEd/Private/Kismet2/StructureEditorUtils.cpp
+++ b/Engine/Source/Editor/UnrealEd/Private/Kismet2/StructureEditorUtils.cpp
@@ -49,7 +49,7 @@ UUserDefinedStruct* FStructureEditorUtils::CreateUserDefinedStruct(UObject* InPa
 
 		{
 			const UEdGraphSchema_K2* K2Schema = GetDefault<UEdGraphSchema_K2>();
-			AddVariable(Struct, FEdGraphPinType(K2Schema->PC_Boolean, FString(), NULL, false, false, false, false, FEdGraphTerminalType()));
+			AddVariable(Struct, FEdGraphPinType(K2Schema->PC_Boolean, FString(), nullptr, EPinContainerType::None, false, FEdGraphTerminalType()));
 		}
 	}
 
@@ -209,7 +209,7 @@ bool FStructureEditorUtils::CanHaveAMemberVariableOfType(const UUserDefinedStruc
 	}
 	else
 	{
-		const auto PinSubCategoryClass = Cast<const UClass>(VarType.PinSubCategoryObject.Get());
+		const UClass* PinSubCategoryClass = Cast<const UClass>(VarType.PinSubCategoryObject.Get());
 		if (PinSubCategoryClass && PinSubCategoryClass->IsChildOf(UBlueprint::StaticClass()))
 		{
 			if (OutMsg)
@@ -317,7 +317,7 @@ bool FStructureEditorUtils::RemoveVariable(UUserDefinedStruct* Struct, FGuid Var
 {
 	if(Struct)
 	{
-		const auto OldNum = GetVarDesc(Struct).Num();
+		const int32 OldNum = GetVarDesc(Struct).Num();
 		const bool bAllowToMakeEmpty = false;
 		if (bAllowToMakeEmpty || (OldNum > 1))
 		{
@@ -343,7 +343,7 @@ bool FStructureEditorUtils::RenameVariable(UUserDefinedStruct* Struct, FGuid Var
 {
 	if (Struct)
 	{
-		auto VarDesc = GetVarDescByGuid(Struct, VarGuid);
+		FStructVariableDescription* VarDesc = GetVarDescByGuid(Struct, VarGuid);
 		if (VarDesc 
 			&& !NewDisplayNameStr.IsEmpty()
 			&& IsUniqueVariableDisplayName(Struct, NewDisplayNameStr))
@@ -379,7 +379,7 @@ bool FStructureEditorUtils::ChangeVariableType(UUserDefinedStruct* Struct, FGuid
 			return false;
 		}
 
-		auto VarDesc = GetVarDescByGuid(Struct, VarGuid);
+		FStructVariableDescription* VarDesc = GetVarDescByGuid(Struct, VarGuid);
 		if(VarDesc)
 		{
 			const bool bChangedType = (VarDesc->ToPinType() != NewType);
@@ -429,7 +429,7 @@ bool FStructureEditorUtils::ChangeVariableDefaultValue(UUserDefinedStruct* Struc
 		return bResult;
 	};
 
-	auto VarDesc = GetVarDescByGuid(Struct, VarGuid);
+	FStructVariableDescription* VarDesc = GetVarDescByGuid(Struct, VarGuid);
 	if (VarDesc 
 		&& (NewDefaultValue != VarDesc->DefaultValue)
 		&& ValidateDefaultValue(*VarDesc, NewDefaultValue))
@@ -437,7 +437,7 @@ bool FStructureEditorUtils::ChangeVariableDefaultValue(UUserDefinedStruct* Struc
 		bool bAdvancedValidation = true;
 		if (!NewDefaultValue.IsEmpty())
 		{
-			const auto Property = FindField<UProperty>(Struct, VarDesc->VarName);
+			const UProperty* Property = FindField<UProperty>(Struct, VarDesc->VarName);
 			FStructOnScope StructDefaultMem(Struct);
 			bAdvancedValidation = StructDefaultMem.IsValid() && Property &&
 				FBlueprintEditorUtils::PropertyValueFromString(Property, NewDefaultValue, StructDefaultMem.GetStructMemory());
@@ -463,7 +463,7 @@ bool FStructureEditorUtils::IsUniqueVariableDisplayName(const UUserDefinedStruct
 {
 	if(Struct)
 	{
-		for (auto& VarDesc : GetVarDesc(Struct))
+		for (const FStructVariableDescription& VarDesc : GetVarDesc(Struct))
 		{
 			if (VarDesc.FriendlyName == DisplayName)
 			{
@@ -477,7 +477,7 @@ bool FStructureEditorUtils::IsUniqueVariableDisplayName(const UUserDefinedStruct
 
 FString FStructureEditorUtils::GetVariableDisplayName(const UUserDefinedStruct* Struct, FGuid VarGuid)
 {
-	const auto VarDesc = GetVarDescByGuid(Struct, VarGuid);
+	const FStructVariableDescription* VarDesc = GetVarDescByGuid(Struct, VarGuid);
 	return VarDesc ? VarDesc->FriendlyName : FString();
 }
 
@@ -489,7 +489,7 @@ bool FStructureEditorUtils::UserDefinedStructEnabled()
 
 void FStructureEditorUtils::RecreateDefaultInstanceInEditorData(UUserDefinedStruct* Struct)
 {
-	auto StructEditorData = Struct ? CastChecked<UUserDefinedStructEditorData>(Struct->EditorData) : nullptr;
+	UUserDefinedStructEditorData* StructEditorData = Struct ? CastChecked<UUserDefinedStructEditorData>(Struct->EditorData) : nullptr;
 	if (StructEditorData)
 	{
 		StructEditorData->RecreateDefaultInstance();
@@ -501,7 +501,7 @@ bool FStructureEditorUtils::Fill_MakeStructureDefaultValue(const UUserDefinedStr
 	bool bResult = true;
 	if (Struct && StructData)
 	{
-		auto StructEditorData = CastChecked<UUserDefinedStructEditorData>(Struct->EditorData);
+		UUserDefinedStructEditorData* StructEditorData = CastChecked<UUserDefinedStructEditorData>(Struct->EditorData);
 		const uint8* DefaultInstance = StructEditorData->GetDefaultInstance();
 		if (DefaultInstance)
 		{
@@ -611,7 +611,7 @@ void FStructureEditorUtils::RemoveInvalidStructureMemberVariableFromBlueprint(UB
 
 		if (ZombieMemberNames.Num())
 		{
-			auto Response = FMessageDialog::Open( 
+			EAppReturnType::Type Response = FMessageDialog::Open( 
 				EAppMsgType::OkCancel,
 				FText::Format(
 					LOCTEXT("RemoveInvalidStructureMemberVariable_Msg", "The following member variables in blueprint '{0}' have invalid type. Would you like to remove them? \n\n{1}"), 
@@ -624,9 +624,8 @@ void FStructureEditorUtils::RemoveInvalidStructureMemberVariableFromBlueprint(UB
 			{				
 				Blueprint->Modify();
 
-				for (auto NameIter = ZombieMemberNames.CreateConstIterator(); NameIter; ++NameIter)
+				for (const FName Name : ZombieMemberNames)
 				{
-					const FName Name = *NameIter;
 					Blueprint->NewVariables.RemoveAll(FFindByNameHelper<FBPVariableDescription>(Name)); //TODO: Add RemoveFirst to TArray
 					FBlueprintEditorUtils::RemoveVariableNodes(Blueprint, Name);
 				}
@@ -661,13 +660,13 @@ const TArray<FStructVariableDescription>* FStructureEditorUtils::GetVarDescPtr(c
 
 FString FStructureEditorUtils::GetTooltip(const UUserDefinedStruct* Struct)
 {
-	const auto StructEditorData = Struct ? Cast<const UUserDefinedStructEditorData>(Struct->EditorData) : NULL;
+	const UUserDefinedStructEditorData* StructEditorData = Struct ? Cast<const UUserDefinedStructEditorData>(Struct->EditorData) : NULL;
 	return StructEditorData ? StructEditorData->ToolTip : FString();
 }
 
 bool FStructureEditorUtils::ChangeTooltip(UUserDefinedStruct* Struct, const FString& InTooltip)
 {
-	auto StructEditorData = Struct ? Cast<UUserDefinedStructEditorData>(Struct->EditorData) : NULL;
+	UUserDefinedStructEditorData* StructEditorData = Struct ? Cast<UUserDefinedStructEditorData>(Struct->EditorData) : NULL;
 	if (StructEditorData && (InTooltip != StructEditorData->ToolTip))
 	{
 		const FScopedTransaction Transaction(LOCTEXT("ChangeTooltip", "Change UDS Tooltip"));
@@ -683,20 +682,20 @@ bool FStructureEditorUtils::ChangeTooltip(UUserDefinedStruct* Struct, const FStr
 
 FString FStructureEditorUtils::GetVariableTooltip(const UUserDefinedStruct* Struct, FGuid VarGuid)
 {
-	const auto VarDesc = GetVarDescByGuid(Struct, VarGuid);
+	const FStructVariableDescription* VarDesc = GetVarDescByGuid(Struct, VarGuid);
 	return VarDesc ? VarDesc->ToolTip : FString();
 }
 
 bool FStructureEditorUtils::ChangeVariableTooltip(UUserDefinedStruct* Struct, FGuid VarGuid, const FString& InTooltip)
 {
-	auto VarDesc = GetVarDescByGuid(Struct, VarGuid);
+	FStructVariableDescription* VarDesc = GetVarDescByGuid(Struct, VarGuid);
 	if (VarDesc && (InTooltip != VarDesc->ToolTip))
 	{
 		const FScopedTransaction Transaction(LOCTEXT("ChangeVariableTooltip", "Change UDS Variable Tooltip"));
 		ModifyStructData(Struct);
 		VarDesc->ToolTip = InTooltip;
 
-		auto Property = FindField<UProperty>(Struct, VarDesc->VarName);
+		UProperty* Property = FindField<UProperty>(Struct, VarDesc->VarName);
 		if (Property)
 		{
 			Property->SetMetaData(FBlueprintMetadata::MD_Tooltip, *VarDesc->ToolTip);
@@ -710,7 +709,7 @@ bool FStructureEditorUtils::ChangeVariableTooltip(UUserDefinedStruct* Struct, FG
 bool FStructureEditorUtils::ChangeEditableOnBPInstance(UUserDefinedStruct* Struct, FGuid VarGuid, bool bInIsEditable)
 {
 	const UEdGraphSchema_K2* K2Schema = GetDefault<UEdGraphSchema_K2>();
-	auto VarDesc = GetVarDescByGuid(Struct, VarGuid);
+	FStructVariableDescription* VarDesc = GetVarDescByGuid(Struct, VarGuid);
 	const bool bNewDontEditoOnInstance = !bInIsEditable;
 	if (VarDesc && (bNewDontEditoOnInstance != VarDesc->bDontEditoOnInstance))
 	{
@@ -729,7 +728,7 @@ bool FStructureEditorUtils::MoveVariable(UUserDefinedStruct* Struct, FGuid VarGu
 	if (Struct)
 	{
 		const bool bMoveUp = (EMoveDirection::MD_Up == MoveDirection);
-		auto& DescArray = GetVarDesc(Struct);
+		TArray<FStructVariableDescription>& DescArray = GetVarDesc(Struct);
 		const int32 InitialIndex = bMoveUp ? 1 : 0;
 		const int32 IndexLimit = DescArray.Num() - (bMoveUp ? 0 : 1);
 		for (int32 Index = InitialIndex; Index < IndexLimit; ++Index)
@@ -760,7 +759,7 @@ void FStructureEditorUtils::ModifyStructData(UUserDefinedStruct* Struct)
 
 bool FStructureEditorUtils::CanEnableMultiLineText(const UUserDefinedStruct* Struct, FGuid VarGuid)
 {
-	auto VarDesc = GetVarDescByGuid(Struct, VarGuid);
+	const FStructVariableDescription* VarDesc = GetVarDescByGuid(Struct, VarGuid);
 	if (VarDesc)
 	{
 		UProperty* Property = FindField<UProperty>(Struct, VarDesc->VarName);
@@ -783,14 +782,14 @@ bool FStructureEditorUtils::CanEnableMultiLineText(const UUserDefinedStruct* Str
 
 bool FStructureEditorUtils::ChangeMultiLineTextEnabled(UUserDefinedStruct* Struct, FGuid VarGuid, bool bIsEnabled)
 {
-	auto VarDesc = GetVarDescByGuid(Struct, VarGuid);
+	FStructVariableDescription* VarDesc = GetVarDescByGuid(Struct, VarGuid);
 	if (CanEnableMultiLineText(Struct, VarGuid) && VarDesc->bEnableMultiLineText != bIsEnabled)
 	{
 		const FScopedTransaction Transaction(LOCTEXT("ChangeMultiLineTextEnabled", "Change Multi-line Text Enabled"));
 		ModifyStructData(Struct);
 
 		VarDesc->bEnableMultiLineText = bIsEnabled;
-		auto Property = FindField<UProperty>(Struct, VarDesc->VarName);
+		UProperty* Property = FindField<UProperty>(Struct, VarDesc->VarName);
 		if (Property)
 		{
 			if (VarDesc->bEnableMultiLineText)
@@ -810,7 +809,7 @@ bool FStructureEditorUtils::ChangeMultiLineTextEnabled(UUserDefinedStruct* Struc
 
 bool FStructureEditorUtils::IsMultiLineTextEnabled(const UUserDefinedStruct* Struct, FGuid VarGuid)
 {
-	auto VarDesc = GetVarDescByGuid(Struct, VarGuid);
+	const FStructVariableDescription* VarDesc = GetVarDescByGuid(Struct, VarGuid);
 	if (CanEnableMultiLineText(Struct, VarGuid))
 	{
 		return VarDesc->bEnableMultiLineText;
@@ -820,8 +819,8 @@ bool FStructureEditorUtils::IsMultiLineTextEnabled(const UUserDefinedStruct* Str
 
 bool FStructureEditorUtils::CanEnable3dWidget(const UUserDefinedStruct* Struct, FGuid VarGuid)
 {
-	const auto VarDesc = GetVarDescByGuid(Struct, VarGuid);
-	const auto PropertyStruct = VarDesc ? Cast<const UStruct>(VarDesc->SubCategoryObject.Get()) : NULL;
+	const FStructVariableDescription* VarDesc = GetVarDescByGuid(Struct, VarGuid);
+	const UStruct* PropertyStruct = VarDesc ? Cast<const UStruct>(VarDesc->SubCategoryObject.Get()) : NULL;
 	return FEdMode::CanCreateWidgetForStructure(PropertyStruct);
 }
 
@@ -859,22 +858,22 @@ bool FStructureEditorUtils::Change3dWidgetEnabled(UUserDefinedStruct* Struct, FG
 
 bool FStructureEditorUtils::Is3dWidgetEnabled(const UUserDefinedStruct* Struct, FGuid VarGuid)
 {
-	const auto VarDesc = GetVarDescByGuid(Struct, VarGuid);
-	const auto PropertyStruct = VarDesc ? Cast<const UStruct>(VarDesc->SubCategoryObject.Get()) : NULL;
+	const FStructVariableDescription* VarDesc = GetVarDescByGuid(Struct, VarGuid);
+	const UStruct* PropertyStruct = VarDesc ? Cast<const UStruct>(VarDesc->SubCategoryObject.Get()) : nullptr;
 	return VarDesc && VarDesc->bEnable3dWidget && FEdMode::CanCreateWidgetForStructure(PropertyStruct);
 }
 
 FGuid FStructureEditorUtils::GetGuidForProperty(const UProperty* Property)
 {
-	auto UDStruct = Property ? Cast<const UUserDefinedStruct>(Property->GetOwnerStruct()) : NULL;
-	auto VarDesc = UDStruct ? GetVarDesc(UDStruct).FindByPredicate(FFindByNameHelper<FStructVariableDescription>(Property->GetFName())) : NULL;
+	const UUserDefinedStruct* UDStruct = Property ? Cast<const UUserDefinedStruct>(Property->GetOwnerStruct()) : nullptr;
+	const FStructVariableDescription* VarDesc = UDStruct ? GetVarDesc(UDStruct).FindByPredicate(FFindByNameHelper<FStructVariableDescription>(Property->GetFName())) : nullptr;
 	return VarDesc ? VarDesc->VarGuid : FGuid();
 }
 
 UProperty* FStructureEditorUtils::GetPropertyByGuid(const UUserDefinedStruct* Struct, const FGuid VarGuid)
 {
-	const auto VarDesc = GetVarDescByGuid(Struct, VarGuid);
-	return VarDesc ? FindField<UProperty>(Struct, VarDesc->VarName) : NULL;
+	const FStructVariableDescription* VarDesc = GetVarDescByGuid(Struct, VarGuid);
+	return VarDesc ? FindField<UProperty>(Struct, VarDesc->VarName) : nullptr;
 }
 
 FGuid FStructureEditorUtils::GetGuidFromPropertyName(const FName Name)
@@ -892,9 +891,9 @@ struct FReinstanceDataTableHelper
 		{
 			TArray<UObject*> DataTables;
 			GetObjectsOfClass(UDataTable::StaticClass(), DataTables);
-			for (auto DataTableObj : DataTables)
+			for (UObject* DataTableObj : DataTables)
 			{
-				auto DataTable = Cast<UDataTable>(DataTableObj);
+				UDataTable* DataTable = Cast<UDataTable>(DataTableObj);
 				if (DataTable && (Struct == DataTable->RowStruct))
 				{
 					Result.Add(DataTable);
@@ -908,8 +907,8 @@ struct FReinstanceDataTableHelper
 void FStructureEditorUtils::BroadcastPreChange(UUserDefinedStruct* Struct)
 {
 	FStructureEditorUtils::FStructEditorManager::Get().PreChange(Struct, FStructureEditorUtils::FStructEditorManager::ActiveChange);
-	auto DataTables = FReinstanceDataTableHelper::GetTablesDependentOnStruct(Struct);
-	for (auto DataTable : DataTables)
+	TArray<UDataTable*> DataTables = FReinstanceDataTableHelper::GetTablesDependentOnStruct(Struct);
+	for (UDataTable* DataTable : DataTables)
 	{
 		DataTable->CleanBeforeStructChange();
 	}
@@ -917,8 +916,8 @@ void FStructureEditorUtils::BroadcastPreChange(UUserDefinedStruct* Struct)
 
 void FStructureEditorUtils::BroadcastPostChange(UUserDefinedStruct* Struct)
 {
-	auto DataTables = FReinstanceDataTableHelper::GetTablesDependentOnStruct(Struct);
-	for (auto DataTable : DataTables)
+	TArray<UDataTable*> DataTables = FReinstanceDataTableHelper::GetTablesDependentOnStruct(Struct);
+	for (UDataTable* DataTable : DataTables)
 	{
 		DataTable->RestoreAfterStructChange();
 	}
diff --git a/Engine/Source/Editor/UnrealEd/Private/Lightmass/Lightmass.h b/Engine/Source/Editor/UnrealEd/Private/Lightmass/Lightmass.h
--- a/Engine/Source/Editor/UnrealEd/Private/Lightmass/Lightmass.h
+++ b/Engine/Source/Editor/UnrealEd/Private/Lightmass/Lightmass.h
@@ -109,7 +109,7 @@ public:
 
 	void AddPortal(const ULightmassPortalComponent* InPortalComponent)
 	{
-		Portals.Add(InPortalComponent->ComponentToWorld.ToMatrixWithScale());
+		Portals.Add(InPortalComponent->GetComponentTransform().ToMatrixWithScale());
 	}
 
 	// if provided, InStaticLightingMesh is used to UV unwrap the material into the static lighting textures
diff --git a/Engine/Source/Editor/UnrealEd/Private/MaterialGraphNode.cpp b/Engine/Source/Editor/UnrealEd/Private/MaterialGraphNode.cpp
--- a/Engine/Source/Editor/UnrealEd/Private/MaterialGraphNode.cpp
+++ b/Engine/Source/Editor/UnrealEd/Private/MaterialGraphNode.cpp
@@ -485,7 +485,7 @@ void UMaterialGraphNode::CreateInputPins()
 		FString PinCategory = MaterialExpression->IsInputConnectionRequired(Index) ? Schema->PC_Required : Schema->PC_Optional;
 		FString PinSubCategory = TEXT("");
 
-		UEdGraphPin* NewPin = CreatePin(EGPD_Input, PinCategory, PinSubCategory, NULL, /*bIsArray=*/ false, /*bIsReference=*/ false, InputName);
+		UEdGraphPin* NewPin = CreatePin(EGPD_Input, PinCategory, PinSubCategory, nullptr, InputName);
 		if (NewPin->PinName.IsEmpty())
 		{
 			// Makes sure pin has a name for lookup purposes but user will never see it
@@ -538,7 +538,7 @@ void UMaterialGraphNode::CreateOutputPins()
 			OutputName = ExpressionOutput.OutputName;
 		}
 
-		UEdGraphPin* NewPin = CreatePin(EGPD_Output, PinCategory, PinSubCategory, NULL, /*bIsArray=*/ false, /*bIsReference=*/ false, OutputName);
+		UEdGraphPin* NewPin = CreatePin(EGPD_Output, PinCategory, PinSubCategory, nullptr, OutputName);
 		if (NewPin->PinName.IsEmpty())
 		{
 			// Makes sure pin has a name for lookup purposes but user will never see it
diff --git a/Engine/Source/Editor/UnrealEd/Private/MaterialGraphNode_Knot.cpp b/Engine/Source/Editor/UnrealEd/Private/MaterialGraphNode_Knot.cpp
--- a/Engine/Source/Editor/UnrealEd/Private/MaterialGraphNode_Knot.cpp
+++ b/Engine/Source/Editor/UnrealEd/Private/MaterialGraphNode_Knot.cpp
@@ -25,10 +25,10 @@ void UMaterialGraphNode_Knot::AllocateDefaultPins()
 	const FString InputPinName(TEXT("InputPin"));
 	const FString OutputPinName(TEXT("OutputPin"));
 
-	UEdGraphPin* MyInputPin = CreatePin(EGPD_Input, PC_Wildcard, FString(), nullptr, /*bIsArray=*/ false, /*bIsReference=*/ false, InputPinName);
+	UEdGraphPin* MyInputPin = CreatePin(EGPD_Input, PC_Wildcard, FString(), nullptr, InputPinName);
 	MyInputPin->bDefaultValueIsIgnored = true;
 
-	UEdGraphPin* MyOutputPin = CreatePin(EGPD_Output, PC_Wildcard, FString(), nullptr, /*bIsArray=*/ false, /*bIsReference=*/ false, OutputPinName);
+	UEdGraphPin* MyOutputPin = CreatePin(EGPD_Output, PC_Wildcard, FString(), nullptr, OutputPinName);
 }
 
 FText UMaterialGraphNode_Knot::GetTooltipText() const
diff --git a/Engine/Source/Editor/UnrealEd/Private/MaterialGraphNode_Root.cpp b/Engine/Source/Editor/UnrealEd/Private/MaterialGraphNode_Root.cpp
--- a/Engine/Source/Editor/UnrealEd/Private/MaterialGraphNode_Root.cpp
+++ b/Engine/Source/Editor/UnrealEd/Private/MaterialGraphNode_Root.cpp
@@ -50,9 +50,9 @@ void UMaterialGraphNode_Root::CreateInputPins()
 	UMaterialGraph* MaterialGraph = CastChecked<UMaterialGraph>(GetGraph());
 	const UMaterialGraphSchema* Schema = CastChecked<UMaterialGraphSchema>(GetSchema());
 
-	for (int32 Index = 0; Index < MaterialGraph->MaterialInputs.Num(); ++Index)
+	for (const FMaterialInputInfo& MaterialInput : MaterialGraph->MaterialInputs)
 	{
-		UEdGraphPin* InputPin = CreatePin(EGPD_Input, Schema->PC_MaterialInput, FString::Printf(TEXT("%d"), (int32)MaterialGraph->MaterialInputs[Index].GetProperty()), NULL, /*bIsArray=*/ false, /*bIsReference=*/ false, MaterialGraph->MaterialInputs[Index].GetName().ToString());
+		UEdGraphPin* InputPin = CreatePin(EGPD_Input, Schema->PC_MaterialInput, FString::Printf(TEXT("%d"), (int32)MaterialInput.GetProperty()), nullptr, MaterialInput.GetName().ToString());
 	}
 
 }
diff --git a/Engine/Source/Editor/UnrealEd/Private/ObjectTools.cpp b/Engine/Source/Editor/UnrealEd/Private/ObjectTools.cpp
--- a/Engine/Source/Editor/UnrealEd/Private/ObjectTools.cpp
+++ b/Engine/Source/Editor/UnrealEd/Private/ObjectTools.cpp
@@ -2926,7 +2926,7 @@ namespace ObjectTools
 		{
 			if(bMoveRedirectorFailed)
 			{
-				ErrorMessage += FText::Format( NSLOCTEXT("UnrealEd", "Error_CouldntRenameObjectRedirectorF", "Couldn't rename '{0}' object because there is an object redirector of the same name, please run FixupRedirects.\n"),
+				ErrorMessage += FText::Format( NSLOCTEXT("UnrealEd", "Error_CouldntRenameObjectRedirectorF", "Couldn't rename '{0}' object because there is an object redirector of the same name, please fixup redirect from editor by enabling Show Redirects in content browser.\n"),
 					FText::FromString(Object->GetFullName()) ).ToString();
 			}
 			else
diff --git a/Engine/Source/Editor/UnrealEd/Private/Settings/SettingsClasses.cpp b/Engine/Source/Editor/UnrealEd/Private/Settings/SettingsClasses.cpp
--- a/Engine/Source/Editor/UnrealEd/Private/Settings/SettingsClasses.cpp
+++ b/Engine/Source/Editor/UnrealEd/Private/Settings/SettingsClasses.cpp
@@ -33,6 +33,8 @@
 #include "Misc/ConfigCacheIni.h" // for FConfigCacheIni::GetString()
 #include "SourceCodeNavigation.h"
 #include "Developer/BlueprintProfiler/Public/BlueprintProfilerModule.h"
+#include "IProjectManager.h"
+#include "ProjectDescriptor.h"
 
 #define LOCTEXT_NAMESPACE "SettingsClasses"
 
@@ -649,6 +651,36 @@ void UProjectPackagingSettings::PostEditChangeProperty( FPropertyChangedEvent& P
 			}
 		}
 	}
+	else if (Name == FName((TEXT("BlueprintNativizationMethod"))))
+	{
+		IProjectManager& ProjManager = IProjectManager::Get();
+		{
+			// NOTE: these are hardcoded to match the path constructed by AddBlueprintPluginPathArgument() on CookCommand.Automation.cs, and the defaults in 
+			//       FBlueprintNativeCodeGenPaths::GetDefaultCodeGenPaths(); if you alter this (or either of those) then you need to update the others
+			const FString NativizedPluginDir  = TEXT("./Intermediate/Plugins");
+			const FString NativizedPluginName = TEXT("NativizedAssets");
+			const FString FullPluginPath = FPaths::ConvertRelativePathToFull( FPaths::ConvertRelativePathToFull(FPaths::GameDir()), NativizedPluginDir );
+
+			if (BlueprintNativizationMethod == EProjectPackagingBlueprintNativizationMethod::Disabled)
+			{
+				
+				ProjManager.UpdateAdditionalPluginDirectory(FullPluginPath, /*bAddOrRemove =*/false);
+				FText PluginDisableFailure;
+				ProjManager.SetPluginEnabled(NativizedPluginName, /*bEnabled =*/false, PluginDisableFailure);
+			}
+			else
+			{
+				ProjManager.UpdateAdditionalPluginDirectory(FullPluginPath, /*bAddOrRemove =*/true);
+				// plugin is enabled by default, so let's remove it from the uproject list entirely (else it causes problem in the packaged project)
+				// SetPluginEnabled() will only remove it if the plugin exists (it may not yet), so we rely on this explicit removal
+				FText PluginEnableFailure;
+				ProjManager.RemovePluginReference(NativizedPluginName, PluginEnableFailure);
+			}
+
+			FText SaveFailure;
+			ProjManager.SaveCurrentProjectToDisk(SaveFailure);
+		}
+	}
 	else if (Name == FName((TEXT("NativizeBlueprintAssets"))))
 	{
 		int32 AssetIndex;
diff --git a/Engine/Source/Editor/UnrealEd/Private/StaticLightingSystem/StaticLightingDebug.cpp b/Engine/Source/Editor/UnrealEd/Private/StaticLightingSystem/StaticLightingDebug.cpp
--- a/Engine/Source/Editor/UnrealEd/Private/StaticLightingSystem/StaticLightingDebug.cpp
+++ b/Engine/Source/Editor/UnrealEd/Private/StaticLightingSystem/StaticLightingDebug.cpp
@@ -226,9 +226,9 @@ void SetDebugLightmapSample(TArray<UActorComponent*>* Components, UModel* Model,
 						uint32 Index2 = Indices[TriangleIndex + 2];
 
 						// Transform positions to world space
-						FVector Position0 = SMComponent->ComponentToWorld.TransformPosition(LODModel.PositionVertexBuffer.VertexPosition(Index0));
-						FVector Position1 = SMComponent->ComponentToWorld.TransformPosition(LODModel.PositionVertexBuffer.VertexPosition(Index1));
-						FVector Position2 = SMComponent->ComponentToWorld.TransformPosition(LODModel.PositionVertexBuffer.VertexPosition(Index2));
+						FVector Position0 = SMComponent->GetComponentTransform().TransformPosition(LODModel.PositionVertexBuffer.VertexPosition(Index0));
+						FVector Position1 = SMComponent->GetComponentTransform().TransformPosition(LODModel.PositionVertexBuffer.VertexPosition(Index1));
+						FVector Position2 = SMComponent->GetComponentTransform().TransformPosition(LODModel.PositionVertexBuffer.VertexPosition(Index2));
 
 						float PlaneDistance;
 						FVector BaryCentricWeights;
diff --git a/Engine/Source/Editor/UnrealEd/Private/ThumbnailHelpers.cpp b/Engine/Source/Editor/UnrealEd/Private/ThumbnailHelpers.cpp
--- a/Engine/Source/Editor/UnrealEd/Private/ThumbnailHelpers.cpp
+++ b/Engine/Source/Editor/UnrealEd/Private/ThumbnailHelpers.cpp
@@ -179,7 +179,7 @@ void FParticleSystemThumbnailScene::SetParticleSystem(UParticleSystem* ParticleS
 		ParticleSystem->PreviewComponent = NewObject<UParticleSystemComponent>();
 		ParticleSystem->PreviewComponent->Template = ParticleSystem;
 
-		ParticleSystem->PreviewComponent->ComponentToWorld.SetIdentity();
+		ParticleSystem->PreviewComponent->SetComponentToWorld(FTransform::Identity);
 
 		bNewComponent = true;
 	}
diff --git a/Engine/Source/Editor/UnrealEd/Private/UnrealEdPrivatePCH.h b/Engine/Source/Editor/UnrealEd/Private/UnrealEdPrivatePCH.h
--- a/Engine/Source/Editor/UnrealEd/Private/UnrealEdPrivatePCH.h
+++ b/Engine/Source/Editor/UnrealEd/Private/UnrealEdPrivatePCH.h
@@ -462,7 +462,6 @@
 
 // From AssetRegistry:
 #include "AssetData.h"
-#include "SharedMapView.h"
 #include "IAssetRegistry.h"
 #include "AssetRegistryModule.h"
 #include "ARFilter.h"
diff --git a/Engine/Source/Editor/UnrealEd/Private/UnrealEdSrv.cpp b/Engine/Source/Editor/UnrealEd/Private/UnrealEdSrv.cpp
--- a/Engine/Source/Editor/UnrealEd/Private/UnrealEdSrv.cpp
+++ b/Engine/Source/Editor/UnrealEd/Private/UnrealEdSrv.cpp
@@ -1079,7 +1079,7 @@ bool UUnrealEdEngine::Exec( UWorld* InWorld, const TCHAR* Stream, FOutputDevice&
 					SlowTask.EnterProgressFrame();
 
 					// Optimization - check the asset has import information before loading it
-					TOptional<FAssetImportInfo> ImportInfo = FAssetSourceFilenameCache::ExtractAssetImportInfo(Asset.TagsAndValues);
+					TOptional<FAssetImportInfo> ImportInfo = FAssetSourceFilenameCache::ExtractAssetImportInfo(Asset);
 					if (ImportInfo.IsSet() && ImportInfo->SourceFiles.Num())
 					{
 						RemoveSourcePath(ImportInfo.GetValue(), Asset, SearchTerms);
diff --git a/Engine/Source/Editor/UnrealEd/Private/UserDefinedStructEditorData.cpp b/Engine/Source/Editor/UnrealEd/Private/UserDefinedStructEditorData.cpp
--- a/Engine/Source/Editor/UnrealEd/Private/UserDefinedStructEditorData.cpp
+++ b/Engine/Source/Editor/UnrealEd/Private/UserDefinedStructEditorData.cpp
@@ -10,22 +10,28 @@
 
 #define LOCTEXT_NAMESPACE "UserDefinedStructEditorData"
 
+void FStructVariableDescription::PostSerialize(const FArchive& Ar)
+{
+	if (ContainerType == EPinContainerType::None)
+	{
+		ContainerType = FEdGraphPinType::ToPinContainerType(bIsArray_DEPRECATED, bIsSet_DEPRECATED, bIsMap_DEPRECATED);
+	}
+}
+
 bool FStructVariableDescription::SetPinType(const FEdGraphPinType& VarType)
 {
 	Category = VarType.PinCategory;
 	SubCategory = VarType.PinSubCategory;
 	SubCategoryObject = VarType.PinSubCategoryObject.Get();
 	PinValueType = VarType.PinValueType;
-	bIsArray = VarType.bIsArray;
-	bIsSet = VarType.bIsSet;
-	bIsMap = VarType.bIsMap;
+	ContainerType = VarType.ContainerType;
 
 	return !VarType.bIsReference && !VarType.bIsWeakPointer;
 }
 
 FEdGraphPinType FStructVariableDescription::ToPinType() const
 {
-	return FEdGraphPinType(Category, SubCategory, SubCategoryObject.LoadSynchronous(), bIsArray, false, bIsSet, bIsMap, PinValueType);
+	return FEdGraphPinType(Category, SubCategory, SubCategoryObject.LoadSynchronous(), ContainerType, false, PinValueType);
 }
 
 UUserDefinedStructEditorData::UUserDefinedStructEditorData(const FObjectInitializer& ObjectInitializer)
diff --git a/Engine/Source/Editor/UnrealEd/Private/VertexSnapping.cpp b/Engine/Source/Editor/UnrealEd/Private/VertexSnapping.cpp
--- a/Engine/Source/Editor/UnrealEd/Private/VertexSnapping.cpp
+++ b/Engine/Source/Editor/UnrealEd/Private/VertexSnapping.cpp
@@ -76,7 +76,7 @@ class FStaticMeshVertexIterator : public FVertexIterator
 {
 public:
 	FStaticMeshVertexIterator( UStaticMeshComponent* SMC )
-		: ComponentToWorldIT( SMC->ComponentToWorld.ToInverseMatrixWithScale().GetTransposed() )
+		: ComponentToWorldIT( SMC->GetComponentTransform().ToInverseMatrixWithScale().GetTransposed() )
 		, StaticMeshComponent( SMC )
 		, PositionBuffer( SMC->GetStaticMesh()->RenderData->LODResources[0].PositionVertexBuffer )
 		, VertexBuffer( SMC->GetStaticMesh()->RenderData->LODResources[0].VertexBuffer )
@@ -88,7 +88,7 @@ public:
 	/** FVertexIterator interface */
 	virtual FVector Position() override
 	{
-		return StaticMeshComponent->ComponentToWorld.TransformPosition( PositionBuffer.VertexPosition( CurrentVertexIndex ) );	
+		return StaticMeshComponent->GetComponentTransform().TransformPosition( PositionBuffer.VertexPosition( CurrentVertexIndex ) );	
 	}
 
 	virtual FVector Normal() override
@@ -144,7 +144,7 @@ public:
 	/** FVertexIterator interface */
 	virtual FVector Position() override
 	{
-		return BrushComponent->ComponentToWorld.TransformPosition( Vertices[CurrentVertexIndex] );
+		return BrushComponent->GetComponentTransform().TransformPosition( Vertices[CurrentVertexIndex] );
 	}
 
 	/** FVertexIterator interface */
@@ -181,7 +181,7 @@ class FSkeletalMeshVertexIterator : public FVertexIterator
 {
 public:
 	FSkeletalMeshVertexIterator( USkinnedMeshComponent* InSkinnedMeshComp )
-		: ComponentToWorldIT( InSkinnedMeshComp->ComponentToWorld.ToInverseMatrixWithScale().GetTransposed() )
+		: ComponentToWorldIT( InSkinnedMeshComp->GetComponentTransform().ToInverseMatrixWithScale().GetTransposed() )
 		, SkinnedMeshComponent( InSkinnedMeshComp )
 		, LODModel( InSkinnedMeshComp->GetSkeletalMeshResource()->LODModels[0] )
 		, CurrentSectionIndex( 0 )
@@ -194,7 +194,7 @@ public:
 	virtual FVector Position() override
 	{
 		const FSkelMeshSection& Section = LODModel.Sections[CurrentSectionIndex];
-		return SkinnedMeshComponent->ComponentToWorld.TransformPosition(Section.SoftVertices[SoftVertexIndex].Position );
+		return SkinnedMeshComponent->GetComponentTransform().TransformPosition(Section.SoftVertices[SoftVertexIndex].Position );
 	}
 
 	virtual FVector Normal() override
diff --git a/Engine/Source/Editor/UnrealEd/Public/AutoReimport/AssetSourceFilenameCache.h b/Engine/Source/Editor/UnrealEd/Public/AutoReimport/AssetSourceFilenameCache.h
--- a/Engine/Source/Editor/UnrealEd/Public/AutoReimport/AssetSourceFilenameCache.h
+++ b/Engine/Source/Editor/UnrealEd/Public/AutoReimport/AssetSourceFilenameCache.h
@@ -9,7 +9,6 @@
 
 class IAssetRegistry;
 
-template <class FKey, class FValue> class TSharedMapView;
 
 /** Class responsible for maintaing a cache of clean source file names (bla.txt) to asset data */
 class UNREALED_API FAssetSourceFilenameCache
@@ -21,8 +20,7 @@ public:
 	static FAssetSourceFilenameCache& Get();
 
 	/** Helper functions to extract asset import information from asset registry tags */
-	static TOptional<FAssetImportInfo> ExtractAssetImportInfo(const TSharedMapView<FName, FString>& InTags);
-	static TOptional<FAssetImportInfo> ExtractAssetImportInfo(const TArray<UObject::FAssetRegistryTag>& InTags);
+	static TOptional<FAssetImportInfo> ExtractAssetImportInfo(const FAssetData& AssetData);
 
 	/** Retrieve a list of assets that were imported from the specified filename */
 	TArray<FAssetData> GetAssetsPertainingToFile(const IAssetRegistry& Registry, const FString& AbsoluteFilename) const;
diff --git a/Engine/Source/Editor/UnrealEd/Public/Commandlets/AssetRegistryGenerator.h b/Engine/Source/Editor/UnrealEd/Public/Commandlets/AssetRegistryGenerator.h
--- a/Engine/Source/Editor/UnrealEd/Public/Commandlets/AssetRegistryGenerator.h
+++ b/Engine/Source/Editor/UnrealEd/Public/Commandlets/AssetRegistryGenerator.h
@@ -51,8 +51,9 @@ public:
 	 * @param RemovedPackages list of packages that existed before, but do not any more
 	 * @param IdenticalCookedPackages list of cooked packages that have not changed
 	 * @param IdenticalUncookedPackages list of uncooked packages that have not changed. These were filtered out by platform or editor only
+	 * @param bRecurseModifications if true, modified packages are recursed to X in X->Y->Z chains. Otherwise, only Y and Z are seen as modified
 	 */
-	void ComputePackageDifferences(TSet<FName>& ModifiedPackages, TSet<FName>& NewPackages, TSet<FName>& RemovedPackages, TSet<FName>& IdenticalCookedPackages, TSet<FName>& IdenticalUncookedPackages);
+	void ComputePackageDifferences(TSet<FName>& ModifiedPackages, TSet<FName>& NewPackages, TSet<FName>& RemovedPackages, TSet<FName>& IdenticalCookedPackages, TSet<FName>& IdenticalUncookedPackages, bool bRecurseModifications);
 
 	/**
 	 * GenerateChunkManifest 
diff --git a/Engine/Source/Editor/UnrealEd/Public/Kismet2/BlueprintEditorUtils.h b/Engine/Source/Editor/UnrealEd/Public/Kismet2/BlueprintEditorUtils.h
--- a/Engine/Source/Editor/UnrealEd/Public/Kismet2/BlueprintEditorUtils.h
+++ b/Engine/Source/Editor/UnrealEd/Public/Kismet2/BlueprintEditorUtils.h
@@ -257,6 +257,9 @@ public:
 	 */
 	static bool ShouldRegenerateBlueprint(UBlueprint* Blueprint);
 
+	/** Returns true if compilation for the given blueprint has been disabled */
+	static bool IsCompileOnLoadDisabled(UBlueprint* Blueprint);
+
 	/**
 	 * Blueprint has structurally changed (added/removed functions, graphs, etc...). Performs the following actions:
 	 *  - Recompiles the skeleton class.
@@ -1080,14 +1083,16 @@ public:
 	/** Indicates if the variable is used on any graphs in this Blueprint*/
 	static bool IsVariableUsed(const UBlueprint* Blueprint, const FName& Name, UEdGraph* LocalGraphScope = nullptr);
 
-	static void ImportKismetDefaultValueToProperty(UEdGraphPin* SourcePin, UProperty* DestinationProperty, uint8* DestinationAddress, UObject* OwnerObject);
-	
-	static void ExportPropertyToKismetDefaultValue(UEdGraphPin* TargetPin, UProperty* SourceProperty, uint8* SourceAddress, UObject* OwnerObject);
+	/** Copies the value from the passed in string into a property. ContainerMem points to the Struct or Class containing Property */
+	static bool PropertyValueFromString(const UProperty* Property, const FString& StrValue, uint8* Container);
 
-	static bool PropertyValueFromString(const UProperty* Property, const FString& StrValue, uint8* ContainerMem);
+	/** Copies the value from the passed in string into a property. DirectValue is the raw memory address of the property value */
+	static bool PropertyValueFromString_Direct(const UProperty* Property, const FString& StrValue, uint8* DirectValue);
 
+	/** Copies the value from a property into the string OutForm. ContainerMem points to the Struct or Class containing Property */
 	static bool PropertyValueToString(const UProperty* Property, const uint8* Container, FString& OutForm);
 
+	/** Copies the value from a property into the string OutForm. DirectValue is the raw memory address of the property value */
 	static bool PropertyValueToString_Direct(const UProperty* Property, const uint8* DirectValue, FString& OutForm);
 
 	/** Call PostEditChange() on all Actors based on the given Blueprint */
diff --git a/Engine/Source/Editor/UnrealEd/Public/Kismet2/CompilerResultsLog.h b/Engine/Source/Editor/UnrealEd/Public/Kismet2/CompilerResultsLog.h
--- a/Engine/Source/Editor/UnrealEd/Public/Kismet2/CompilerResultsLog.h
+++ b/Engine/Source/Editor/UnrealEd/Public/Kismet2/CompilerResultsLog.h
@@ -13,8 +13,6 @@
 #include "EdGraphToken.h"
 #endif
 
-class Error;
-
 #if WITH_EDITOR
 
 /** This class maps from final objects to their original source object, across cloning, autoexpansion, etc... */
@@ -25,6 +23,7 @@ protected:
 	TMap<UObject const*, UObject*> SourceBacktrackMap;
 	// Maps from transient pins created during compiling to original 'source pin' object
 	TMap<UEdGraphPin*, UEdGraphPin*> PinSourceBacktrackMap;
+
 public:
 	/** Update the source backtrack map to note that NewObject was most closely generated/caused by the SourceObject */
 	void NotifyIntermediateObjectCreation(UObject* NewObject, UObject* SourceObject);
@@ -110,6 +109,9 @@ public:
 protected:
 	// Maps from transient object created during compiling to original 'source code' object
 	FBacktrackMap SourceBacktrackMap;
+	
+	// Maps immediately back to source uobject, which may itself by an intermediate object:
+	TMap<UEdGraphNode const*, UEdGraphNode*> FullSourceBacktrackMap;
 
 	// Name of the source object being compiled
 	FString SourcePath;
@@ -197,7 +199,7 @@ public:
 	UEdGraphNode* GetSourceNode(const UEdGraphNode* IntermediateNode);
 
 	/** Returns the intermediate tunnel instance that generated the node */
-	UEdGraphNode* GetIntermediateTunnelInstance(const UEdGraphNode* IntermediateNode);
+	UEdGraphNode* GetIntermediateTunnelInstance(const UEdGraphNode* IntermediateNode) const;
 
 	/** Returns the source tunnel node for the intermediate node */
 	UEdGraphNode* GetSourceTunnelNode(const UEdGraphNode* IntermediateNode);
@@ -212,6 +214,9 @@ public:
 	UObject* FindSourceObject(UObject* PossiblyDuplicatedObject);
 	UObject const* FindSourceObject(UObject const* PossiblyDuplicatedObject) const;
 
+	/** Returns a int32 used to uniquely identify an action for the latent action manager */
+	int32 CalculateStableIdentifierForLatentActionManager( const UEdGraphNode* Node );
+
 	/** Returns the true source object for the passed in object; does type checking on the result */
 	template <typename T>
 	T* FindSourceObjectTypeChecked(UObject* PossiblyDuplicatedObject)
diff --git a/Engine/Source/Editor/UnrealEd/Public/Kismet2/KismetReinstanceUtilities.h b/Engine/Source/Editor/UnrealEd/Public/Kismet2/KismetReinstanceUtilities.h
--- a/Engine/Source/Editor/UnrealEd/Public/Kismet2/KismetReinstanceUtilities.h
+++ b/Engine/Source/Editor/UnrealEd/Public/Kismet2/KismetReinstanceUtilities.h
@@ -138,6 +138,9 @@ public:
 
 	/** Batch replaces a mapping of one or more classes to their new class by leveraging ReplaceInstancesOfClass */
 	static void BatchReplaceInstancesOfClass(TMap<UClass*, UClass*>& InOldToNewClassMap, TSet<UObject*>* ObjectsThatShouldUseOldStuff = NULL, bool bClassObjectReplaced = false, bool bPreserveRootComponent = true);
+	
+	/** Function used to safely discard a CDO, so that the class can have its layout changed, callers must move parent CDOs aside before moving child CDOs aside: */
+	static UClass* MoveCDOToNewClass(UClass* OwnerClass, const TMap<UClass*, UClass*>& OldToNewMap, bool bAvoidCDODuplication);
 
 	/**
 	 * When re-instancing a component, we have to make sure all instance owners' 
diff --git a/Engine/Source/Editor/UnrealEd/Public/Kismet2/StructureEditorUtils.h b/Engine/Source/Editor/UnrealEd/Public/Kismet2/StructureEditorUtils.h
--- a/Engine/Source/Editor/UnrealEd/Public/Kismet2/StructureEditorUtils.h
+++ b/Engine/Source/Editor/UnrealEd/Public/Kismet2/StructureEditorUtils.h
@@ -129,7 +129,7 @@ public:
 	{
 		if (Struct)
 		{
-			auto VarDescArray = GetVarDescPtr(Struct);
+			TArray<FStructVariableDescription>* VarDescArray = GetVarDescPtr(Struct);
 			return VarDescArray ? VarDescArray->FindByPredicate(FFindByGuidHelper<FStructVariableDescription>(VarGuid)) : nullptr;
 		}
 		return nullptr;
@@ -139,7 +139,7 @@ public:
 	{
 		if (Struct)
 		{
-			auto VarDescArray = GetVarDescPtr(Struct);
+			const TArray<FStructVariableDescription>* VarDescArray = GetVarDescPtr(Struct);
 			return VarDescArray ? VarDescArray->FindByPredicate(FFindByGuidHelper<FStructVariableDescription>(VarGuid)) : nullptr;
 		}
 		return nullptr;
diff --git a/Engine/Source/Editor/UnrealEd/Public/UnrealEd.h b/Engine/Source/Editor/UnrealEd/Public/UnrealEd.h
--- a/Engine/Source/Editor/UnrealEd/Public/UnrealEd.h
+++ b/Engine/Source/Editor/UnrealEd/Public/UnrealEd.h
@@ -189,7 +189,6 @@ MONOLITHIC_HEADER_BOILERPLATE()
 #include "Classes/Factories/FbxSceneImportOptionsSkeletalMesh.h"
 #include "Classes/Factories/FbxSceneImportOptionsStaticMesh.h"
 #include "Classes/Commandlets/FileServerCommandlet.h"
-#include "Classes/Commandlets/FixupRedirectsCommandlet.h"
 #include "Classes/Commandlets/GatherTextCommandletBase.h"
 #include "Classes/Commandlets/GatherTextCommandlet.h"
 #include "Classes/Commandlets/GatherTextFromAssetsCommandlet.h"
diff --git a/Engine/Source/Editor/UnrealEd/Public/UnrealEdSharedPCH.h b/Engine/Source/Editor/UnrealEd/Public/UnrealEdSharedPCH.h
--- a/Engine/Source/Editor/UnrealEd/Public/UnrealEdSharedPCH.h
+++ b/Engine/Source/Editor/UnrealEd/Public/UnrealEdSharedPCH.h
@@ -459,7 +459,6 @@
 
 // From AssetRegistry:
 #include "AssetData.h"
-#include "SharedMapView.h"
 #include "ARFilter.h"
 
 // From Engine:
diff --git a/Engine/Source/Editor/WorldBrowser/Private/LevelCollectionModel.cpp b/Engine/Source/Editor/WorldBrowser/Private/LevelCollectionModel.cpp
--- a/Engine/Source/Editor/WorldBrowser/Private/LevelCollectionModel.cpp
+++ b/Engine/Source/Editor/WorldBrowser/Private/LevelCollectionModel.cpp
@@ -1201,12 +1201,9 @@ void FLevelCollectionModel::SCCDiffAgainstDepot(const FLevelModelList& InList, U
 					FString TempFileName;
 					if (Revision->Get(TempFileName))
 					{
-						// Forcibly disable compile on load in case we are loading old blueprints that might try to update/compile
-						TGuardValue<bool> DisableCompileOnLoad(GForceDisableBlueprintCompileOnLoad, true);
-
 						// Try and load that package
 						FText NotMapReason;
-						UPackage* OldPackage = LoadPackage(NULL, *TempFileName, LOAD_None);
+						UPackage* OldPackage = LoadPackage(NULL, *TempFileName, LOAD_DisableCompileOnLoad);
 						if(OldPackage != NULL && InEditor->PackageIsAMapFile(*TempFileName, NotMapReason))
 						{
 							/* Set the revision information*/
diff --git a/Engine/Source/Programs/AutomationTool/AutomationUtils/CommandletUtils.cs b/Engine/Source/Programs/AutomationTool/AutomationUtils/CommandletUtils.cs
--- a/Engine/Source/Programs/AutomationTool/AutomationUtils/CommandletUtils.cs
+++ b/Engine/Source/Programs/AutomationTool/AutomationUtils/CommandletUtils.cs
@@ -451,6 +451,44 @@ namespace AutomationTool
 			}
 		}
 
+		/// <summary>
+		/// Converts project name to FileReference used by other commandlet functions
+		/// </summary>
+		/// <param name="ProjectName">Project name.</param>
+		/// <returns>FileReference to project location</returns>
+		public static FileReference GetCommandletProjectFile(string ProjectName)
+		{
+			FileReference ProjectFullPath = null;
+			var OriginalProjectName = ProjectName;
+			ProjectName = ProjectName.Trim(new char[] { '\"' });
+			if (ProjectName.IndexOfAny(new char[] { '\\', '/' }) < 0)
+			{
+				ProjectName = CombinePaths(CmdEnv.LocalRoot, ProjectName, ProjectName + ".uproject");
+			}
+			else if (!FileExists_NoExceptions(ProjectName))
+			{
+				ProjectName = CombinePaths(CmdEnv.LocalRoot, ProjectName);
+			}
+			if (FileExists_NoExceptions(ProjectName))
+			{
+				ProjectFullPath = new FileReference(ProjectName);
+			}
+			else
+			{
+				var Branch = new BranchInfo(new List<UnrealTargetPlatform> { UnrealBuildTool.BuildHostPlatform.Current.Platform });
+				var GameProj = Branch.FindGame(OriginalProjectName);
+				if (GameProj != null)
+				{
+					ProjectFullPath = GameProj.FilePath;
+				}
+				if (!FileExists_NoExceptions(ProjectFullPath.FullName))
+				{
+					throw new AutomationException("Could not find a project file {0}.", ProjectName);
+				}
+			}
+			return ProjectFullPath;
+		}
+
 		#endregion
 	}
 }
diff --git a/Engine/Source/Programs/AutomationTool/AutomationUtils/ProjectParams.cs b/Engine/Source/Programs/AutomationTool/AutomationUtils/ProjectParams.cs
--- a/Engine/Source/Programs/AutomationTool/AutomationUtils/ProjectParams.cs
+++ b/Engine/Source/Programs/AutomationTool/AutomationUtils/ProjectParams.cs
@@ -283,6 +283,7 @@ namespace AutomationTool
             this.BundleName = InParams.BundleName;
 			this.RunCommandline = InParams.RunCommandline;
 			this.ServerCommandline = InParams.ServerCommandline;
+            this.ClientCommandline = InParams.ClientCommandline;
             this.Package = InParams.Package;
 			this.Deploy = InParams.Deploy;
 			this.DeployFolder = InParams.DeployFolder;
@@ -633,6 +634,8 @@ namespace AutomationTool
 			this.RunCommandline = this.RunCommandline.Replace('\'', '\"'); // replace any single quotes with double quotes
 			this.ServerCommandline = ParseParamValueIfNotSpecified(Command, ServerCommandline, "servercmdline");
 			this.ServerCommandline = this.ServerCommandline.Replace('\'', '\"'); // replace any single quotes with double quotes
+            this.ClientCommandline = ParseParamValueIfNotSpecified(Command, ClientCommandline, "clientcmdline");
+            this.ClientCommandline = this.ClientCommandline.Replace('\'', '\"'); // replace any single quotes with double quotes
             this.Package = GetParamValueIfNotSpecified(Command, Package, this.Package, "package");
 
 			this.Deploy = GetParamValueIfNotSpecified(Command, Deploy, this.Deploy, "deploy");
@@ -1590,6 +1593,12 @@ namespace AutomationTool
 		[Help("servercmdline", "Additional command line arguments for the program")]
 		public string ServerCommandline;
 
+        /// <summary>
+		/// Run: Override command line arguments to pass to the client, if set it will not try to guess at IPs or settings
+		/// </summary>
+		[Help("clientcmdline", "Override command line arguments to pass to the client")]
+        public string ClientCommandline;
+
         /// <summary>
         /// Run:adds -nullrhi to the client commandline
         /// </summary>
diff --git a/Engine/Source/Programs/AutomationTool/Scripts/CookCommand.Automation.cs b/Engine/Source/Programs/AutomationTool/Scripts/CookCommand.Automation.cs
--- a/Engine/Source/Programs/AutomationTool/Scripts/CookCommand.Automation.cs
+++ b/Engine/Source/Programs/AutomationTool/Scripts/CookCommand.Automation.cs
@@ -30,6 +30,8 @@ public partial class Project : CommandUtils
             string PlatformPlaceholderPattern = "<PLAT>";
 
             string ProjectDir = Params.RawProjectPath.Directory.ToString();
+            // NOTE: in UProjectPackagingSettings::PostEditChangeProperty() there is a hardcoded file path/name that is set to match this; 
+            //       if you alter this path then you need to update that and likely FBlueprintNativeCodeGenPaths::GetDefaultCodeGenPaths() as well
             PluginPath = CombinePaths(ProjectDir, "Intermediate", "Plugins", PlatformPlaceholderPattern, "NativizedAssets", "NativizedAssets.uplugin");
 
             ProjectParams.BlueprintPluginKey PluginKey = new ProjectParams.BlueprintPluginKey();
diff --git a/Engine/Source/Programs/AutomationTool/Scripts/FixupRedirects.Automation.cs b/Engine/Source/Programs/AutomationTool/Scripts/FixupRedirects.Automation.cs
--- a/Engine/Source/Programs/AutomationTool/Scripts/FixupRedirects.Automation.cs
+++ b/Engine/Source/Programs/AutomationTool/Scripts/FixupRedirects.Automation.cs
@@ -10,13 +10,10 @@ class FixupRedirects : BuildCommand
 {
 	public override void ExecuteBuild()
 	{
-		var EditorExe = CombinePaths(CmdEnv.LocalRoot, @"Engine/Binaries/Win64/UE4Editor-Cmd.exe");
-		Log("********** Running FixupRedirects: {0} -run=FixupRedirects -unattended -nopause -buildmachine", EditorExe);
-		var RunResult = Run(EditorExe, String.Format("-run=FixupRedirects -unattended -nopause -buildmachine -autosubmit"));
-		if (RunResult.ExitCode != 0)
-		{
-			throw new AutomationException("BUILD FAILED: FixupRedirects failed");
-		}
+		var ProjectName = ParseParamValue("project", "");
 
+		var EditorExe = CombinePaths(CmdEnv.LocalRoot, @"Engine/Binaries/Win64/UE4Editor-Cmd.exe");
+		Log("********** Running FixupRedirects: {0} -run=ResavePackages -unattended -nopause -buildmachine -fixupredirects -autocheckout -autocheckin -projectonly", EditorExe);
+		RunCommandlet(GetCommandletProjectFile(ProjectName), EditorExe, "ResavePackages", "-unattended -nopause -buildmachine -fixupredirects -autocheckout -autocheckin -projectonly");
 	}
 }
diff --git a/Engine/Source/Programs/AutomationTool/Scripts/RunProjectCommand.Automation.cs b/Engine/Source/Programs/AutomationTool/Scripts/RunProjectCommand.Automation.cs
--- a/Engine/Source/Programs/AutomationTool/Scripts/RunProjectCommand.Automation.cs
+++ b/Engine/Source/Programs/AutomationTool/Scripts/RunProjectCommand.Automation.cs
@@ -130,15 +130,19 @@ public partial class Project : CommandUtils
 			{
 				TargetPlatformDescriptor ServerPlatformDesc = Params.ServerTargetPlatforms[0];
 				ServerProcess = RunDedicatedServer(Params, ServerLogFile, Params.RunCommandline);
-				// With dedicated server, the client connects to local host to load a map.
-				if (ServerPlatformDesc.Type == UnrealTargetPlatform.Linux)
-				{
-					Params.MapToRun = Params.ServerDeviceAddress;
-				}
-				else
-				{
-					Params.MapToRun = "127.0.0.1";
-				}
+
+                // With dedicated server, the client connects to local host to load a map, unless client parameters are already specified
+                if (String.IsNullOrEmpty(Params.ClientCommandline))
+                {
+                    if (!String.IsNullOrEmpty(Params.ServerDeviceAddress))
+                    {
+                        Params.ClientCommandline = Params.ServerDeviceAddress;
+                    }
+                    else
+                    {
+                        Params.ClientCommandline = "127.0.0.1";
+                    }
+                }
 			}
 			else
 			{
@@ -626,11 +630,15 @@ public partial class Project : CommandUtils
 		{
 			List<string> Exes = SC.StageTargetPlatform.GetExecutableNames(SC, true);
 			ClientApp = Exes[0];
-			if (SC.StageTargetPlatform.PlatformType != UnrealTargetPlatform.IOS)
-			{
-				TempCmdLine += SC.ProjectArgForCommandLines + " ";
-			}
-			TempCmdLine += Params.MapToRun + " ";
+
+            if (!String.IsNullOrEmpty(Params.ClientCommandline))
+            {
+                TempCmdLine += Params.ClientCommandline + " ";
+            }
+            else
+            {
+                TempCmdLine += Params.MapToRun + " ";
+            }
 
 			if (Params.CookOnTheFly || Params.FileServer)
 			{
@@ -861,13 +869,12 @@ public partial class Project : CommandUtils
 		{
 			TempCmdLine += "-abslog=" + CommandUtils.MakePathSafeToUseWithCommandLine(ClientLogFile) + " ";
 		}
-		if (SC.StageTargetPlatform.PlatformType != UnrealTargetPlatform.IOS && SC.StageTargetPlatform.PlatformType != UnrealTargetPlatform.Linux)
+		if (SC.StageTargetPlatform.PlatformType == UnrealTargetPlatform.Win32 || SC.StageTargetPlatform.PlatformType == UnrealTargetPlatform.Win64)
 		{
-			TempCmdLine += "-Messaging -nomcp -Windowed ";
+			TempCmdLine += "-Messaging -Windowed ";
 		}
 		else
 		{
-			// skip arguments which don't make sense for iOS
 			TempCmdLine += "-Messaging ";
 		}
 		if (Params.NullRHI && SC.StageTargetPlatform.PlatformType != UnrealTargetPlatform.Mac) // all macs have GPUs, and currently the mac dies with nullrhi
@@ -986,12 +993,7 @@ public partial class Project : CommandUtils
 			{
 				Args += " -unattended";
 			}
-			// Do not blindly add -nomcp, only do so if the client is using it
-			if (Params.RunCommandline.Contains("-nomcp"))
-			{
-				Args += " -nomcp";
-			}
-
+			
 			if (Params.ServerCommandline.Length > 0)
 			{
 				Args += " " + Params.ServerCommandline;
diff --git a/Engine/Source/Programs/UnrealBuildTool/System/ModuleDescriptor.cs b/Engine/Source/Programs/UnrealBuildTool/System/ModuleDescriptor.cs
--- a/Engine/Source/Programs/UnrealBuildTool/System/ModuleDescriptor.cs
+++ b/Engine/Source/Programs/UnrealBuildTool/System/ModuleDescriptor.cs
@@ -141,6 +141,16 @@ namespace UnrealBuildTool
 		/// </summary>
 		public UnrealTargetPlatform[] BlacklistPlatforms;
 
+		/// <summary>
+		/// List of allowed targets
+		/// </summary>
+		public TargetType[] WhitelistTargets;
+
+		/// <summary>
+		/// List of disallowed targets
+		/// </summary>
+		public TargetType[] BlacklistTargets;
+
 		/// <summary>
 		/// List of additional dependencies for building this module.
 		/// </summary>
@@ -184,6 +194,18 @@ namespace UnrealBuildTool
 				Module.BlacklistPlatforms = BlacklistPlatforms;
 			}
 
+			TargetType[] WhitelistTargets;
+			if (InObject.TryGetEnumArrayField<TargetType>("WhitelistTargets", out WhitelistTargets))
+			{
+				Module.WhitelistTargets = WhitelistTargets;
+			}
+
+			TargetType[] BlacklistTargets;
+			if (InObject.TryGetEnumArrayField<TargetType>("BlacklistTargets", out BlacklistTargets))
+			{
+				Module.BlacklistTargets = BlacklistTargets;
+			}
+
 			string[] AdditionalDependencies;
 			if (InObject.TryGetStringArrayField("AdditionalDependencies", out AdditionalDependencies))
 			{
@@ -221,6 +243,24 @@ namespace UnrealBuildTool
 				}
 				Writer.WriteArrayEnd();
 			}
+			if (WhitelistTargets != null && WhitelistTargets.Length > 0)
+			{
+				Writer.WriteArrayStart("WhitelistTargets");
+				foreach (TargetType WhitelistTarget in WhitelistTargets)
+				{
+					Writer.WriteValue(WhitelistTarget.ToString());
+				}
+				Writer.WriteArrayEnd();
+			}
+			if (BlacklistTargets != null && BlacklistTargets.Length > 0)
+			{
+				Writer.WriteArrayStart("BlacklistTargets");
+				foreach (TargetType BlacklistTarget in BlacklistTargets)
+				{
+					Writer.WriteValue(BlacklistTarget.ToString());
+				}
+				Writer.WriteArrayEnd();
+			}
 			if (AdditionalDependencies != null && AdditionalDependencies.Length > 0)
 			{
 				Writer.WriteArrayStart("AdditionalDependencies");
@@ -274,6 +314,18 @@ namespace UnrealBuildTool
 				return false;
 			}
 
+			// Check the target is whitelisted
+			if (WhitelistTargets != null && WhitelistTargets.Length > 0 && !WhitelistTargets.Contains(TargetType))
+			{
+				return false;
+			}
+
+			// Check the target is not blacklisted
+			if (BlacklistTargets != null && BlacklistTargets.Contains(TargetType))
+			{
+				return false;
+			}
+
 			// Check the module is compatible with this target.
 			switch (Type)
 			{
diff --git a/Engine/Source/Programs/UnrealBuildTool/System/ProjectFileGenerator.cs b/Engine/Source/Programs/UnrealBuildTool/System/ProjectFileGenerator.cs
--- a/Engine/Source/Programs/UnrealBuildTool/System/ProjectFileGenerator.cs
+++ b/Engine/Source/Programs/UnrealBuildTool/System/ProjectFileGenerator.cs
@@ -966,6 +966,26 @@ namespace UnrealBuildTool
 				EngineProject.AddFilesToProject(SourceFileSearch.FindFiles(UHTConfigDirectory), UnrealBuildTool.EngineDirectory);
 			}
 		}
+		private List<FileReference> DiscoverExtraPlugins(List<UProjectInfo> AllGameProjects)
+		{
+			List<FileReference> AddedPlugins = new List<FileReference>();
+			foreach (UProjectInfo GameProject in AllGameProjects)
+			{
+				ProjectDescriptor ProjectDesc = ProjectDescriptor.FromFile(GameProject.FilePath.FullName);
+				foreach (DirectoryReference PluginDir in ProjectDesc.AdditionalPluginDirectories)
+				{
+					// @TODO: Right now, we're only including additional plugins that are still inside a game directory.
+					//        This is so FindProjectForModule() succeeds in matching up a project with the plugin. If 
+					//        the plugin is outside of a game directory, then we'll need to add special handling to AddProjectsForAllModules()
+					//        and generate a standalone project, since multiple games could share the same plugin
+					if (PluginDir.IsUnderDirectory(GameProject.Folder))
+					{
+						AddedPlugins.AddRange(Plugins.EnumeratePlugins(PluginDir));
+					}
+				}
+			}
+			return AddedPlugins;
+		}
 
 		/// <summary>
 		/// Finds all module files (filtering by platform)
@@ -976,7 +996,7 @@ namespace UnrealBuildTool
 			List<FileReference> AllModuleFiles = new List<FileReference>();
 
 			// Locate all modules (*.Build.cs files)
-			List<FileReference> FoundModuleFiles = RulesCompiler.FindAllRulesSourceFiles( RulesCompiler.RulesFileType.Module, GameFolders: AllGameProjects.Select(x => x.Folder).ToList(), ForeignPlugins:null, AdditionalSearchPaths:null );
+			List<FileReference> FoundModuleFiles = RulesCompiler.FindAllRulesSourceFiles( RulesCompiler.RulesFileType.Module, GameFolders: AllGameProjects.Select(x => x.Folder).ToList(), ForeignPlugins: DiscoverExtraPlugins(AllGameProjects), AdditionalSearchPaths:null );
 			foreach( FileReference BuildFileName in FoundModuleFiles )
 			{
 				AllModuleFiles.Add( BuildFileName );
@@ -1020,7 +1040,7 @@ namespace UnrealBuildTool
 			List<string> UnsupportedPlatformNameStrings = Utils.MakeListOfUnsupportedPlatforms( SupportedPlatforms );
 
 			// Locate all targets (*.Target.cs files)
-			List<FileReference> FoundTargetFiles = RulesCompiler.FindAllRulesSourceFiles( RulesCompiler.RulesFileType.Target, AllGameProjects.Select(x => x.Folder).ToList(), ForeignPlugins:null, AdditionalSearchPaths:null );
+			List<FileReference> FoundTargetFiles = RulesCompiler.FindAllRulesSourceFiles( RulesCompiler.RulesFileType.Target, AllGameProjects.Select(x => x.Folder).ToList(), ForeignPlugins: DiscoverExtraPlugins(AllGameProjects), AdditionalSearchPaths:null );
 			foreach( FileReference CurTargetFile in FoundTargetFiles )
 			{
 				string CleanTargetFileName = Utils.CleanDirectorySeparators( CurTargetFile.FullName );
diff --git a/Engine/Source/Programs/UnrealHeaderTool/Private/BaseParser.cpp b/Engine/Source/Programs/UnrealHeaderTool/Private/BaseParser.cpp
--- a/Engine/Source/Programs/UnrealHeaderTool/Private/BaseParser.cpp
+++ b/Engine/Source/Programs/UnrealHeaderTool/Private/BaseParser.cpp
@@ -90,6 +90,10 @@ namespace
 			FMetadataKeyword& FriendlyName = Dictionary.Add(TEXT("FriendlyName"), EMetadataValueArgument::Required);
 			FriendlyName.InsertAddAction(TEXT("FriendlyName"), TEXT(""));
 
+			FMetadataKeyword& BlueprintInternalUseOnly = Dictionary.Add(TEXT("BlueprintInternalUseOnly"), EMetadataValueArgument::None);
+			BlueprintInternalUseOnly.InsertAddAction(TEXT("BlueprintInternalUseOnly"), TEXT("true"));
+			BlueprintInternalUseOnly.InsertAddAction(TEXT("BlueprintType"), TEXT("true"));
+
 			FMetadataKeyword& BlueprintType = Dictionary.Add(TEXT("BlueprintType"), EMetadataValueArgument::None);
 			BlueprintType.InsertAddAction(TEXT("BlueprintType"), TEXT("true"));
 
diff --git a/Engine/Source/Programs/UnrealHeaderTool/Private/ClassDeclarationMetaData.cpp b/Engine/Source/Programs/UnrealHeaderTool/Private/ClassDeclarationMetaData.cpp
--- a/Engine/Source/Programs/UnrealHeaderTool/Private/ClassDeclarationMetaData.cpp
+++ b/Engine/Source/Programs/UnrealHeaderTool/Private/ClassDeclarationMetaData.cpp
@@ -28,14 +28,14 @@ namespace
 }
 
 FClassDeclarationMetaData::FClassDeclarationMetaData()
-	: ClassFlags(0)
+	: ClassFlags(CLASS_None)
 	, WantsToBePlaceable(false)
 {
 }
 
 void FClassDeclarationMetaData::ParseClassProperties(const TArray<FPropertySpecifier>& InClassSpecifiers, const FString& InRequiredAPIMacroIfPresent)
 {
-	ClassFlags = 0;
+	ClassFlags = CLASS_None;
 	// Record that this class is RequiredAPI if the CORE_API style macro was present
 	if (!InRequiredAPIMacroIfPresent.IsEmpty())
 	{
diff --git a/Engine/Source/Programs/UnrealHeaderTool/Private/ClassDeclarationMetaData.h b/Engine/Source/Programs/UnrealHeaderTool/Private/ClassDeclarationMetaData.h
--- a/Engine/Source/Programs/UnrealHeaderTool/Private/ClassDeclarationMetaData.h
+++ b/Engine/Source/Programs/UnrealHeaderTool/Private/ClassDeclarationMetaData.h
@@ -5,6 +5,7 @@
 #include "CoreMinimal.h"
 #include "UObject/NameTypes.h"
 #include "Containers/Map.h"
+#include "Class.h"
 
 class FClass;
 class FClasses;
@@ -16,7 +17,7 @@ class FClassDeclarationMetaData
 public:
 	FClassDeclarationMetaData();
 
-	uint64 ClassFlags;
+	EClassFlags ClassFlags;
 	TMap<FName, FString> MetaData;
 	FString ClassWithin;
 	FString ConfigName;
diff --git a/Engine/Source/Programs/UnrealHeaderTool/Private/CodeGenerator.cpp b/Engine/Source/Programs/UnrealHeaderTool/Private/CodeGenerator.cpp
--- a/Engine/Source/Programs/UnrealHeaderTool/Private/CodeGenerator.cpp
+++ b/Engine/Source/Programs/UnrealHeaderTool/Private/CodeGenerator.cpp
@@ -1415,7 +1415,7 @@ void FNativeClassHeaderGenerator::ExportNativeGeneratedInitCode(FOutputDevice& O
 		GeneratedClassRegisterFunctionText.Logf(TEXT("\t\t\t{\r\n"), ClassNameCPP);
 		GeneratedClassRegisterFunctionText.Logf(TEXT("\t\t\t\tUObjectForceRegistration(OuterClass);\r\n"));
 		uint32 Flags = (Class->ClassFlags & CLASS_SaveInCompiledInClasses) | CLASS_Constructed;
-		GeneratedClassRegisterFunctionText.Logf(TEXT("\t\t\t\tOuterClass->ClassFlags |= 0x%08X;\r\n"), Flags);
+		GeneratedClassRegisterFunctionText.Logf(TEXT("\t\t\t\tOuterClass->ClassFlags |= (EClassFlags)0x%08Xu;\r\n"), Flags);
 		TheFlagAudit.Add(Class, TEXT("ClassFlags"), Flags);
 		GeneratedClassRegisterFunctionText.Logf(TEXT("\r\n"));
 		GeneratedClassRegisterFunctionText.Log(CallSingletons);
@@ -2202,7 +2202,9 @@ void FNativeClassHeaderGenerator::ExportClassFromSourceFileInner(
 			// to script VM functions
 			if (SuperClass->IsChildOf(UInterface::StaticClass()))
 			{
-				InterfaceBoilerplate.Logf(TEXT("\tvirtual UObject* _getUObject() const = 0;\r\n"));
+				// Note: This used to be declared as a pure virtual function, but it was changed here in order to allow the Blueprint nativization process
+				// to detect C++ interface classes that explicitly declare pure virtual functions via type traits. This code will no longer trigger that check.
+				InterfaceBoilerplate.Logf(TEXT("\tvirtual UObject* _getUObject() const { check(0 && \"Missing required implementation.\"); return nullptr; }\r\n"));
 			}
 
 			if (bNeedsRep && !bHasGetLifetimeReplicatedProps)
@@ -4784,7 +4786,7 @@ void ResolveSuperClasses(UPackage* Package)
 	TArray<UObject*> Objects;
 	GetObjectsWithOuter(Package, Objects);
 
-	for (auto* Object : Objects)
+	for (UObject* Object : Objects)
 	{
 		if (!Object->IsA<UClass>())
 		{
@@ -5258,7 +5260,7 @@ ECompilationResult::Type UnrealHeaderTool_Main(const FString& ModuleInfoFilename
 	return Result;
 }
 
-UClass* ProcessParsedClass(bool bClassIsAnInterface, TArray<FHeaderProvider> &DependentOn, const FString& ClassName, const FString& BaseClassName, UObject* InParent, EObjectFlags Flags)
+UClass* ProcessParsedClass(bool bClassIsAnInterface, TArray<FHeaderProvider>& DependentOn, const FString& ClassName, const FString& BaseClassName, UObject* InParent, EObjectFlags Flags)
 {
 	FString ClassNameStripped = GetClassNameWithPrefixRemoved(*ClassName);
 
diff --git a/Engine/Source/Programs/UnrealHeaderTool/Private/HeaderParser.cpp b/Engine/Source/Programs/UnrealHeaderTool/Private/HeaderParser.cpp
--- a/Engine/Source/Programs/UnrealHeaderTool/Private/HeaderParser.cpp
+++ b/Engine/Source/Programs/UnrealHeaderTool/Private/HeaderParser.cpp
@@ -283,7 +283,7 @@ namespace
 					}
 					else if ( (FuncInfo.FunctionFlags & FUNC_Private) )
 					{
-						FError::Throwf(TEXT("A Private function cannot be a BlueprintNativeEvent!") );
+						UE_LOG_ERROR_UHT(TEXT("A Private function cannot be a BlueprintNativeEvent!") );
 					}
 
 					FuncInfo.FunctionFlags |= FUNC_Event;
@@ -304,7 +304,7 @@ namespace
 					}
 					else if ( (FuncInfo.FunctionFlags & FUNC_Private) )
 					{
-						FError::Throwf(TEXT("A Private function cannot be a BlueprintImplementableEvent!") );
+						UE_LOG_ERROR_UHT(TEXT("A Private function cannot be a BlueprintImplementableEvent!") );
 					}
 
 					FuncInfo.FunctionFlags |= FUNC_Event;
@@ -318,7 +318,7 @@ namespace
 					FuncInfo.FunctionFlags |= FUNC_Exec;
 					if( FuncInfo.FunctionFlags & FUNC_Net )
 					{
-						FError::Throwf(TEXT("Exec functions cannot be replicated!") );
+						UE_LOG_ERROR_UHT(TEXT("Exec functions cannot be replicated!") );
 					}
 				}
 				break;
@@ -346,7 +346,7 @@ namespace
 
 					if( FuncInfo.FunctionFlags & FUNC_Exec )
 					{
-						FError::Throwf(TEXT("Exec functions cannot be replicated!") );
+						UE_LOG_ERROR_UHT(TEXT("Exec functions cannot be replicated!") );
 					}
 				}
 				break;
@@ -503,37 +503,43 @@ namespace
 			bool bIsNetService  = !!(FuncInfo.FunctionFlags & (FUNC_NetRequest | FUNC_NetResponse));
 			bool bIsNetReliable = !!(FuncInfo.FunctionFlags & FUNC_NetReliable);
 
-			if ( FuncInfo.FunctionFlags & FUNC_Static )
-				FError::Throwf(TEXT("Static functions can't be replicated") );
+			if (FuncInfo.FunctionFlags & FUNC_Static)
+			{
+				UE_LOG_ERROR_UHT(TEXT("Static functions can't be replicated"));
+			}
 
 			if (!bIsNetReliable && !bSpecifiedUnreliable && !bIsNetService)
-				FError::Throwf(TEXT("Replicated function: 'reliable' or 'unreliable' is required"));
+			{
+				UE_LOG_ERROR_UHT(TEXT("Replicated function: 'reliable' or 'unreliable' is required"));
+			}
 
 			if (bIsNetReliable && bSpecifiedUnreliable && !bIsNetService)
-				FError::Throwf(TEXT("'reliable' and 'unreliable' are mutually exclusive"));
+			{
+				UE_LOG_ERROR_UHT(TEXT("'reliable' and 'unreliable' are mutually exclusive"));
+			}
 		}
 		else if (FuncInfo.FunctionFlags & FUNC_NetReliable)
 		{
-			FError::Throwf(TEXT("'reliable' specified without 'client' or 'server'"));
+			UE_LOG_ERROR_UHT(TEXT("'reliable' specified without 'client' or 'server'"));
 		}
 		else if (bSpecifiedUnreliable)
 		{
-			FError::Throwf(TEXT("'unreliable' specified without 'client' or 'server'"));
+			UE_LOG_ERROR_UHT(TEXT("'unreliable' specified without 'client' or 'server'"));
 		}
 
 		if (FuncInfo.bSealedEvent && !(FuncInfo.FunctionFlags & FUNC_Event))
 		{
-			FError::Throwf(TEXT("SealedEvent may only be used on events"));
+			UE_LOG_ERROR_UHT(TEXT("SealedEvent may only be used on events"));
 		}
 
 		if (FuncInfo.bSealedEvent && FuncInfo.FunctionFlags & FUNC_BlueprintEvent)
 		{
-			FError::Throwf(TEXT("SealedEvent cannot be used on Blueprint events"));
+			UE_LOG_ERROR_UHT(TEXT("SealedEvent cannot be used on Blueprint events"));
 		}
 
 		if (FuncInfo.bForceBlueprintImpure && (FuncInfo.FunctionFlags & FUNC_BlueprintPure) != 0)
 		{
-			FError::Throwf(TEXT("BlueprintPure (or BlueprintPure=true) and BlueprintPure=false should not both appear on the same function, they are mutually exclusive"));
+			UE_LOG_ERROR_UHT(TEXT("BlueprintPure (or BlueprintPure=true) and BlueprintPure=false should not both appear on the same function, they are mutually exclusive"));
 		}
 	}
 
@@ -1017,12 +1023,15 @@ namespace
 			return IsPropertySupportedByBlueprint(MapProperty->KeyProp, false) &&
 				IsPropertySupportedByBlueprint(MapProperty->ValueProp, false);
 		}
+		else if (const UStructProperty* StructProperty = Cast<const UStructProperty>(Property))
+		{
+			return (StructProperty->Struct->GetBoolMetaDataHierarchical(TEXT("BlueprintType")));
+		}
 
 		const bool bSupportedType = Property->IsA<UInterfaceProperty>()
 			|| Property->IsA<UClassProperty>()
 			|| Property->IsA<UAssetObjectProperty>()
 			|| Property->IsA<UObjectProperty>()
-			|| Property->IsA<UStructProperty>()
 			|| Property->IsA<UFloatProperty>()
 			|| Property->IsA<UIntProperty>()
 			|| Property->IsA<UByteProperty>()
@@ -1992,7 +2001,7 @@ UScriptStruct* FHeaderParser::CompileStructDeclaration(FClasses& AllClasses)
 
 				if (!FPaths::IsSamePath(Filename, GTypeDefinitionInfoMap[UObject::StaticClass()]->GetUnrealSourceFile().GetFilename()))
 				{
-					FError::Throwf(TEXT("Immutable is being phased out in favor of SerializeNative, and is only legal on the mirror structs declared in UObject"));
+					UE_LOG_ERROR_UHT(TEXT("Immutable is being phased out in favor of SerializeNative, and is only legal on the mirror structs declared in UObject"));
 				}
 			}
 			break;
@@ -2675,7 +2684,7 @@ void FHeaderParser::FixupDelegateProperties( FClasses& AllClasses, UStruct* Stru
 									const bool bAllowedArrayRefFromBP = bClassGeneratedFromBP && FuncParam->IsA<UArrayProperty>();
 									if (!bAllowedArrayRefFromBP)
 									{
-										FError::Throwf(TEXT("BlueprintAssignable delegates do not support non-const references at the moment. Function: %s Parameter: '%s'"), *SourceDelegateFunction->GetName(), *FuncParam->GetName());
+										UE_LOG_ERROR_UHT(TEXT("BlueprintAssignable delegates do not support non-const references at the moment. Function: %s Parameter: '%s'"), *SourceDelegateFunction->GetName(), *FuncParam->GetName());
 									}
 								}
 							}
@@ -3026,7 +3035,7 @@ void FHeaderParser::GetVarType(
 				{
 					if (bSeenEditSpecifier)
 					{
-						FError::Throwf(TEXT("Found more than one edit/visibility specifier (%s), only one is allowed"), *Specifier.Key);
+						UE_LOG_ERROR_UHT(TEXT("Found more than one edit/visibility specifier (%s), only one is allowed"), *Specifier.Key);
 					}
 					Flags |= CPF_Edit;
 					bSeenEditSpecifier = true;
@@ -3037,7 +3046,7 @@ void FHeaderParser::GetVarType(
 				{
 					if (bSeenEditSpecifier)
 					{
-						FError::Throwf(TEXT("Found more than one edit/visibility specifier (%s), only one is allowed"), *Specifier.Key);
+						UE_LOG_ERROR_UHT(TEXT("Found more than one edit/visibility specifier (%s), only one is allowed"), *Specifier.Key);
 					}
 					Flags |= CPF_Edit | CPF_DisableEditOnTemplate;
 					bSeenEditSpecifier = true;
@@ -3048,7 +3057,7 @@ void FHeaderParser::GetVarType(
 				{
 					if (bSeenEditSpecifier)
 					{
-						FError::Throwf(TEXT("Found more than one edit/visibility specifier (%s), only one is allowed"), *Specifier.Key);
+						UE_LOG_ERROR_UHT(TEXT("Found more than one edit/visibility specifier (%s), only one is allowed"), *Specifier.Key);
 					}
 					Flags |= CPF_Edit | CPF_DisableEditOnInstance;
 					bSeenEditSpecifier = true;
@@ -3059,7 +3068,7 @@ void FHeaderParser::GetVarType(
 				{
 					if (bSeenEditSpecifier)
 					{
-						FError::Throwf(TEXT("Found more than one edit/visibility specifier (%s), only one is allowed"), *Specifier.Key);
+						UE_LOG_ERROR_UHT(TEXT("Found more than one edit/visibility specifier (%s), only one is allowed"), *Specifier.Key);
 					}
 					Flags |= CPF_Edit | CPF_EditConst;
 					bSeenEditSpecifier = true;
@@ -3070,7 +3079,7 @@ void FHeaderParser::GetVarType(
 				{
 					if (bSeenEditSpecifier)
 					{
-						FError::Throwf(TEXT("Found more than one edit/visibility specifier (%s), only one is allowed"), *Specifier.Key);
+						UE_LOG_ERROR_UHT(TEXT("Found more than one edit/visibility specifier (%s), only one is allowed"), *Specifier.Key);
 					}
 					Flags |= CPF_Edit | CPF_EditConst | CPF_DisableEditOnTemplate;
 					bSeenEditSpecifier = true;
@@ -3081,7 +3090,7 @@ void FHeaderParser::GetVarType(
 				{
 					if (bSeenEditSpecifier)
 					{
-						FError::Throwf(TEXT("Found more than one edit/visibility specifier (%s), only one is allowed"), *Specifier.Key);
+						UE_LOG_ERROR_UHT(TEXT("Found more than one edit/visibility specifier (%s), only one is allowed"), *Specifier.Key);
 					}
 					Flags |= CPF_Edit | CPF_EditConst | CPF_DisableEditOnInstance;
 					bSeenEditSpecifier = true;
@@ -3092,19 +3101,19 @@ void FHeaderParser::GetVarType(
 				{
 					if (bSeenBlueprintEditSpecifier)
 					{
-						FError::Throwf(TEXT("Found more than one Blueprint read/write specifier (%s), only one is allowed"), *Specifier.Key);
+						UE_LOG_ERROR_UHT(TEXT("Found more than one Blueprint read/write specifier (%s), only one is allowed"), *Specifier.Key);
 					}
 
 					const FString* PrivateAccessMD = MetaDataFromNewStyle.Find(TEXT("AllowPrivateAccess"));  // FBlueprintMetadata::MD_AllowPrivateAccess
 					const bool bAllowPrivateAccess = PrivateAccessMD ? (*PrivateAccessMD == TEXT("true")) : false;
 					if (CurrentAccessSpecifier == ACCESS_Private && !bAllowPrivateAccess)
 					{
-						FError::Throwf(TEXT("BlueprintReadWrite should not be used on private members"));
+						UE_LOG_ERROR_UHT(TEXT("BlueprintReadWrite should not be used on private members"));
 					}
 
 					if ((Flags & CPF_EditorOnly) != 0 && OwnerStruct->IsA<UScriptStruct>())
 					{
-						FError::Throwf(TEXT("Blueprint exposed struct members cannot be editor only"));
+						UE_LOG_ERROR_UHT(TEXT("Blueprint exposed struct members cannot be editor only"));
 					}
 
 					Flags |= CPF_BlueprintVisible;
@@ -3116,19 +3125,19 @@ void FHeaderParser::GetVarType(
 				{
 					if (bSeenBlueprintEditSpecifier)
 					{
-						FError::Throwf(TEXT("Found more than one Blueprint read/write specifier (%s), only one is allowed"), *Specifier.Key);
+						UE_LOG_ERROR_UHT(TEXT("Found more than one Blueprint read/write specifier (%s), only one is allowed"), *Specifier.Key);
 					}
 
 					const FString* PrivateAccessMD = MetaDataFromNewStyle.Find(TEXT("AllowPrivateAccess"));  // FBlueprintMetadata::MD_AllowPrivateAccess
 					const bool bAllowPrivateAccess = PrivateAccessMD ? (*PrivateAccessMD == TEXT("true")) : false;
 					if (CurrentAccessSpecifier == ACCESS_Private && !bAllowPrivateAccess)
 					{
-						FError::Throwf(TEXT("BlueprintReadOnly should not be used on private members"));
+						UE_LOG_ERROR_UHT(TEXT("BlueprintReadOnly should not be used on private members"));
 					}
 
 					if ((Flags & CPF_EditorOnly) != 0 && OwnerStruct->IsA<UScriptStruct>())
 					{
-						FError::Throwf(TEXT("Blueprint exposed struct members cannot be editor only"));
+						UE_LOG_ERROR_UHT(TEXT("Blueprint exposed struct members cannot be editor only"));
 					}
 
 					Flags        |= CPF_BlueprintVisible | CPF_BlueprintReadOnly;
@@ -3151,7 +3160,7 @@ void FHeaderParser::GetVarType(
 
 				case EVariableSpecifier::Localized:
 				{
-					FError::Throwf(TEXT("The Localized specifier is deprecated"));
+					UE_LOG_ERROR_UHT(TEXT("The Localized specifier is deprecated"));
 				}
 				break;
 
@@ -3194,7 +3203,7 @@ void FHeaderParser::GetVarType(
 
 				case EVariableSpecifier::EditInline:
 				{
-					FError::Throwf(TEXT("EditInline is deprecated. Remove it, or use Instanced instead."));
+					UE_LOG_ERROR_UHT(TEXT("EditInline is deprecated. Remove it, or use Instanced instead."));
 				}
 				break;
 
@@ -3215,7 +3224,7 @@ void FHeaderParser::GetVarType(
 				{
 					if (OwnerStruct->IsA<UScriptStruct>())
 					{
-						FError::Throwf(TEXT("Struct members cannot be replicated"));
+						UE_LOG_ERROR_UHT(TEXT("Struct members cannot be replicated"));
 					}
 
 					Flags |= CPF_Net;
@@ -3233,7 +3242,7 @@ void FHeaderParser::GetVarType(
 				{
 					if (!OwnerStruct->IsA<UScriptStruct>())
 					{
-						FError::Throwf(TEXT("Only Struct members can be marked NotReplicated"));
+						UE_LOG_ERROR_UHT(TEXT("Only Struct members can be marked NotReplicated"));
 					}
 
 					Flags |= CPF_RepSkip;
@@ -3242,7 +3251,7 @@ void FHeaderParser::GetVarType(
 
 				case EVariableSpecifier::RepRetry:
 				{
-					FError::Throwf(TEXT("'RepRetry' is deprecated."));
+					UE_LOG_ERROR_UHT(TEXT("'RepRetry' is deprecated."));
 				}
 				break;
 
@@ -3317,7 +3326,7 @@ void FHeaderParser::GetVarType(
 
 				default:
 				{
-					FError::Throwf(TEXT("Unknown variable specifier '%s'"), *Specifier.Key);
+					UE_LOG_ERROR_UHT(TEXT("Unknown variable specifier '%s'"), *Specifier.Key);
 				}
 				break;
 			}
@@ -3347,14 +3356,14 @@ void FHeaderParser::GetVarType(
 					}
 					else
 					{
-						FError::Throwf(TEXT("Only parameters in service request functions can be marked NotReplicated"));
+						UE_LOG_ERROR_UHT(TEXT("Only parameters in service request functions can be marked NotReplicated"));
 					}
 				}
 				break;
 
 				default:
 				{
-					FError::Throwf(TEXT("Unknown variable specifier '%s'"), *Specifier.Key);
+					UE_LOG_ERROR_UHT(TEXT("Unknown variable specifier '%s'"), *Specifier.Key);
 				}
 				break;
 			}
@@ -3519,7 +3528,7 @@ void FHeaderParser::GetVarType(
 	{
 		if (IsBitfieldProperty())
 		{
-			FError::Throwf(TEXT("bool bitfields are not supported."));
+			UE_LOG_ERROR_UHT(TEXT("bool bitfields are not supported."));
 		}
 		// C++ bool type
 		VarProperty = FPropertyBase(CPT_Bool);
@@ -4232,7 +4241,7 @@ void FHeaderParser::GetVarType(
 		}
 	}
 
-	if ( VarProperty.IsObject() && VarProperty.MetaClass == NULL && (VarProperty.PropertyFlags&CPF_Config) != 0 )
+	if ( VarProperty.IsObject() && VarProperty.Type != CPT_AssetObjectReference && VarProperty.MetaClass == nullptr && (VarProperty.PropertyFlags&CPF_Config) != 0 )
 	{
 		FError::Throwf(TEXT("Not allowed to use 'config' with object variables"));
 	}
@@ -5581,8 +5590,10 @@ bool FHeaderParser::IsValidDelegateDeclaration(const FToken& Token) const
 void FHeaderParser::ParseParameterList(FClasses& AllClasses, UFunction* Function, bool bExpectCommaBeforeName, TMap<FName, FString>* MetaData)
 {
 	// Get parameter list.
-	if ( MatchSymbol(TEXT(")")) )
+	if (MatchSymbol(TEXT(")")))
+	{
 		return;
+	}
 
 	FAdvancedDisplayParameterHandler AdvancedDisplay(MetaData);
 	do
@@ -5613,27 +5624,41 @@ void FHeaderParser::ParseParameterList(FClasses& AllClasses, UFunction* Function
 			if (!(Function->FunctionFlags & FUNC_NetRequest))
 			{
 				if (Property.PropertyFlags & CPF_OutParm)
-					FError::Throwf(TEXT("Replicated functions cannot contain out parameters"));
+				{
+					UE_LOG_ERROR_UHT(TEXT("Replicated functions cannot contain out parameters"));
+				}
 
 				if (Property.PropertyFlags & CPF_RepSkip)
-					FError::Throwf(TEXT("Only service request functions cannot contain NoReplication parameters"));
+				{
+					UE_LOG_ERROR_UHT(TEXT("Only service request functions cannot contain NoReplication parameters"));
+				}
 
 				if ((Prop->GetClass()->ClassCastFlags & CASTCLASS_UDelegateProperty) != 0)
-					FError::Throwf(TEXT("Replicated functions cannot contain delegate parameters (this would be insecure)"));
+				{
+					UE_LOG_ERROR_UHT(TEXT("Replicated functions cannot contain delegate parameters (this would be insecure)"));
+				}
 
 				if (Property.Type == CPT_String && Property.RefQualifier != ERefQualifier::ConstRef && Prop->ArrayDim == 1)
-					FError::Throwf(TEXT("Replicated FString parameters must be passed by const reference"));
+				{
+					UE_LOG_ERROR_UHT(TEXT("Replicated FString parameters must be passed by const reference"));
+				}
 
 				if (Property.ArrayType == EArrayType::Dynamic && Property.RefQualifier != ERefQualifier::ConstRef && Prop->ArrayDim == 1)
-					FError::Throwf(TEXT("Replicated TArray parameters must be passed by const reference"));
+				{
+					UE_LOG_ERROR_UHT(TEXT("Replicated TArray parameters must be passed by const reference"));
+				}
 			}
 			else
 			{
 				if (!(Property.PropertyFlags & CPF_RepSkip) && (Property.PropertyFlags & CPF_OutParm))
-					FError::Throwf(TEXT("Service request functions cannot contain out parameters, unless marked NotReplicated"));
+				{
+					UE_LOG_ERROR_UHT(TEXT("Service request functions cannot contain out parameters, unless marked NotReplicated"));
+				}
 
 				if (!(Property.PropertyFlags & CPF_RepSkip) && (Prop->GetClass()->ClassCastFlags & CASTCLASS_UDelegateProperty) != 0)
-					FError::Throwf(TEXT("Service request functions cannot contain delegate parameters, unless marked NotReplicated"));
+				{
+					UE_LOG_ERROR_UHT(TEXT("Service request functions cannot contain delegate parameters, unless marked NotReplicated"));
+				}
 			}
 		}
 		if ((Function->FunctionFlags & (FUNC_BlueprintEvent|FUNC_BlueprintCallable)) != 0)
@@ -5779,7 +5804,7 @@ UDelegateFunction* FHeaderParser::CompileDelegateDeclaration(FClasses& AllClasse
 	// Multi-cast delegate function signatures are not allowed to have a return value
 	if (bHasReturnValue && bIsMulticast)
 	{
-		FError::Throwf(TEXT("Multi-cast delegates function signatures must not return a value"));
+		UE_LOG_ERROR_UHT(TEXT("Multi-cast delegates function signatures must not return a value"));
 	}
 
 	// Delegate signature
@@ -6080,7 +6105,7 @@ void FHeaderParser::CompileFunctionDeclaration(FClasses& AllClasses)
 	if ((FuncInfo.FunctionFlags & FUNC_BlueprintPure) && GetCurrentClass()->HasAnyClassFlags(CLASS_Interface))
 	{
 		// Until pure interface casts are supported, we don't allow pures in interfaces
-		FError::Throwf(TEXT("BlueprintPure specifier is not allowed for interface functions"));
+		UE_LOG_ERROR_UHT(TEXT("BlueprintPure specifier is not allowed for interface functions"));
 	}
 
 	if (FuncInfo.FunctionFlags & FUNC_Net)
@@ -6122,7 +6147,7 @@ void FHeaderParser::CompileFunctionDeclaration(FClasses& AllClasses)
 			// To allow for quick iteration, don't enforce the requirement that game functions have to be categorized
 			if (!bModuleIsGame)
 			{
-				FError::Throwf(TEXT("Blueprint accessible functions must have a category specified"));
+				UE_LOG_ERROR_UHT(TEXT("Blueprint accessible functions must have a category specified"));
 			}
 		} 
 	}
@@ -6136,7 +6161,7 @@ void FHeaderParser::CompileFunctionDeclaration(FClasses& AllClasses)
 			// Ensure that blueprint events are only allowed in implementable interfaces. Internal only functions allowed
 			if (!bCanImplementInBlueprints && !bInternalOnly)
 			{
-				FError::Throwf(TEXT("Interfaces that are not implementable in blueprints cannot have BlueprintImplementableEvent members."));
+				UE_LOG_ERROR_UHT(TEXT("Interfaces that are not implementable in blueprints cannot have BlueprintImplementableEvent members."));
 			}
 		}
 		
@@ -6145,7 +6170,7 @@ void FHeaderParser::CompileFunctionDeclaration(FClasses& AllClasses)
 			// Ensure that if this interface contains blueprint callable functions that are not blueprint defined, that it must be implemented natively
 			if (bCanImplementInBlueprints)
 			{
-				FError::Throwf(TEXT("Blueprint implementable interfaces cannot contain BlueprintCallable functions that are not BlueprintImplementableEvents.  Use CannotImplementInterfaceInBlueprint on the interface if you wish to keep this function."));
+				UE_LOG_ERROR_UHT(TEXT("Blueprint implementable interfaces cannot contain BlueprintCallable functions that are not BlueprintImplementableEvents.  Use CannotImplementInterfaceInBlueprint on the interface if you wish to keep this function."));
 			}
 		}
 	}
@@ -6203,7 +6228,7 @@ void FHeaderParser::CompileFunctionDeclaration(FClasses& AllClasses)
 			// if this is a BlueprintNativeEvent, make sure it's not "virtual"
 			else if (FuncInfo.FunctionFlags & FUNC_Native)
 			{
-				FError::Throwf(TEXT("BlueprintNativeEvent functions must be non-virtual."));
+				UE_LOG_ERROR_UHT(TEXT("BlueprintNativeEvent functions must be non-virtual."));
 			}
 
 			else
@@ -6230,7 +6255,7 @@ void FHeaderParser::CompileFunctionDeclaration(FClasses& AllClasses)
 
 		if (GetCurrentClass()->HasAnyClassFlags(CLASS_Interface))
 		{
-			FError::Throwf(TEXT("Interface functions cannot be declared 'final'"));
+			UE_LOG_ERROR_UHT(TEXT("Interface functions cannot be declared 'final'"));
 		}
 	}
 
@@ -6264,12 +6289,16 @@ void FHeaderParser::CompileFunctionDeclaration(FClasses& AllClasses)
 	{
 		bool bIsNetService = !!(FuncInfo.FunctionFlags & (FUNC_NetRequest | FUNC_NetResponse));
 		if (bHasReturnValue && !bIsNetService)
+		{
 			FError::Throwf(TEXT("Replicated functions can't have return values"));
+		}
 
 		if (FuncInfo.RPCId > 0)
 		{
 			if (FString* ExistingFunc = UsedRPCIds.Find(FuncInfo.RPCId))
+			{
 				FError::Throwf(TEXT("Function %s already uses identifier %d"), **ExistingFunc, FuncInfo.RPCId);
+			}
 
 			UsedRPCIds.Add(FuncInfo.RPCId, FuncInfo.Function.Identifier);
 			if (FuncInfo.FunctionFlags & FUNC_NetResponse)
@@ -6364,7 +6393,7 @@ void FHeaderParser::CompileFunctionDeclaration(FClasses& AllClasses)
 		// we don't want to prevent compilation:
 		if (!bClassGeneratedFromBP)
 		{
-			FError::Throwf(TEXT("BlueprintPure specifier is not allowed for functions with no return value and no output parameters."));
+			UE_LOG_ERROR_UHT(TEXT("BlueprintPure specifier is not allowed for functions with no return value and no output parameters."));
 		}
 	}
 
@@ -6580,7 +6609,9 @@ void FHeaderParser::CompileFunctionDeclaration(FClasses& AllClasses)
 
 			if (!IsPropertySupportedByBlueprint(Param, false))
 			{
-				FError::Throwf(TEXT("Type '%s' is not supported by blueprint. Function: %s Parameter %s\n"), *Param->GetCPPType(), *TopFunction->GetName(), *Param->GetName());
+				FString ExtendedCPPType;
+				FString CPPType = Param->GetCPPType(&ExtendedCPPType);
+				UE_LOG_ERROR_UHT(TEXT("Type '%s%s' is not supported by blueprint. %s.%s"), *CPPType, *ExtendedCPPType, *TopFunction->GetName(), *Param->GetName());
 			}
 		}
 	}
@@ -6717,7 +6748,7 @@ void FHeaderParser::ValidatePropertyIsDeprecatedIfNecessary(FPropertyBase& VarPr
 	if ( VarProperty.MetaClass != NULL && VarProperty.MetaClass->HasAnyClassFlags(CLASS_Deprecated) && !(VarProperty.PropertyFlags & CPF_Deprecated) &&
 		(OuterPropertyType == NULL || !(OuterPropertyType->PropertyFlags & CPF_Deprecated)) )
 	{
-		FError::Throwf(TEXT("Property is using a deprecated class: %s.  Property should be marked deprecated as well."), *VarProperty.MetaClass->GetPathName());
+		UE_LOG_ERROR_UHT(TEXT("Property is using a deprecated class: %s.  Property should be marked deprecated as well."), *VarProperty.MetaClass->GetPathName());
 	}
 
 	// check to see if we have a UObjectProperty using a deprecated class.
@@ -6727,7 +6758,7 @@ void FHeaderParser::ValidatePropertyIsDeprecatedIfNecessary(FPropertyBase& VarPr
 		&& (VarProperty.PropertyFlags&CPF_Deprecated) == 0					// and this property isn't marked deprecated as well
 		&& (OuterPropertyType == NULL || !(OuterPropertyType->PropertyFlags & CPF_Deprecated)) ) // and this property isn't in an array that was marked deprecated either
 	{
-		FError::Throwf(TEXT("Property is using a deprecated class: %s.  Property should be marked deprecated as well."), *VarProperty.PropertyClass->GetPathName());
+		UE_LOG_ERROR_UHT(TEXT("Property is using a deprecated class: %s.  Property should be marked deprecated as well."), *VarProperty.PropertyClass->GetPathName());
 	}
 }
 
@@ -6789,7 +6820,7 @@ void FHeaderParser::CompileVariableDeclaration(FClasses& AllClasses, UStruct* St
 		}
 		else
 		{
-			FError::Throwf(TEXT("Property is exposed to the editor or blueprints but has no Category specified."));
+			UE_LOG_ERROR_UHT(TEXT("Property is exposed to the editor or blueprints but has no Category specified."));
 		}
 	}
 
@@ -6804,7 +6835,7 @@ void FHeaderParser::CompileVariableDeclaration(FClasses& AllClasses, UStruct* St
 	// If the category was specified explicitly, it wins
 	if (Category && !(OriginalProperty.PropertyFlags & (CPF_Edit|CPF_BlueprintVisible|CPF_BlueprintAssignable|CPF_BlueprintCallable)))
 	{
-		FError::Throwf(TEXT("Property has a Category set but is not exposed to the editor or Blueprints with EditAnywhere, BlueprintReadWrite, VisibleAnywhere, BlueprintReadOnly, BlueprintAssignable, BlueprintCallable keywords.\r\n"));
+		UE_LOG_WARNING_UHT(TEXT("Property has a Category set but is not exposed to the editor or Blueprints with EditAnywhere, BlueprintReadWrite, VisibleAnywhere, BlueprintReadOnly, BlueprintAssignable, BlueprintCallable keywords.\r\n"));
 	}
 
 	// Make sure that editblueprint variables are editable
@@ -6812,12 +6843,12 @@ void FHeaderParser::CompileVariableDeclaration(FClasses& AllClasses, UStruct* St
 	{
 		if (OriginalProperty.PropertyFlags & CPF_DisableEditOnInstance)
 		{
-			FError::Throwf(TEXT("Property cannot have 'DisableEditOnInstance' without being editable"));
+			UE_LOG_ERROR_UHT(TEXT("Property cannot have 'DisableEditOnInstance' without being editable"));
 		}
 
 		if (OriginalProperty.PropertyFlags & CPF_DisableEditOnTemplate)
 		{
-			FError::Throwf(TEXT("Property cannot have 'DisableEditOnTemplate' without being editable"));
+			UE_LOG_ERROR_UHT(TEXT("Property cannot have 'DisableEditOnTemplate' without being editable"));
 		}
 	}
 
@@ -6831,7 +6862,7 @@ void FHeaderParser::CompileVariableDeclaration(FClasses& AllClasses, UStruct* St
 	{
 		if ((*ExposeOnSpawnValue == TEXT("true")) && !FExposeOnSpawnValidator::IsSupported(OriginalProperty))
 		{
-			FError::Throwf(TEXT("ExposeOnSpawn - Property cannot be exposed"));
+			UE_LOG_ERROR_UHT(TEXT("ExposeOnSpawn - Property cannot be exposed"));
 		}
 	}
 
@@ -6880,14 +6911,24 @@ void FHeaderParser::CompileVariableDeclaration(FClasses& AllClasses, UStruct* St
 			}
 		}
 
-		if (NewProperty->HasAnyPropertyFlags(CPF_BlueprintVisible) && (NewProperty->ArrayDim > 1))
+		if (NewProperty->HasAnyPropertyFlags(CPF_BlueprintVisible))
 		{
-			FError::Throwf(TEXT("Static array cannot be exposed to blueprint %s.%s"), *Struct->GetName(), *NewProperty->GetName());
-		}
+			if (Struct->IsA<UScriptStruct>() && !Struct->GetBoolMetaDataHierarchical(TEXT("BlueprintType")))
+			{
+				UE_LOG_ERROR_UHT(TEXT("Cannot expose property to blueprints in a struct that is not a BlueprintType. %s.%s"), *Struct->GetName(), *NewProperty->GetName());
+			}
 
-		if (NewProperty->HasAnyPropertyFlags(CPF_BlueprintVisible) && !IsPropertySupportedByBlueprint(NewProperty, true))
-		{
-			FError::Throwf(TEXT("Type '%s' is not supported by blueprint. %s.%s"), *NewProperty->GetCPPType(), *Struct->GetName(), *NewProperty->GetName());
+			if (NewProperty->ArrayDim > 1)
+			{
+				UE_LOG_ERROR_UHT(TEXT("Static array cannot be exposed to blueprint %s.%s"), *Struct->GetName(), *NewProperty->GetName());
+			}
+
+			if (!IsPropertySupportedByBlueprint(NewProperty, true))
+			{
+				FString ExtendedCPPType;
+				FString CPPType = NewProperty->GetCPPType(&ExtendedCPPType);
+				UE_LOG_ERROR_UHT(TEXT("Type '%s%s' is not supported by blueprint. %s.%s"), *CPPType, *ExtendedCPPType, *Struct->GetName(), *NewProperty->GetName());
+			}
 		}
 
 	} while( MatchSymbol(TEXT(",")) );
@@ -7996,7 +8037,7 @@ void FHeaderParser::SimplifiedClassParse(const TCHAR* Filename, const TCHAR* InB
 		{
 			if (bFoundGeneratedInclude)
 			{
-				FError::Throwf(TEXT("#include found after .generated.h file - the .generated.h file should always be the last #include in a header"));
+				UE_LOG_ERROR_UHT(TEXT("#include found after .generated.h file - the .generated.h file should always be the last #include in a header"));
 			}
 
 			// Handle #include directives as if they were 'dependson' keywords.
@@ -8160,7 +8201,7 @@ void FHeaderParser::SimplifiedClassParse(const TCHAR* Filename, const TCHAR* InB
 
 	if (bFoundExportedClasses && !bFoundGeneratedInclude)
 	{
-		FError::Throwf(TEXT("No #include found for the .generated.h file - the .generated.h file should always be the last #include in a header"));
+		UE_LOG_ERROR_UHT(TEXT("No #include found for the .generated.h file - the .generated.h file should always be the last #include in a header"));
 	}
 }
 
diff --git a/Engine/Source/Programs/UnrealHeaderTool/Private/UHTMakefile/ClassArchiveProxy.cpp b/Engine/Source/Programs/UnrealHeaderTool/Private/UHTMakefile/ClassArchiveProxy.cpp
--- a/Engine/Source/Programs/UnrealHeaderTool/Private/UHTMakefile/ClassArchiveProxy.cpp
+++ b/Engine/Source/Programs/UnrealHeaderTool/Private/UHTMakefile/ClassArchiveProxy.cpp
@@ -8,7 +8,7 @@ FClassArchiveProxy::FClassArchiveProxy(FUHTMakefile& UHTMakefile, const UClass*
 	: FStructArchiveProxy(UHTMakefile, Class)
 {
 	ClassUnique = Class->ClassUnique;
-	ClassFlags = Class->ClassFlags;
+	ClassFlags = Class->GetClassFlags();
 	ClassCastFlags = Class->ClassCastFlags;
 	ClassWithinIndex = UHTMakefile.GetClassIndex(Class->ClassWithin);
 	ClassGeneratedByIndex = UHTMakefile.GetObjectIndex(Class->ClassGeneratedBy);
@@ -82,7 +82,7 @@ FArchive& operator<<(FArchive& Ar, FClassArchiveProxy& ClassArchiveProxy)
 {
 	Ar << static_cast<FStructArchiveProxy&>(ClassArchiveProxy);
 	Ar << ClassArchiveProxy.ClassUnique;
-	Ar << ClassArchiveProxy.ClassFlags;
+	Ar << (uint32&)ClassArchiveProxy.ClassFlags;
 	Ar << ClassArchiveProxy.ClassCastFlags;
 	Ar << ClassArchiveProxy.ClassWithinIndex;
 	Ar << ClassArchiveProxy.ClassGeneratedByIndex;
diff --git a/Engine/Source/Programs/UnrealHeaderTool/Private/UHTMakefile/ClassArchiveProxy.h b/Engine/Source/Programs/UnrealHeaderTool/Private/UHTMakefile/ClassArchiveProxy.h
--- a/Engine/Source/Programs/UnrealHeaderTool/Private/UHTMakefile/ClassArchiveProxy.h
+++ b/Engine/Source/Programs/UnrealHeaderTool/Private/UHTMakefile/ClassArchiveProxy.h
@@ -26,7 +26,7 @@ struct FClassArchiveProxy : public FStructArchiveProxy
 	static void AddReferencedNames(const UClass* Class, FUHTMakefile& UHTMakefile);
 
 	int32 ClassUnique;
-	uint32 ClassFlags;
+	EClassFlags ClassFlags;
 	EClassCastFlags ClassCastFlags;
 	FSerializeIndex ClassWithinIndex;
 	FSerializeIndex ClassGeneratedByIndex;
diff --git a/Engine/Source/Programs/UnrealHeaderTool/Private/UnrealHeaderToolGlobals.h b/Engine/Source/Programs/UnrealHeaderTool/Private/UnrealHeaderToolGlobals.h
--- a/Engine/Source/Programs/UnrealHeaderTool/Private/UnrealHeaderToolGlobals.h
+++ b/Engine/Source/Programs/UnrealHeaderTool/Private/UnrealHeaderToolGlobals.h
@@ -9,5 +9,7 @@
 DECLARE_LOG_CATEGORY_EXTERN(LogCompile, Log, All);
 
 extern bool GUHTWarningLogged;
+extern bool GUHTErrorLogged;
 
 #define UE_LOG_WARNING_UHT(Format, ...) { GUHTWarningLogged = true; UE_LOG(LogCompile, Warning, Format, ##__VA_ARGS__); }
+#define UE_LOG_ERROR_UHT(Format, ...) { GUHTErrorLogged = true; UE_LOG(LogCompile, Error, Format, ##__VA_ARGS__); }
diff --git a/Engine/Source/Programs/UnrealHeaderTool/Private/UnrealHeaderToolMain.cpp b/Engine/Source/Programs/UnrealHeaderTool/Private/UnrealHeaderToolMain.cpp
--- a/Engine/Source/Programs/UnrealHeaderTool/Private/UnrealHeaderToolMain.cpp
+++ b/Engine/Source/Programs/UnrealHeaderTool/Private/UnrealHeaderToolMain.cpp
@@ -35,6 +35,7 @@ IMPLEMENT_APPLICATION(UnrealHeaderTool, "UnrealHeaderTool");
 DEFINE_LOG_CATEGORY(LogCompile);
 
 bool GUHTWarningLogged = false;
+bool GUHTErrorLogged = false;
 
 /**
  * Application entry point
@@ -125,7 +126,7 @@ INT32_MAIN_INT32_ARGC_TCHAR_ARGV()
 	extern ECompilationResult::Type UnrealHeaderTool_Main(const FString& ModuleInfoFilename);
 	ECompilationResult::Type Result = UnrealHeaderTool_Main(ModuleInfoFilename);
 
-	if (Result == ECompilationResult::Succeeded && GUHTWarningLogged && GWarn->TreatWarningsAsErrors)
+	if (Result == ECompilationResult::Succeeded && (GUHTErrorLogged || (GUHTWarningLogged && GWarn->TreatWarningsAsErrors)))
 	{
 		Result = ECompilationResult::OtherCompilationError;
 	}
diff --git a/Engine/Source/Runtime/AIModule/Classes/EnvironmentQuery/EnvQueryManager.h b/Engine/Source/Runtime/AIModule/Classes/EnvironmentQuery/EnvQueryManager.h
--- a/Engine/Source/Runtime/AIModule/Classes/EnvironmentQuery/EnvQueryManager.h
+++ b/Engine/Source/Runtime/AIModule/Classes/EnvironmentQuery/EnvQueryManager.h
@@ -226,8 +226,8 @@ class AIMODULE_API UEnvQueryManager : public UObject, public FTickableGameObject
 	static UEnvQueryManager* GetCurrent(UWorld* World);
 	static UEnvQueryManager* GetCurrent(const UObject* WorldContextObject);
 	
-	UFUNCTION(BlueprintCallable, Category = "AI|EQS", meta = (WorldContext = "WorldContext", AdvancedDisplay = "WrapperClass"))
-	static UEnvQueryInstanceBlueprintWrapper* RunEQSQuery(UObject* WorldContext, UEnvQuery* QueryTemplate, UObject* Querier, TEnumAsByte<EEnvQueryRunMode::Type> RunMode, TSubclassOf<UEnvQueryInstanceBlueprintWrapper> WrapperClass);
+	UFUNCTION(BlueprintCallable, Category = "AI|EQS", meta = (WorldContext = "WorldContextObject", AdvancedDisplay = "WrapperClass"))
+	static UEnvQueryInstanceBlueprintWrapper* RunEQSQuery(UObject* WorldContextObject, UEnvQuery* QueryTemplate, UObject* Querier, TEnumAsByte<EEnvQueryRunMode::Type> RunMode, TSubclassOf<UEnvQueryInstanceBlueprintWrapper> WrapperClass);
 
 	void RegisterActiveWrapper(UEnvQueryInstanceBlueprintWrapper& Wrapper);
 	void UnregisterActiveWrapper(UEnvQueryInstanceBlueprintWrapper& Wrapper);
diff --git a/Engine/Source/Runtime/AIModule/Classes/Navigation/NavLocalGridManager.h b/Engine/Source/Runtime/AIModule/Classes/Navigation/NavLocalGridManager.h
--- a/Engine/Source/Runtime/AIModule/Classes/Navigation/NavLocalGridManager.h
+++ b/Engine/Source/Runtime/AIModule/Classes/Navigation/NavLocalGridManager.h
@@ -80,28 +80,28 @@ public:
 	/** get version of grid data, incremented with each rebuild */
 	int32 GetVersion() const { return VersionNum; }
 
-	UFUNCTION(BlueprintCallable, Category = "AI|Navigation", meta = (WorldContext = "WorldContext"))
-	static bool SetLocalNavigationGridDensity(UObject* WorldContext, float CellSize);
+	UFUNCTION(BlueprintCallable, Category = "AI|Navigation", meta = (WorldContext = "WorldContextObject"))
+	static bool SetLocalNavigationGridDensity(UObject* WorldContextObject, float CellSize);
 
 	/** creates new grid data for single point */
-	UFUNCTION(BlueprintCallable, Category = "AI|Navigation", meta = (WorldContext = "WorldContext"))
-	static int32 AddLocalNavigationGridForPoint(UObject* WorldContext, const FVector& Location, const int32 Radius2D = 5, const float Height = 100.0f, bool bRebuildGrids = true);
+	UFUNCTION(BlueprintCallable, Category = "AI|Navigation", meta = (WorldContext = "WorldContextObject"))
+	static int32 AddLocalNavigationGridForPoint(UObject* WorldContextObject, const FVector& Location, const int32 Radius2D = 5, const float Height = 100.0f, bool bRebuildGrids = true);
 
 	/** creates single grid data for set of points */
-	UFUNCTION(BlueprintCallable, Category = "AI|Navigation", meta = (WorldContext = "WorldContext"))
-	static int32 AddLocalNavigationGridForPoints(UObject* WorldContext, const TArray<FVector>& Locations, const int32 Radius2D = 5, const float Height = 100.0f, bool bRebuildGrids = true);
+	UFUNCTION(BlueprintCallable, Category = "AI|Navigation", meta = (WorldContext = "WorldContextObject"))
+	static int32 AddLocalNavigationGridForPoints(UObject* WorldContextObject, const TArray<FVector>& Locations, const int32 Radius2D = 5, const float Height = 100.0f, bool bRebuildGrids = true);
 
-	UFUNCTION(BlueprintCallable, Category = "AI|Navigation", meta = (WorldContext = "WorldContext"))
-	static int32 AddLocalNavigationGridForBox(UObject* WorldContext, const FVector& Location, FVector Extent = FVector(1,1,1), FRotator Rotation = FRotator::ZeroRotator, const int32 Radius2D = 5, const float Height = 100.0f, bool bRebuildGrids = true);
+	UFUNCTION(BlueprintCallable, Category = "AI|Navigation", meta = (WorldContext = "WorldContextObject"))
+	static int32 AddLocalNavigationGridForBox(UObject* WorldContextObject, const FVector& Location, FVector Extent = FVector(1,1,1), FRotator Rotation = FRotator::ZeroRotator, const int32 Radius2D = 5, const float Height = 100.0f, bool bRebuildGrids = true);
 
-	UFUNCTION(BlueprintCallable, Category = "AI|Navigation", meta = (WorldContext = "WorldContext"))
-	static int32 AddLocalNavigationGridForCapsule(UObject* WorldContext, const FVector& Location, float CapsuleRadius, float CapsuleHalfHeight, const int32 Radius2D = 5, const float Height = 100.0f, bool bRebuildGrids = true);
+	UFUNCTION(BlueprintCallable, Category = "AI|Navigation", meta = (WorldContext = "WorldContextObject"))
+	static int32 AddLocalNavigationGridForCapsule(UObject* WorldContextObject, const FVector& Location, float CapsuleRadius, float CapsuleHalfHeight, const int32 Radius2D = 5, const float Height = 100.0f, bool bRebuildGrids = true);
 
-	UFUNCTION(BlueprintCallable, Category = "AI|Navigation", meta = (WorldContext = "WorldContext"))
-	static void RemoveLocalNavigationGrid(UObject* WorldContext, int32 GridId, bool bRebuildGrids = true);
+	UFUNCTION(BlueprintCallable, Category = "AI|Navigation", meta = (WorldContext = "WorldContextObject"))
+	static void RemoveLocalNavigationGrid(UObject* WorldContextObject, int32 GridId, bool bRebuildGrids = true);
 
-	UFUNCTION(BlueprintCallable, Category = "AI|Navigation", meta = (WorldContext = "WorldContext"))
-	static bool FindLocalNavigationGridPath(UObject* WorldContext, const FVector& Start, const FVector& End, TArray<FVector>& PathPoints);
+	UFUNCTION(BlueprintCallable, Category = "AI|Navigation", meta = (WorldContext = "WorldContextObject"))
+	static bool FindLocalNavigationGridPath(UObject* WorldContextObject, const FVector& Start, const FVector& End, TArray<FVector>& PathPoints);
 
 	static UNavLocalGridManager* GetCurrent(UWorld* World);
 	static UNavLocalGridManager* GetCurrent(const UObject* WorldContextObject);
