Fixes to generated code sizes in non-editor builds - merged from CL#s 3446070 and 3446091.

diff --git a/Engine/Source/Programs/UnrealHeaderTool/Private/CodeGenerator.cpp b/Engine/Source/Programs/UnrealHeaderTool/Private/CodeGenerator.cpp
--- a/Engine/Source/Programs/UnrealHeaderTool/Private/CodeGenerator.cpp
+++ b/Engine/Source/Programs/UnrealHeaderTool/Private/CodeGenerator.cpp
@@ -1796,7 +1796,7 @@ void FNativeClassHeaderGenerator::ExportNatives(FOutputDevice& Out, FClass* Clas
 			}
 
 			Out.Log(TEXT("\t\t};\r\n"));
-			Out.Logf(TEXT("\t\tFNativeFunctionRegistrar::RegisterFunctions(Class, AnsiFuncs, %d);\r\n"), AnsiNamedFunctionsToExport.Num());
+			Out.Logf(TEXT("\t\tFNativeFunctionRegistrar::RegisterFunctions(Class, AnsiFuncs, ARRAY_COUNT(AnsiFuncs));\r\n"));
 		}
 
 		if (TCharNamedFunctionsToExport.Num())
@@ -1818,7 +1818,7 @@ void FNativeClassHeaderGenerator::ExportNatives(FOutputDevice& Out, FClass* Clas
 			}
 
 			Out.Log(TEXT("\t\t};\r\n"));
-			Out.Logf(TEXT("\t\tFNativeFunctionRegistrar::RegisterFunctions(Class, TCharFuncs, %d);\r\n"), TCharNamedFunctionsToExport.Num());
+			Out.Logf(TEXT("\t\tFNativeFunctionRegistrar::RegisterFunctions(Class, TCharFuncs, ARRAY_COUNT(TCharFuncs));\r\n"));
 		}
 	}
 
