Merging CL# 3449001 from //UE4/Dev-Core to Main (//UE4/Main).

.generated shortened to .gen.

diff --git a/Engine/Source/Programs/UnrealBuildTool/Configuration/UEBuildModuleCPP.cs b/Engine/Source/Programs/UnrealBuildTool/Configuration/UEBuildModuleCPP.cs
--- a/Engine/Source/Programs/UnrealBuildTool/Configuration/UEBuildModuleCPP.cs
+++ b/Engine/Source/Programs/UnrealBuildTool/Configuration/UEBuildModuleCPP.cs
@@ -18,7 +18,7 @@ namespace UnrealBuildTool
 			public class BuildInfoClass
 			{
 				/// <summary>
-				/// The wildcard of the *.generated.cpp file which was generated for the module
+				/// The wildcard of the *.gen.cpp file which was generated for the module
 				/// </summary>
 				public readonly string FileWildcard;
 
@@ -31,7 +31,7 @@ namespace UnrealBuildTool
 			}
 
 			/// <summary>
-			/// Information about how to build the .generated.cpp files. If this is null, then we're not building .generated.cpp files for this module.
+			/// Information about how to build the .gen.cpp files. If this is null, then we're not building .gen.cpp files for this module.
 			/// </summary>
 			public BuildInfoClass BuildInfo;
 
@@ -42,7 +42,7 @@ namespace UnrealBuildTool
 		}
 
 		/// <summary>
-		/// Information about the .generated.cpp file.  If this is null then this module doesn't have any UHT-produced code.
+		/// Information about the .gen.cpp file.  If this is null then this module doesn't have any UHT-produced code.
 		/// </summary>
 		public AutoGenerateCppInfoClass AutoGenerateCppInfo = null;
 
@@ -164,7 +164,7 @@ namespace UnrealBuildTool
 		{
 			return ((null == GeneratedCodeDirectory) || !DirectoryReference.Exists(GeneratedCodeDirectory))
 				? Enumerable.Empty<string>()
-				: DirectoryReference.EnumerateFiles(GeneratedCodeDirectory, "*.generated.cpp", SearchOption.TopDirectoryOnly).Select((Dir) => Dir.FullName);
+				: DirectoryReference.EnumerateFiles(GeneratedCodeDirectory, "*.gen.cpp", SearchOption.TopDirectoryOnly).Select((Dir) => Dir.FullName);
 		}
 
 		protected override void GetReferencedDirectories(HashSet<DirectoryReference> Directories)
@@ -557,7 +557,7 @@ namespace UnrealBuildTool
 
 					if (bModuleUsesUnityBuild)
 					{
-						GeneratedFileItems = Unity.GenerateUnityCPPs(Target, GeneratedFileItems, GeneratedCPPCompileEnvironment, Name + ".generated");
+						GeneratedFileItems = Unity.GenerateUnityCPPs(Target, GeneratedFileItems, GeneratedCPPCompileEnvironment, Name + ".gen");
 						LinkInputFiles.AddRange(CompileUnityFilesWithToolChain(Target, ToolChain, GeneratedCPPCompileEnvironment, ModuleCompileEnvironment, GeneratedFileItems, ActionGraph).ObjectFiles);
 					}
 					else
diff --git a/Engine/Source/Programs/UnrealBuildTool/Configuration/UEBuildTarget.cs b/Engine/Source/Programs/UnrealBuildTool/Configuration/UEBuildTarget.cs
--- a/Engine/Source/Programs/UnrealBuildTool/Configuration/UEBuildTarget.cs
+++ b/Engine/Source/Programs/UnrealBuildTool/Configuration/UEBuildTarget.cs
@@ -3241,7 +3241,7 @@ namespace UnrealBuildTool
 				{
 					bFound = true;
 
-					string FunctionName = "EmptyLinkFunctionForGeneratedCode" + Path.GetFileName(CppFile).Replace(".generated.cpp", "").Replace(".", "_");
+					string FunctionName = "EmptyLinkFunctionForGeneratedCode" + Path.GetFileName(CppFile).Replace(".gen.cpp", "").Replace(".", "_");
 					if (AlreadyAddedEmptyLinkFunctions.Add(FunctionName))
 					{
 						Result.Add("    extern void " + FunctionName + "();");
@@ -4061,7 +4061,7 @@ namespace UnrealBuildTool
 
 			ToolChain.SetUpGlobalEnvironment(Rules);
 
-			// @Hack: This to prevent UHT from listing CoreUObject.generated.cpp as its dependency.
+			// @Hack: This to prevent UHT from listing CoreUObject.init.gen.cpp as its dependency.
 			// We flag the compile environment when we build UHT so that we don't need to check
 			// this for each file when generating their dependencies.
 			GlobalCompileEnvironment.bHackHeaderGenerator = (GetAppName() == "UnrealHeaderTool");
diff --git a/Engine/Source/Programs/UnrealBuildTool/System/ExternalExecution.cs b/Engine/Source/Programs/UnrealBuildTool/System/ExternalExecution.cs
--- a/Engine/Source/Programs/UnrealBuildTool/System/ExternalExecution.cs
+++ b/Engine/Source/Programs/UnrealBuildTool/System/ExternalExecution.cs
@@ -181,7 +181,7 @@ namespace UnrealBuildTool
 		public string PCH;
 
 		/// <summary>
-		/// Base (i.e. extensionless) path+filename of the .generated files
+		/// Base (i.e. extensionless) path+filename of the .gen files
 		/// </summary>
 		public string GeneratedCPPFilenameBase;
 
@@ -489,12 +489,12 @@ namespace UnrealBuildTool
 				if (UHTModuleInfo.Info.PublicUObjectClassesHeaders.Count > 0 || UHTModuleInfo.Info.PrivateUObjectHeaders.Count > 0 || UHTModuleInfo.Info.PublicUObjectHeaders.Count > 0)
 				{
 					// If we've got this far and there are no source files then it's likely we're installed and ignoring
-					// engine files, so we don't need a .generated.cpp either
+					// engine files, so we don't need a .gen.cpp either
 					UEBuildModuleCPP.AutoGenerateCppInfoClass.BuildInfoClass BuildInfo = null;
-					UHTModuleInfo.Info.GeneratedCPPFilenameBase = Path.Combine(Module.GeneratedCodeDirectory.FullName, UHTModuleInfo.Info.ModuleName) + ".generated";
+					UHTModuleInfo.Info.GeneratedCPPFilenameBase = Path.Combine(Module.GeneratedCodeDirectory.FullName, UHTModuleInfo.Info.ModuleName) + ".gen";
 					if (Module.SourceFilesToBuild.Count != 0)
 					{
-						BuildInfo = new UEBuildModuleCPP.AutoGenerateCppInfoClass.BuildInfoClass(Path.Combine(Module.GeneratedCodeDirectory.FullName, "*.generated.cpp"));
+						BuildInfo = new UEBuildModuleCPP.AutoGenerateCppInfoClass.BuildInfoClass(Path.Combine(Module.GeneratedCodeDirectory.FullName, "*.gen.cpp"));
 					}
 
 					Module.AutoGenerateCppInfo = new UEBuildModuleCPP.AutoGenerateCppInfoClass(BuildInfo);
@@ -673,9 +673,9 @@ namespace UnrealBuildTool
 		}
 
 		/// <summary>
-		/// Gets the timestamp of CoreUObject.generated.cpp file.
+		/// Gets the timestamp of CoreUObject.gen.cpp file.
 		/// </summary>
-		/// <returns>Last write time of CoreUObject.generated.cpp or DateTime.MaxValue if it doesn't exist.</returns>
+		/// <returns>Last write time of CoreUObject.gen.cpp or DateTime.MaxValue if it doesn't exist.</returns>
 		private static DateTime GetCoreGeneratedTimestamp(string ModuleName, string ModuleGeneratedCodeDirectory)
 		{
 			// In Installed Builds, we don't check the timestamps on engine headers.  Default to a very old date.
@@ -684,8 +684,8 @@ namespace UnrealBuildTool
 				return DateTime.MinValue;
 			}
 
-			// Otherwise look for CoreUObject.generated.Init.cpp
-			FileInfo CoreGeneratedFileInfo = new FileInfo(Path.Combine(ModuleGeneratedCodeDirectory, ModuleName + ".init.generated.cpp"));
+			// Otherwise look for CoreUObject.init.gen.cpp
+			FileInfo CoreGeneratedFileInfo = new FileInfo(Path.Combine(ModuleGeneratedCodeDirectory, ModuleName + ".init.gen.cpp"));
 			if (CoreGeneratedFileInfo.Exists)
 			{
 				return CoreGeneratedFileInfo.LastWriteTime;
@@ -704,7 +704,7 @@ namespace UnrealBuildTool
 		/// <returns>    True if the code files are out of date</returns>
 		private static bool AreGeneratedCodeFilesOutOfDate(List<UHTModuleInfo> UObjectModules, DateTime HeaderToolTimestamp, bool bUsePrecompiled)
 		{
-			// Get CoreUObject.generated.cpp timestamp.  If the source files are older than the CoreUObject generated code, we'll
+			// Get CoreUObject.init.gen.cpp timestamp.  If the source files are older than the CoreUObject generated code, we'll
 			// need to regenerate code for the module
 			DateTime? CoreGeneratedTimestamp = null;
 			{
diff --git a/Engine/Source/Programs/UnrealBuildTool/System/UnrealBuildTool.cs b/Engine/Source/Programs/UnrealBuildTool/System/UnrealBuildTool.cs
--- a/Engine/Source/Programs/UnrealBuildTool/System/UnrealBuildTool.cs
+++ b/Engine/Source/Programs/UnrealBuildTool/System/UnrealBuildTool.cs
@@ -2484,7 +2484,7 @@ namespace UnrealBuildTool
 		public static bool ShouldSourceFileBePartOfWorkingSet(string SourceFileAbsolutePath)
 		{
 			// Generated .cpp files should never be treated as part of the working set
-			if (SourceFileAbsolutePath.EndsWith(".generated.cpp"))
+			if (SourceFileAbsolutePath.EndsWith(".gen.cpp"))
 			{
 				return false;
 			}
diff --git a/Engine/Source/Programs/UnrealHeaderTool/Private/CodeGenerator.cpp b/Engine/Source/Programs/UnrealHeaderTool/Private/CodeGenerator.cpp
--- a/Engine/Source/Programs/UnrealHeaderTool/Private/CodeGenerator.cpp
+++ b/Engine/Source/Programs/UnrealHeaderTool/Private/CodeGenerator.cpp
@@ -4344,7 +4344,7 @@ FNativeClassHeaderGenerator::FNativeClassHeaderGenerator(
 		FString BaseSourceFilename = GeneratedIncludeDirectory / StrippedName;
 
 		FUHTStringBuilder GeneratedHeaderText;
-		FGeneratedCPP& GeneratedCPP = GeneratedCPPs.Emplace(SourceFile, BaseSourceFilename + TEXT(".generated.cpp"));
+		FGeneratedCPP& GeneratedCPP = GeneratedCPPs.Emplace(SourceFile, BaseSourceFilename + TEXT(".gen.cpp"));
 		GeneratedCPP.RelativeIncludes.Add(MoveTemp(ModuleRelativeFilename));
 
 		UniqueCrossModuleReferences = &GeneratedCPP.CrossModuleReferences;
@@ -4495,8 +4495,8 @@ FNativeClassHeaderGenerator::FNativeClassHeaderGenerator(
 	// now export the names for the functions in this package
 	// notice we always export this file (as opposed to only exporting if we have any marked names)
 	// because there would be no way to know when the file was created otherwise
-	// Export .generated.cpp
-	UE_LOG(LogCompile, Log, TEXT("Autogenerating boilerplate cpp: %s.generated.cpp"), *PackageName);
+	// Export .gen.cpp
+	UE_LOG(LogCompile, Log, TEXT("Autogenerating boilerplate cpp: %s.gen.cpp"), *PackageName);
 
 	if (GeneratedFunctionDeclarations.Len())
 	{
@@ -4515,7 +4515,7 @@ FNativeClassHeaderGenerator::FNativeClassHeaderGenerator(
 			}
 		}
 
-		FGeneratedCPP& GeneratedCPP = GeneratedCPPs.Emplace(nullptr, GeneratedIncludeDirectory / FString::Printf(TEXT("%s.init.generated.cpp"), *PackageName));
+		FGeneratedCPP& GeneratedCPP = GeneratedCPPs.Emplace(nullptr, GeneratedIncludeDirectory / FString::Printf(TEXT("%s.init.gen.cpp"), *PackageName));
 		ExportGeneratedPackageInitCode(GeneratedCPP.GeneratedText, *GeneratedFunctionDeclarations, Package, CombinedCRC);
 	}
 
@@ -4545,7 +4545,7 @@ FNativeClassHeaderGenerator::FNativeClassHeaderGenerator(
 		ExportGeneratedCPP(
 			FileText,
 			GeneratedCPP.Value.CrossModuleReferences,
-			*FPaths::GetCleanFilename(GeneratedCPP.Value.GeneratedCppFullFilename).Replace(TEXT(".generated.cpp"), TEXT("")).Replace(TEXT("."), TEXT("_")),
+			*FPaths::GetCleanFilename(GeneratedCPP.Value.GeneratedCppFullFilename).Replace(TEXT(".gen.cpp"), TEXT("")).Replace(TEXT("."), TEXT("_")),
 			*GeneratedCPP.Value.GeneratedText,
 			*GeneratedIncludes
 		);
@@ -4557,27 +4557,29 @@ FNativeClassHeaderGenerator::FNativeClassHeaderGenerator(
 
 	if (bAllowSaveExportedHeaders)
 	{
-	    // Delete old generated .cpp files which we don't need because we generated less code than last time.
-	    TArray<FString> FoundFiles;
-	    FString BaseDir = FPaths::GetPath(ModuleInfo->GeneratedCPPFilenameBase);
-	    IFileManager::Get().FindFiles(FoundFiles, *FPaths::Combine(BaseDir, TEXT("*.generated.cpp")), true, false);
-	    IFileManager::Get().FindFiles(FoundFiles, *FPaths::Combine(BaseDir, TEXT("*.generated.*.cpp")), true, false);
-	    for (FString& File : FoundFiles)
-	    {
-		    if (!GeneratedCPPNames.Contains(File))
-		    {
-			    IFileManager::Get().Delete(*FPaths::Combine(*BaseDir, *File));
-		    }
-	    }
-    
-	    // delete the old .cpp file that will cause link errors if it's left around (Engine.generated.cpp and Engine.generated.1.cpp will 
-	    // conflict now that we no longer use Engine.generated.cpp to #include Engine.generated.1.cpp, and UBT would compile all 3)
-	    // @todo: This is a temp measure so we don't force everyone to require a Clean
-	    if (GeneratedFunctionBodyTextSplit.Num() > 1)
-	    {
-		    FString CppPath = ModuleInfo->GeneratedCPPFilenameBase + TEXT(".cpp");
-		    IFileManager::Get().Delete(*CppPath);
-	    }
+		// Delete old generated .cpp files which we don't need because we generated less code than last time.
+		TArray<FString> FoundFiles;
+		FString BaseDir = FPaths::GetPath(ModuleInfo->GeneratedCPPFilenameBase);
+		IFileManager::Get().FindFiles(FoundFiles, *FPaths::Combine(BaseDir, TEXT("*.generated.cpp")), true, false);
+		IFileManager::Get().FindFiles(FoundFiles, *FPaths::Combine(BaseDir, TEXT("*.generated.*.cpp")), true, false);
+		IFileManager::Get().FindFiles(FoundFiles, *FPaths::Combine(BaseDir, TEXT("*.gen.cpp")), true, false);
+		IFileManager::Get().FindFiles(FoundFiles, *FPaths::Combine(BaseDir, TEXT("*.gen.*.cpp")), true, false);
+		for (FString& File : FoundFiles)
+		{
+			if (!GeneratedCPPNames.Contains(File))
+			{
+				IFileManager::Get().Delete(*FPaths::Combine(*BaseDir, *File));
+			}
+		}
+
+		// delete the old .cpp file that will cause link errors if it's left around (Engine.gen.cpp and Engine.gen.1.cpp will 
+		// conflict now that we no longer use Engine.gen.cpp to #include Engine.gen.1.cpp, and UBT would compile all 3)
+		// @todo: This is a temp measure so we don't force everyone to require a Clean
+		if (GeneratedFunctionBodyTextSplit.Num() > 1)
+		{
+			FString CppPath = ModuleInfo->GeneratedCPPFilenameBase + TEXT(".cpp");
+			IFileManager::Get().Delete(*CppPath);
+		}
 	}
 
 	// Export all changed headers from their temp files to the .h files
@@ -4793,7 +4795,6 @@ void FNativeClassHeaderGenerator::ExportUpdatedHeaders(FString PackageName)
 
 /**
  * Exports C++ definitions for boilerplate that was generated for a package.
- * They are exported to a file using the name <PackageName>.generated.cpp
  */
 void FNativeClassHeaderGenerator::ExportGeneratedCPP(FOutputDevice& Out, const TSet<FString>& InCrossModuleReferences, const TCHAR* EmptyLinkFunctionPostfix, const TCHAR* Body, const TCHAR* OtherIncludes)
 {
@@ -4822,7 +4823,6 @@ void FNativeClassHeaderGenerator::ExportGeneratedCPP(FOutputDevice& Out, const T
 		}
 		Out.Logf(TEXT("// End Cross Module References\r\n"));
 	}
-//	Out.Log(Declarations);
 	Out.Log(Body);
 	Out.Log(EnableDeprecationWarnings);
 	Out.Log(EnableWarning4883);
