Copying //UE4/Dev-Framework to //UE4/Dev-Main (Source: //UE4/Dev-Framework @ 3459469)

diff --git a/Engine/Source/Editor/UnrealEd/Private/Commandlets/DerivedDataCacheCommandlet.cpp b/Engine/Source/Editor/UnrealEd/Private/Commandlets/DerivedDataCacheCommandlet.cpp
--- a/Engine/Source/Editor/UnrealEd/Private/Commandlets/DerivedDataCacheCommandlet.cpp
+++ b/Engine/Source/Editor/UnrealEd/Private/Commandlets/DerivedDataCacheCommandlet.cpp
@@ -201,8 +201,7 @@ int32 UDerivedDataCacheCommandlet::Main( const FString& Params )
 				}
 				if (bDoSubset)
 				{
-					const FString& SubPackageName = FPackageName::PackageFromPath(*Filename);
-					if (FCrc::StrCrc_DEPRECATED(*SubPackageName.ToUpper()) % SubsetMod != SubsetTarget)
+					if (FCrc::StrCrc_DEPRECATED(*PackageName.ToUpper()) % SubsetMod != SubsetTarget)
 					{
 						continue;
 					}
diff --git a/Engine/Source/Editor/UnrealEd/Private/Commandlets/PackageUtilities.cpp b/Engine/Source/Editor/UnrealEd/Private/Commandlets/PackageUtilities.cpp
--- a/Engine/Source/Editor/UnrealEd/Private/Commandlets/PackageUtilities.cpp
+++ b/Engine/Source/Editor/UnrealEd/Private/Commandlets/PackageUtilities.cpp
@@ -211,10 +211,10 @@ bool NormalizePackageNames( TArray<FString> PackageNames, TArray<FString>& Packa
 	if ( (PackageFilter&NORMALIZE_ResetExistingLoaders) != 0 )
 	{
 		// reset the loaders for the packages we want to load so that we don't find the wrong version of the file
-		for ( int32 PackageIndex = 0; PackageIndex < PackageNames.Num(); PackageIndex++ )
+		for ( int32 PackageIndex = 0; PackageIndex < PackagePathNames.Num(); PackageIndex++ )
 		{
 			// (otherwise, attempting to run a commandlet on e.g. Engine.xxx will always return results for Engine.u instead)
-			const FString& PackageName = FPackageName::PackageFromPath(*PackageNames[PackageIndex]);
+			const FString& PackageName = PackagePathNames[PackageIndex];
 			UPackage* ExistingPackage = FindObject<UPackage>(NULL, *PackageName, true);
 			if ( ExistingPackage != NULL )
 			{
@@ -723,24 +723,27 @@ int32 ULoadPackageCommandlet::Main( const FString& Params )
 
 		UE_LOG(LogPackageUtilities, Warning, TEXT("Loading %s"), *Filename );
 
-		const FString& PackageName = FPackageName::PackageFromPath(*Filename);
-		UPackage* Package = FindObject<UPackage>(NULL, *PackageName, true);
-		if ( Package != NULL && !bLoadAllPackages )
+		FString PackageName;
+		if (FPackageName::TryConvertFilenameToLongPackageName(Filename, PackageName))
 		{
-			ResetLoaders(Package);
+			UPackage* Package = FindObject<UPackage>(nullptr, *PackageName, true);
+			if (Package != NULL && !bLoadAllPackages)
+			{
+				ResetLoaders(Package);
+			}
 		}
 
 		if (bCheckForLegacyPackages)
 		{
 			BeginLoad();
-			auto Linker = GetPackageLinker(NULL,*Filename,LOAD_NoVerify,NULL,NULL);
+			auto Linker = GetPackageLinker(nullptr, *Filename,LOAD_NoVerify,NULL,NULL);
 			EndLoad();
 			MinVersion = FMath::Min<int32>(MinVersion, Linker->Summary.GetFileVersionUE4());
 		}
 		else
 		{
-			Package = LoadPackage( NULL, *Filename, LOAD_None );
-			if( Package == NULL )
+			UPackage* Package = LoadPackage(nullptr, *Filename, LOAD_None );
+			if(Package == nullptr)
 			{
 				UE_LOG(LogPackageUtilities, Error, TEXT("Error loading %s!"), *Filename );
 			}
@@ -1496,11 +1499,14 @@ int32 UPkgInfoCommandlet::Main( const FString& Params )
 		{
 			// reset the loaders for the packages we want to load so that we don't find the wrong version of the file
 			// (otherwise, attempting to run pkginfo on e.g. Engine.xxx will always return results for Engine.u instead)
-			const FString& PackageName = FPackageName::PackageFromPath(*Filename);
-			UPackage* ExistingPackage = FindObject<UPackage>(NULL, *PackageName, true);
-			if ( ExistingPackage != NULL )
+			FString PackageName;
+			if (FPackageName::TryConvertFilenameToLongPackageName(Filename, PackageName))
 			{
-				ResetLoaders(ExistingPackage);
+				UPackage* ExistingPackage = FindObject<UPackage>(nullptr, *PackageName, true);
+				if (ExistingPackage != nullptr)
+				{
+					ResetLoaders(ExistingPackage);
+				}
 			}
 		}
 
diff --git a/Engine/Source/Editor/UnrealEd/Private/ObjectTools.cpp b/Engine/Source/Editor/UnrealEd/Private/ObjectTools.cpp
--- a/Engine/Source/Editor/UnrealEd/Private/ObjectTools.cpp
+++ b/Engine/Source/Editor/UnrealEd/Private/ObjectTools.cpp
@@ -4120,10 +4120,13 @@ namespace ThumbnailTools
 	{
 		FObjectThumbnail* FoundThumbnail = NULL;
 
+		FString PackageName = InPackageFileName;
+		FPackageName::TryConvertFilenameToLongPackageName(PackageName, PackageName);
+
 		// First check to see if the package is already in memory.  If it is, some or all of the thumbnails
 		// may already be loaded and ready.
 		UObject* PackageOuter = NULL;
-		UPackage* Package = FindPackage( PackageOuter, *FPackageName::PackageFromPath( *InPackageFileName ) );
+		UPackage* Package = FindPackage( PackageOuter, *PackageName);
 		if( Package != NULL )
 		{
 			FoundThumbnail = FindCachedThumbnailInPackage( Package, InObjectFullName );
diff --git a/Engine/Source/Programs/UnrealHeaderTool/Private/HeaderParser.cpp b/Engine/Source/Programs/UnrealHeaderTool/Private/HeaderParser.cpp
--- a/Engine/Source/Programs/UnrealHeaderTool/Private/HeaderParser.cpp
+++ b/Engine/Source/Programs/UnrealHeaderTool/Private/HeaderParser.cpp
@@ -256,9 +256,10 @@ namespace
 	 * @param FuncInfo   - The FFuncInfo object to populate.
 	 * @param Specifiers - The specifiers to process.
 	 */
-	void ProcessFunctionSpecifiers(FFuncInfo& FuncInfo, const TArray<FPropertySpecifier>& Specifiers)
+	void ProcessFunctionSpecifiers(FFuncInfo& FuncInfo, const TArray<FPropertySpecifier>& Specifiers, TMap<FName, FString>& MetaData)
 	{
 		bool bSpecifiedUnreliable = false;
+		bool bSawPropertyAccessor = false;
 
 		for (const FPropertySpecifier& Specifier : Specifiers)
 		{
@@ -274,12 +275,16 @@ namespace
 				{
 					if (FuncInfo.FunctionFlags & FUNC_Net)
 					{
-						FError::Throwf(TEXT("BlueprintNativeEvent functions cannot be replicated!") );
+						UE_LOG_ERROR_UHT(TEXT("BlueprintNativeEvent functions cannot be replicated!") );
 					}
 					else if ( (FuncInfo.FunctionFlags & FUNC_BlueprintEvent) && !(FuncInfo.FunctionFlags & FUNC_Native) )
 					{
 						// already a BlueprintImplementableEvent
-						FError::Throwf(TEXT("A function cannot be both BlueprintNativeEvent and BlueprintImplementableEvent!") );
+						UE_LOG_ERROR_UHT(TEXT("A function cannot be both BlueprintNativeEvent and BlueprintImplementableEvent!") );
+					}
+					else if (bSawPropertyAccessor)
+					{
+						UE_LOG_ERROR_UHT(TEXT("A function cannot be both BlueprintNativeEvent and a Blueprint Property accessor!"));
 					}
 					else if ( (FuncInfo.FunctionFlags & FUNC_Private) )
 					{
@@ -295,12 +300,16 @@ namespace
 				{
 					if (FuncInfo.FunctionFlags & FUNC_Net)
 					{
-						FError::Throwf(TEXT("BlueprintImplementableEvent functions cannot be replicated!") );
+						UE_LOG_ERROR_UHT(TEXT("BlueprintImplementableEvent functions cannot be replicated!") );
 					}
 					else if ( (FuncInfo.FunctionFlags & FUNC_BlueprintEvent) && (FuncInfo.FunctionFlags & FUNC_Native) )
 					{
 						// already a BlueprintNativeEvent
-						FError::Throwf(TEXT("A function cannot be both BlueprintNativeEvent and BlueprintImplementableEvent!") );
+						UE_LOG_ERROR_UHT(TEXT("A function cannot be both BlueprintNativeEvent and BlueprintImplementableEvent!") );
+					}
+					else if (bSawPropertyAccessor)
+					{
+						UE_LOG_ERROR_UHT(TEXT("A function cannot be both BlueprintImplementableEvent and a Blueprint Property accessor!"));
 					}
 					else if ( (FuncInfo.FunctionFlags & FUNC_Private) )
 					{
@@ -445,6 +454,33 @@ namespace
 				}
 				break;
 
+				case EFunctionSpecifier::BlueprintGetter:
+				{
+					if (FuncInfo.FunctionFlags & FUNC_Event)
+					{
+						UE_LOG_ERROR_UHT(TEXT("Function cannot be a blueprint event and a blueprint getter."));
+					}
+
+					bSawPropertyAccessor = true;
+					FuncInfo.FunctionFlags |= FUNC_BlueprintCallable;
+					FuncInfo.FunctionFlags |= FUNC_BlueprintPure;
+					MetaData.Add(TEXT("BlueprintGetter"));
+				}
+				break;
+
+				case EFunctionSpecifier::BlueprintSetter:
+				{
+					if (FuncInfo.FunctionFlags & FUNC_Event)
+					{
+						UE_LOG_ERROR_UHT(TEXT("Function cannot be a blueprint event and a blueprint setter."));
+					}
+
+					bSawPropertyAccessor = true;
+					FuncInfo.FunctionFlags |= FUNC_BlueprintCallable;
+					MetaData.Add(TEXT("BlueprintSetter"));
+				}
+				break;
+
 				case EFunctionSpecifier::BlueprintPure:
 				{
 					bool bIsPure = true;
@@ -2710,93 +2746,190 @@ void FHeaderParser::FixupDelegateProperties( FClasses& AllClasses, UStruct* Stru
 	}
 }
 
-/**
- * Verifies that all specified class's UProperties with CFG_RepNotify have valid callback targets with no parameters nor return values
- *
- * @param	TargetClass			class to verify rep notify properties for
- */
-void FHeaderParser::VerifyRepNotifyCallbacks( UClass* TargetClass )
+void FHeaderParser::VerifyBlueprintPropertyGetter(UProperty* Prop, UFunction* TargetFunc)
+{
+	check(TargetFunc);
+
+	UProperty* ReturnProp = TargetFunc->GetReturnProperty();
+	if (TargetFunc->NumParms > 1 || (TargetFunc->NumParms == 1 && ReturnProp == nullptr))
+	{
+		UE_LOG_ERROR_UHT(TEXT("Blueprint Property getter function %s must not have parameters."), *TargetFunc->GetName());
+	}
+
+	if (ReturnProp == nullptr || !Prop->SameType(ReturnProp))
+	{
+		FString ExtendedCPPType;
+		FString CPPType = Prop->GetCPPType(&ExtendedCPPType);
+		UE_LOG_ERROR_UHT(TEXT("Blueprint Property getter function %s must have return value of type %s%s."), *TargetFunc->GetName(), *CPPType, *ExtendedCPPType);
+	}
+
+	if (TargetFunc->HasAnyFunctionFlags(FUNC_Event))
+	{
+		UE_LOG_ERROR_UHT(TEXT("Blueprint Property setter function cannot be a blueprint event."));
+	}
+	else if (!TargetFunc->HasAnyFunctionFlags(FUNC_BlueprintPure))
+	{
+		UE_LOG_ERROR_UHT(TEXT("Blueprint Property getter function must be pure."));
+	}
+}
+
+void FHeaderParser::VerifyBlueprintPropertySetter(UProperty* Prop, UFunction* TargetFunc)
+{
+	check(TargetFunc);
+	UProperty* ReturnProp = TargetFunc->GetReturnProperty();
+
+	if (ReturnProp)
+	{
+		UE_LOG_ERROR_UHT(TEXT("Blueprint Property setter function %s must not have a return value."), *TargetFunc->GetName());
+	}
+	else
+	{
+		TFieldIterator<UProperty> Parm(TargetFunc);
+		if (TargetFunc->NumParms != 1 || !Prop->SameType(*Parm))
+		{
+			FString ExtendedCPPType;
+			FString CPPType = Prop->GetCPPType(&ExtendedCPPType);
+			UE_LOG_ERROR_UHT(TEXT("Blueprint Property setter function %s must have exactly one parameter of type %s%s."), *TargetFunc->GetName(), *CPPType, *ExtendedCPPType);
+		}
+	}
+
+	if (TargetFunc->HasAnyFunctionFlags(FUNC_Event))
+	{
+		UE_LOG_ERROR_UHT(TEXT("Blueprint Property setter function cannot be a blueprint event."));
+	}
+	else if (!TargetFunc->HasAnyFunctionFlags(FUNC_BlueprintCallable))
+	{
+		UE_LOG_ERROR_UHT(TEXT("Blueprint Property setter function must be blueprint callable."));
+	}
+	else if (TargetFunc->HasAnyFunctionFlags(FUNC_BlueprintPure))
+	{
+		UE_LOG_ERROR_UHT(TEXT("Blueprint Property setter function must not be pure."));
+	}
+}
+
+void FHeaderParser::VerifyRepNotifyCallback(UProperty* Prop, UFunction* TargetFunc)
+{
+	if( TargetFunc )
+	{
+		if (TargetFunc->GetReturnProperty())
+		{
+			UE_LOG_ERROR_UHT(TEXT("Replication notification function %s must not have return value."), *TargetFunc->GetName());
+		}
+
+		const bool bIsArrayProperty = ( Prop->ArrayDim > 1 || Cast<UArrayProperty>(Prop) );
+		const int32 MaxParms = bIsArrayProperty ? 2 : 1;
+
+		if ( TargetFunc->NumParms > MaxParms)
+		{
+			UE_LOG_ERROR_UHT(TEXT("Replication notification function %s has too many parameters."), *TargetFunc->GetName());
+		}
+
+		TFieldIterator<UProperty> Parm(TargetFunc);
+		if ( TargetFunc->NumParms >= 1 && Parm)
+		{
+			// First parameter is always the old value:
+			if ( !Prop->SameType(*Parm) )
+			{
+				FString ExtendedCPPType;
+				FString CPPType = Prop->GetCPPType(&ExtendedCPPType);
+				UE_LOG_ERROR_UHT(TEXT("Replication notification function %s has invalid parameter for property %s. First (optional) parameter must be of type %s%s."), *TargetFunc->GetName(), *Prop->GetName(), *CPPType, *ExtendedCPPType);
+			}
+
+			++Parm;
+		}
+
+		if ( TargetFunc->NumParms >= 2 && Parm)
+		{
+			// A 2nd parameter for arrays can be specified as a const TArray<uint8>&. This is a list of element indices that have changed
+			UArrayProperty *ArrayProp = Cast<UArrayProperty>(*Parm);
+			if (!(ArrayProp && Cast<UByteProperty>(ArrayProp->Inner)) || !(Parm->GetPropertyFlags() & CPF_ConstParm) || !(Parm->GetPropertyFlags() & CPF_ReferenceParm))
+			{
+				UE_LOG_ERROR_UHT(TEXT("Replication notification function %s (optional) second parameter must be of type 'const TArray<uint8>&'"), *TargetFunc->GetName());
+			}
+		}
+	}
+	else
+	{
+		// Couldn't find a valid function...
+		UE_LOG_ERROR_UHT(TEXT("Replication notification function %s not found"), *Prop->RepNotifyFunc.ToString() );
+	}
+}
+void FHeaderParser::VerifyPropertyMarkups( UClass* TargetClass )
 {
 	// Iterate over all properties, looking for those flagged as CPF_RepNotify
 	for ( UField* Field = TargetClass->Children; Field; Field = Field->Next )
 	{
-		UProperty* Prop = Cast<UProperty>(Field);
-		if( Prop && (Prop->GetPropertyFlags() & CPF_RepNotify) )
+		if (UProperty* Prop = Cast<UProperty>(Field))
 		{
+			auto FindTargetFunction = [&](const FName FuncName)
+			{
+				// Search through this class and its superclasses looking for the specified callback
+				UFunction* TargetFunc = nullptr;
+				UClass* SearchClass = TargetClass;
+				while( SearchClass && !TargetFunc )
+				{
+					// Since the function map is not valid yet, we have to iterate over the fields to look for the function
+					for( UField* TestField = SearchClass->Children; TestField; TestField = TestField->Next )
+					{
+						UFunction* TestFunc = Cast<UFunction>(TestField);
+						if (TestFunc && FNativeClassHeaderGenerator::GetOverriddenFName(TestFunc) == FuncName)
+						{
+							TargetFunc = TestFunc;
+							break;
+						}
+					}
+					SearchClass = SearchClass->GetSuperClass();
+				}
+
+				return TargetFunc;
+			};
+
 			FClassMetaData* TargetClassData = GScriptHelper.FindClassData(TargetClass);
 			check(TargetClassData);
 			FTokenData* PropertyToken = TargetClassData->FindTokenData(Prop);
 			check(PropertyToken);
 
-			// Search through this class and its superclasses looking for the specified callback
-			UFunction* TargetFunc = NULL;
-			UClass* SearchClass = TargetClass;
-			while( SearchClass && !TargetFunc )
+			TGuardValue<int32> GuardedInputPos(InputPos, PropertyToken->Token.StartPos);
+			TGuardValue<int32> GuardedInputLine(InputLine, PropertyToken->Token.StartLine);
+
+			if (Prop->HasAnyPropertyFlags(CPF_RepNotify))
 			{
-				// Since the function map is not valid yet, we have to iterate over the fields to look for the function
-				for( UField* TestField = SearchClass->Children; TestField; TestField = TestField->Next )
-				{
-					UFunction* TestFunc = Cast<UFunction>(TestField);
-					if (TestFunc && FNativeClassHeaderGenerator::GetOverriddenFName(TestFunc) == Prop->RepNotifyFunc)
-					{
-						TargetFunc = TestFunc;
-						break;
-					}
-				}
-				SearchClass = SearchClass->GetSuperClass();
+				VerifyRepNotifyCallback(Prop, FindTargetFunction(Prop->RepNotifyFunc));
 			}
 
-			if( TargetFunc )
+			if (Prop->HasAnyPropertyFlags(CPF_BlueprintVisible))
 			{
-				if (TargetFunc->GetReturnProperty())
-				{
-					UngetToken(PropertyToken->Token);
-					FError::Throwf(TEXT("Replication notification function %s must not have return values"), *Prop->RepNotifyFunc.ToString());
-					break; //-V779
-				}
-				
-				bool IsArrayProperty = ( Prop->ArrayDim > 1 || Cast<UArrayProperty>(Prop) );
-				int32 MaxParms = IsArrayProperty ? 2 : 1;
-
-				if ( TargetFunc->NumParms > MaxParms)
+				const FString GetterFuncName = Prop->GetMetaData(TEXT("BlueprintGetter"));
+				if (!GetterFuncName.IsEmpty())
 				{
-					UngetToken(PropertyToken->Token);
-					FError::Throwf(TEXT("Replication notification function %s has too many parameters"), *Prop->RepNotifyFunc.ToString());
-					break;
-				}
-				
-				TFieldIterator<UProperty> Parm(TargetFunc);
-				if ( TargetFunc->NumParms >= 1 && Parm)
-				{
-					// First parameter is always the old value:
-					if ( Parm->GetClass() != Prop->GetClass() )
+					if (UFunction* TargetFunc = FindTargetFunction(*GetterFuncName))
 					{
-						UngetToken(PropertyToken->Token);
-						FError::Throwf(TEXT("Replication notification function %s has invalid parameter for property $%s. First (optional) parameter must be a const reference of the same property type."), *Prop->RepNotifyFunc.ToString(), *Prop->GetName());
-						break;
+						VerifyBlueprintPropertyGetter(Prop, TargetFunc);
+					}
+					else
+					{
+						// Couldn't find a valid function...
+						UE_LOG_ERROR_UHT(TEXT("Blueprint Property getter function %s not found"), *GetterFuncName);
 					}
-
-					++Parm;
 				}
 
-				if ( TargetFunc->NumParms >= 2 && Parm)
+				if (!Prop->HasAnyPropertyFlags(CPF_BlueprintReadOnly))
 				{
-					// A 2nd parameter for arrays can be specified as a const TArray<uint8>&. This is a list of element indices that have changed
-					UArrayProperty *ArrayProp = Cast<UArrayProperty>(*Parm);
-					if (!(ArrayProp && Cast<UByteProperty>(ArrayProp->Inner)) || !(Parm->GetPropertyFlags() & CPF_ConstParm) || !(Parm->GetPropertyFlags() & CPF_ReferenceParm))
+					const FString SetterFuncName = Prop->GetMetaData(TEXT("BlueprintSetter"));
+					if (!SetterFuncName.IsEmpty())
 					{
-						UngetToken(PropertyToken->Token);
-						FError::Throwf(TEXT("Replication notification function %s (optional) parameter must be of type 'const TArray<uint8>&'"), *Prop->RepNotifyFunc.ToString());
-						break;
+						if (UFunction* TargetFunc = FindTargetFunction(*SetterFuncName))
+						{
+							VerifyBlueprintPropertySetter(Prop, TargetFunc);
+						}
+						else
+						{
+							// Couldn't find a valid function...
+							UE_LOG_ERROR_UHT(TEXT("Blueprint Property setter function %s not found"), *SetterFuncName);
+						}
 					}
 				}
 			}
-			else
-			{
-				// Couldn't find a valid function...
-				UngetToken(PropertyToken->Token);
-				FError::Throwf(TEXT("Replication notification function %s not found"), *Prop->RepNotifyFunc.ToString() );
-			}
 		}
 	}
 }
@@ -3023,7 +3156,9 @@ void FHeaderParser::GetVarType(
 
 	// Process the list of specifiers
 	bool bSeenEditSpecifier = false;
-	bool bSeenBlueprintEditSpecifier = false;
+	bool bSeenBlueprintWriteSpecifier = false;
+	bool bSeenBlueprintReadOnlySpecifier = false;
+	bool bSeenBlueprintGetterSpecifier = false;
 	for (const FPropertySpecifier& Specifier : SpecifiersFound)
 	{
 		EVariableSpecifier SpecID = (EVariableSpecifier)Algo::FindSortedStringCaseInsensitive(*Specifier.Key, GVariableSpecifierStrings);
@@ -3099,9 +3234,9 @@ void FHeaderParser::GetVarType(
 
 				case EVariableSpecifier::BlueprintReadWrite:
 				{
-					if (bSeenBlueprintEditSpecifier)
+					if (bSeenBlueprintReadOnlySpecifier)
 					{
-						UE_LOG_ERROR_UHT(TEXT("Found more than one Blueprint read/write specifier (%s), only one is allowed"), *Specifier.Key);
+						UE_LOG_ERROR_UHT(TEXT("Cannot specify a property as being both BlueprintReadOnly and BlueprintReadWrite."));
 					}
 
 					const FString* PrivateAccessMD = MetaDataFromNewStyle.Find(TEXT("AllowPrivateAccess"));  // FBlueprintMetadata::MD_AllowPrivateAccess
@@ -3117,15 +3252,34 @@ void FHeaderParser::GetVarType(
 					}
 
 					Flags |= CPF_BlueprintVisible;
-					bSeenBlueprintEditSpecifier = true;
+					bSeenBlueprintWriteSpecifier = true;				}
+				break;
+
+				case EVariableSpecifier::BlueprintSetter:
+				{
+					if (bSeenBlueprintReadOnlySpecifier)
+					{
+						UE_LOG_ERROR_UHT(TEXT("Cannot specify a property as being both BlueprintReadOnly and having a BlueprintSetter."));
+					}
+
+					if (OwnerStruct->IsA<UScriptStruct>())
+					{
+						UE_LOG_ERROR_UHT(TEXT("Cannot specify BlueprintSetter for a struct member."))
+					}
+
+					const FString BlueprintSetterFunction = RequireExactlyOneSpecifierValue(Specifier);
+					MetaDataFromNewStyle.Add(TEXT("BlueprintSetter"), BlueprintSetterFunction);
+
+					Flags |= CPF_BlueprintVisible;
+					bSeenBlueprintWriteSpecifier = true;
 				}
 				break;
 
 				case EVariableSpecifier::BlueprintReadOnly:
 				{
-					if (bSeenBlueprintEditSpecifier)
+					if (bSeenBlueprintWriteSpecifier)
 					{
-						UE_LOG_ERROR_UHT(TEXT("Found more than one Blueprint read/write specifier (%s), only one is allowed"), *Specifier.Key);
+						UE_LOG_ERROR_UHT(TEXT("Cannot specify both BlueprintReadOnly and BlueprintReadWrite or BlueprintSetter."), *Specifier.Key);
 					}
 
 					const FString* PrivateAccessMD = MetaDataFromNewStyle.Find(TEXT("AllowPrivateAccess"));  // FBlueprintMetadata::MD_AllowPrivateAccess
@@ -3142,7 +3296,22 @@ void FHeaderParser::GetVarType(
 
 					Flags        |= CPF_BlueprintVisible | CPF_BlueprintReadOnly;
 					ImpliedFlags &= ~CPF_BlueprintReadOnly;
-					bSeenBlueprintEditSpecifier = true;
+					bSeenBlueprintReadOnlySpecifier = true;
+				}
+				break;
+
+				case EVariableSpecifier::BlueprintGetter:
+				{
+					if (OwnerStruct->IsA<UScriptStruct>())
+					{
+						UE_LOG_ERROR_UHT(TEXT("Cannot specify BlueprintGetter for a struct member."))
+					}
+
+					const FString BlueprintGetterFunction = RequireExactlyOneSpecifierValue(Specifier);
+					MetaDataFromNewStyle.Add(TEXT("BlueprintGetter"), BlueprintGetterFunction);
+
+					Flags        |= CPF_BlueprintVisible;
+					bSeenBlueprintGetterSpecifier = true;
 				}
 				break;
 
@@ -3370,6 +3539,14 @@ void FHeaderParser::GetVarType(
 		}
 	}
 
+	// If we saw a BlueprintGetter but did not see BlueprintSetter or 
+	// or BlueprintReadWrite then treat as BlueprintReadOnly
+	if (bSeenBlueprintGetterSpecifier && !bSeenBlueprintWriteSpecifier)
+	{
+		Flags |= CPF_BlueprintReadOnly;
+		ImpliedFlags &= ~CPF_BlueprintReadOnly;
+	}
+
 	{
 		const FString* ExposeOnSpawnStr = MetaDataFromNewStyle.Find(TEXT("ExposeOnSpawn"));
 		const bool bExposeOnSpawn = (NULL != ExposeOnSpawnStr);
@@ -4665,6 +4842,8 @@ UProperty* FHeaderParser::GetVarNameAndDim
 	}
 
 	VarProperty.TokenProperty = NewProperty;
+	VarProperty.StartLine = InputLine;
+	VarProperty.StartPos = InputPos;
 	FClassMetaData* ScopeData = GScriptHelper.FindClassData(Scope);
 	check(ScopeData);
 	ScopeData->AddProperty(VarProperty, UHTMakefile, CurrentSrcFile);
@@ -5760,13 +5939,15 @@ UDelegateFunction* FHeaderParser::CompileDelegateDeclaration(FClasses& AllClasse
 		TArray<FPropertySpecifier> SpecifiersFound;
 		ReadSpecifierSetInsideMacro(SpecifiersFound, TEXT("Delegate"), MetaData);
 
-		ProcessFunctionSpecifiers(FuncInfo, SpecifiersFound);
+		ProcessFunctionSpecifiers(FuncInfo, SpecifiersFound, MetaData);
 
 		// Get the next token and ensure it looks like a delegate
 		FToken Token;
 		GetToken(Token);
 		if (!IsValidDelegateDeclaration(Token))
+		{
 			FError::Throwf(TEXT("Unexpected token following UDELEGATE(): %s"), Token.Identifier);
+		}
 
 		DelegateMacro = Token.Identifier;
 
@@ -6094,7 +6275,7 @@ void FHeaderParser::CompileFunctionDeclaration(FClasses& AllClasses)
 		FuncInfo.FunctionFlags |= FUNC_EditorOnly;
 	}
 
-	ProcessFunctionSpecifiers(FuncInfo, SpecifiersFound);
+	ProcessFunctionSpecifiers(FuncInfo, SpecifiersFound, MetaData);
 
 	const bool bClassGeneratedFromBP = FClass::IsDynamic(GetCurrentClass());
 	if ((FuncInfo.FunctionFlags & FUNC_NetServer) && !(FuncInfo.FunctionFlags & FUNC_NetValidate) && !bClassGeneratedFromBP)
@@ -6134,14 +6315,16 @@ void FHeaderParser::CompileFunctionDeclaration(FClasses& AllClasses)
 	}
 
 	FString*   InternalPtr = MetaData.Find("BlueprintInternalUseOnly"); // FBlueprintMetadata::MD_BlueprintInternalUseOnly
-	const bool bDeprecated = MetaData.Contains("DeprecatedFunction");       // FBlueprintMetadata::MD_DeprecatedFunction
-	const bool bHasMenuCategory = MetaData.Contains("Category");                 // FBlueprintMetadata::MD_FunctionCategory
 	const bool bInternalOnly = InternalPtr && *InternalPtr == TEXT("true");
 
 	// If this function is blueprint callable or blueprint pure, require a category 
 	if ((FuncInfo.FunctionFlags & (FUNC_BlueprintCallable | FUNC_BlueprintPure)) != 0) 
 	{ 
-		if (!bHasMenuCategory && !bInternalOnly && !bDeprecated) 
+		const bool bDeprecated = MetaData.Contains("DeprecatedFunction");       // FBlueprintMetadata::MD_DeprecatedFunction
+		const bool bBlueprintAccessor = MetaData.Contains("BlueprintSetter") || MetaData.Contains("BlueprintGetter"); // FBlueprintMetadata::MD_BlueprintSetter, // FBlueprintMetadata::MD_BlueprintGetter
+		const bool bHasMenuCategory = MetaData.Contains("Category");                 // FBlueprintMetadata::MD_FunctionCategory
+
+		if (!bHasMenuCategory && !bInternalOnly && !bDeprecated && !bBlueprintAccessor) 
 		{ 
 			const bool bModuleIsGame = CurrentlyParsedModule && (
 				CurrentlyParsedModule->ModuleType == EBuildModuleType::GameDeveloper ||
@@ -8799,7 +8982,7 @@ void FHeaderParser::ResetClassData()
 		CurrentClass->ClassFlags |= (SuperClass->ClassFlags) & CLASS_ScriptInherit;
 		CurrentClass->ClassConfigName = SuperClass->ClassConfigName;
 		check(SuperClass->ClassWithin);
-		if (CurrentClass->ClassWithin == NULL)
+		if (CurrentClass->ClassWithin == nullptr)
 		{
 			CurrentClass->ClassWithin = SuperClass->ClassWithin;
 		}
@@ -8833,7 +9016,7 @@ void FHeaderParser::ResetClassData()
 void FHeaderParser::PostPopNestClass(UClass* CurrentClass)
 {
 	// Validate all the rep notify events here, to make sure they're implemented
-	VerifyRepNotifyCallbacks(CurrentClass);
+	VerifyPropertyMarkups(CurrentClass);
 
 	// Iterate over all the interfaces we claim to implement
 	for (FImplementedInterface& Impl : CurrentClass->Interfaces)
diff --git a/Engine/Source/Programs/UnrealHeaderTool/Private/HeaderParser.h b/Engine/Source/Programs/UnrealHeaderTool/Private/HeaderParser.h
--- a/Engine/Source/Programs/UnrealHeaderTool/Private/HeaderParser.h
+++ b/Engine/Source/Programs/UnrealHeaderTool/Private/HeaderParser.h
@@ -690,13 +690,6 @@ protected:
 	 */
 	void FixupDelegateProperties(FClasses& AllClasses, UStruct* ValidationScope, FScope& Scope, TMap<FName, UFunction*>& DelegateCache);
 
-	/**
-	 * Verifies that all specified class's UProperties with CFG_RepNotify have valid callback targets with no parameters nor return values
-	 *
-	 * @param	TargetClass			class to verify rep notify properties for
-	 */
-	void VerifyRepNotifyCallbacks( UClass* TargetClass );
-
 	// Retry functions.
 	void InitScriptLocation( FScriptLocation& Retry );
 	void ReturnToLocation( const FScriptLocation& Retry, bool Binary=1, bool Text=1 );
@@ -735,6 +728,18 @@ private:
 
 	// Parses possible version declaration in generated code, e.g. GENERATED_BODY(<some_version>).
 	void CompileVersionDeclaration(UStruct* Struct);
+
+	// Verifies that all specified class's UProperties with function associations have valid targets
+	void VerifyPropertyMarkups( UClass* TargetClass );
+
+	// Verifies the target function meets the criteria for a blueprint property getter
+	void VerifyBlueprintPropertyGetter(UProperty* Property, UFunction* TargetFunction);
+
+	// Verifies the target function meets the criteria for a blueprint property setter
+	void VerifyBlueprintPropertySetter(UProperty* Property, UFunction* TargetFunction);
+
+	// Verifies the target function meets the criteria for a replication notify callback
+	void VerifyRepNotifyCallback(UProperty* Property, UFunction* TargetFunction);
 };
 
 /////////////////////////////////////////////////////
diff --git a/Engine/Source/Programs/UnrealHeaderTool/Private/Specifiers/FunctionSpecifiers.def b/Engine/Source/Programs/UnrealHeaderTool/Private/Specifiers/FunctionSpecifiers.def
--- a/Engine/Source/Programs/UnrealHeaderTool/Private/Specifiers/FunctionSpecifiers.def
+++ b/Engine/Source/Programs/UnrealHeaderTool/Private/Specifiers/FunctionSpecifiers.def
@@ -5,9 +5,11 @@
 FUNCTION_SPECIFIER(BlueprintAuthorityOnly)
 FUNCTION_SPECIFIER(BlueprintCallable)
 FUNCTION_SPECIFIER(BlueprintCosmetic)
+FUNCTION_SPECIFIER(BlueprintGetter)
 FUNCTION_SPECIFIER(BlueprintImplementableEvent)
 FUNCTION_SPECIFIER(BlueprintNativeEvent)
 FUNCTION_SPECIFIER(BlueprintPure)
+FUNCTION_SPECIFIER(BlueprintSetter)
 FUNCTION_SPECIFIER(Client)
 FUNCTION_SPECIFIER(CustomThunk)
 FUNCTION_SPECIFIER(Exec)
diff --git a/Engine/Source/Programs/UnrealHeaderTool/Private/Specifiers/VariableSpecifiers.def b/Engine/Source/Programs/UnrealHeaderTool/Private/Specifiers/VariableSpecifiers.def
--- a/Engine/Source/Programs/UnrealHeaderTool/Private/Specifiers/VariableSpecifiers.def
+++ b/Engine/Source/Programs/UnrealHeaderTool/Private/Specifiers/VariableSpecifiers.def
@@ -7,8 +7,10 @@ VARIABLE_SPECIFIER(AssetRegistrySearchable)
 VARIABLE_SPECIFIER(BlueprintAssignable)
 VARIABLE_SPECIFIER(BlueprintAuthorityOnly)
 VARIABLE_SPECIFIER(BlueprintCallable)
+VARIABLE_SPECIFIER(BlueprintGetter)
 VARIABLE_SPECIFIER(BlueprintReadOnly)
 VARIABLE_SPECIFIER(BlueprintReadWrite)
+VARIABLE_SPECIFIER(BlueprintSetter)
 VARIABLE_SPECIFIER(Config)
 VARIABLE_SPECIFIER(Const)
 VARIABLE_SPECIFIER(DuplicateTransient)
diff --git a/Engine/Source/Runtime/AIModule/Classes/BehaviorTree/BTCompositeNode.h b/Engine/Source/Runtime/AIModule/Classes/BehaviorTree/BTCompositeNode.h
--- a/Engine/Source/Runtime/AIModule/Classes/BehaviorTree/BTCompositeNode.h
+++ b/Engine/Source/Runtime/AIModule/Classes/BehaviorTree/BTCompositeNode.h
@@ -184,6 +184,15 @@ protected:
 	/** if set, NotifyNodeDeactivation will be called */
 	uint8 bUseNodeDeactivationNotify : 1;
 
+	/** if set, CanNotifyDecoratorsOnActivation will be called */
+	uint8 bUseDecoratorsActivationCheck : 1;
+
+	/** if set, CanNotifyDecoratorsOnDeactivation will be called */
+	uint8 bUseDecoratorsDeactivationCheck : 1;
+
+	/** if set, CanNotifyDecoratorsOnFailedActivation will be called */
+	uint8 bUseDecoratorsFailedActivationCheck : 1;
+
 	/** called just after child execution, allows to modify result */
 	virtual void NotifyChildExecution(UBehaviorTreeComponent& OwnerComp, uint8* NodeMemory, int32 ChildIdx, EBTNodeResult::Type& NodeResult) const;
 
@@ -193,6 +202,15 @@ protected:
 	/** called when start leaves this node */
 	virtual void NotifyNodeDeactivation(FBehaviorTreeSearchData& SearchData, EBTNodeResult::Type& NodeResult) const;
 
+	/** check if NotifyDecoratorsOnActivation is allowed, requires bUseDecoratorsActivationCheck flag  */
+	virtual bool CanNotifyDecoratorsOnActivation(FBehaviorTreeSearchData& SearchData, int32 ChildIdx) const;
+
+	/** check if NotifyDecoratorsOnDeactivation is allowed, requires bUseDecoratorsDeactivationCheck flag  */
+	virtual bool CanNotifyDecoratorsOnDeactivation(FBehaviorTreeSearchData& SearchData, int32 ChildIdx, EBTNodeResult::Type& NodeResult) const;
+
+	/** check if NotifyDecoratorsOnFailedActivation is allowed, requires bUseDecoratorsActivationCheck flag  */
+	virtual bool CanNotifyDecoratorsOnFailedActivation(FBehaviorTreeSearchData& SearchData, int32 ChildIdx, EBTNodeResult::Type& NodeResult) const;
+
 	/** runs through decorators on given child node and notify them about activation */
 	void NotifyDecoratorsOnActivation(FBehaviorTreeSearchData& SearchData, int32 ChildIdx) const;
 
diff --git a/Engine/Source/Runtime/AIModule/Classes/BehaviorTree/BehaviorTreeTypes.h b/Engine/Source/Runtime/AIModule/Classes/BehaviorTree/BehaviorTreeTypes.h
--- a/Engine/Source/Runtime/AIModule/Classes/BehaviorTree/BehaviorTreeTypes.h
+++ b/Engine/Source/Runtime/AIModule/Classes/BehaviorTree/BehaviorTreeTypes.h
@@ -420,6 +420,9 @@ struct FBehaviorTreeSearchData
 	/** assign unique Id number */
 	void AssignSearchId();
 
+	/** clear state of search */
+	void Reset();
+
 	FBehaviorTreeSearchData(UBehaviorTreeComponent& InOwnerComp) 
 		: OwnerComp(InOwnerComp), bPostponeSearch(false), bSearchInProgress(false)
 	{}
diff --git a/Engine/Source/Runtime/AIModule/Classes/BehaviorTree/Composites/BTComposite_SimpleParallel.h b/Engine/Source/Runtime/AIModule/Classes/BehaviorTree/Composites/BTComposite_SimpleParallel.h
--- a/Engine/Source/Runtime/AIModule/Classes/BehaviorTree/Composites/BTComposite_SimpleParallel.h
+++ b/Engine/Source/Runtime/AIModule/Classes/BehaviorTree/Composites/BTComposite_SimpleParallel.h
@@ -64,6 +64,7 @@ class AIMODULE_API UBTComposite_SimpleParallel : public UBTCompositeNode
 
 	virtual void NotifyChildExecution(UBehaviorTreeComponent& OwnerComp, uint8* NodeMemory, int32 ChildIdx, EBTNodeResult::Type& NodeResult) const override;
 	virtual void NotifyNodeDeactivation(FBehaviorTreeSearchData& SearchData, EBTNodeResult::Type& NodeResult) const override;
+	virtual bool CanNotifyDecoratorsOnDeactivation(FBehaviorTreeSearchData& SearchData, int32 ChildIdx, EBTNodeResult::Type& NodeResult) const override;
 	virtual bool CanPushSubtree(UBehaviorTreeComponent& OwnerComp, uint8* NodeMemory, int32 ChildIdx) const override;
 	virtual void SetChildOverride(FBehaviorTreeSearchData& SearchData, int8 Index) const override;
 	virtual uint16 GetInstanceMemorySize() const override;
diff --git a/Engine/Source/Runtime/AIModule/Private/BehaviorTree/BTCompositeNode.cpp b/Engine/Source/Runtime/AIModule/Private/BehaviorTree/BTCompositeNode.cpp
--- a/Engine/Source/Runtime/AIModule/Private/BehaviorTree/BTCompositeNode.cpp
+++ b/Engine/Source/Runtime/AIModule/Private/BehaviorTree/BTCompositeNode.cpp
@@ -13,6 +13,9 @@ UBTCompositeNode::UBTCompositeNode(const FObjectInitializer& ObjectInitializer)
 	bUseChildExecutionNotify = false;
 	bUseNodeActivationNotify = false;
 	bUseNodeDeactivationNotify = false;
+	bUseDecoratorsActivationCheck = false;
+	bUseDecoratorsDeactivationCheck = false;
+	bUseDecoratorsFailedActivationCheck = false;
 }
 
 void UBTCompositeNode::InitializeComposite(uint16 InLastExecutionIndex)
@@ -40,7 +43,12 @@ int32 UBTCompositeNode::FindChildToExecute(FBehaviorTreeSearchData& SearchData,
 			else
 			{
 				LastResult = EBTNodeResult::Failed;
-				NotifyDecoratorsOnFailedActivation(SearchData, ChildIdx, LastResult);
+
+				const bool bCanNotify = !bUseDecoratorsFailedActivationCheck || CanNotifyDecoratorsOnFailedActivation(SearchData, ChildIdx, LastResult);
+				if (bCanNotify)
+				{
+					NotifyDecoratorsOnFailedActivation(SearchData, ChildIdx, LastResult);
+				}
 			}
 
 			ChildIdx = GetNextChild(SearchData, ChildIdx, LastResult);
@@ -87,7 +95,11 @@ void UBTCompositeNode::OnChildActivation(FBehaviorTreeSearchData& SearchData, in
 
 	// pass to decorators before changing current child in node memory
 	// so they can access previously executed one if needed
-	NotifyDecoratorsOnActivation(SearchData, ChildIndex);
+	const bool bCanNotify = !bUseDecoratorsActivationCheck || CanNotifyDecoratorsOnActivation(SearchData, ChildIndex);
+	if (bCanNotify)
+	{
+		NotifyDecoratorsOnActivation(SearchData, ChildIndex);
+	}
 
 	// don't activate task services here, it's applied BEFORE aborting (e.g. abort lower pri decorator)
 	// use UBehaviorTreeComponent::ExecuteTask instead
@@ -127,7 +139,11 @@ void UBTCompositeNode::OnChildDeactivation(FBehaviorTreeSearchData& SearchData,
 
 	// pass to decorators after composite is updated (so far only simple parallel uses it)
 	// to have them working on correct result + they must be able to modify it if requested (e.g. force success)
-	NotifyDecoratorsOnDeactivation(SearchData, ChildIndex, NodeResult);
+	const bool bCanNotify = !bUseDecoratorsDeactivationCheck || CanNotifyDecoratorsOnDeactivation(SearchData, ChildIndex, NodeResult);
+	if (bCanNotify)
+	{
+		NotifyDecoratorsOnDeactivation(SearchData, ChildIndex, NodeResult);
+	}
 }
 
 void UBTCompositeNode::OnNodeActivation(FBehaviorTreeSearchData& SearchData) const
@@ -580,6 +596,21 @@ bool UBTCompositeNode::CanPushSubtree(UBehaviorTreeComponent& OwnerComp, uint8*
 	return true;
 }
 
+bool UBTCompositeNode::CanNotifyDecoratorsOnActivation(FBehaviorTreeSearchData& SearchData, int32 ChildIdx) const
+{
+	return true;
+}
+
+bool UBTCompositeNode::CanNotifyDecoratorsOnDeactivation(FBehaviorTreeSearchData& SearchData, int32 ChildIdx, EBTNodeResult::Type& NodeResult) const
+{
+	return true;
+}
+
+bool UBTCompositeNode::CanNotifyDecoratorsOnFailedActivation(FBehaviorTreeSearchData& SearchData, int32 ChildIdx, EBTNodeResult::Type& NodeResult) const
+{
+	return true;
+}
+
 void UBTCompositeNode::DescribeRuntimeValues(const UBehaviorTreeComponent& OwnerComp, uint8* NodeMemory, EBTDescriptionVerbosity::Type Verbosity, TArray<FString>& Values) const
 {
 	Super::DescribeRuntimeValues(OwnerComp, NodeMemory, Verbosity, Values);
diff --git a/Engine/Source/Runtime/AIModule/Private/BehaviorTree/BehaviorTreeComponent.cpp b/Engine/Source/Runtime/AIModule/Private/BehaviorTree/BehaviorTreeComponent.cpp
--- a/Engine/Source/Runtime/AIModule/Private/BehaviorTree/BehaviorTreeComponent.cpp
+++ b/Engine/Source/Runtime/AIModule/Private/BehaviorTree/BehaviorTreeComponent.cpp
@@ -288,6 +288,13 @@ void UBehaviorTreeComponent::StopTree(EBTStopMode::Type StopMode)
 		UE_VLOG(GetOwner(), LogBehaviorTree, Warning, TEXT("StopTree was forced while waiting for tasks to finish aborting!"));
 	}
 
+	// make sure that all nodes are getting deactivation notifies
+	if (InstanceStack.Num())
+	{
+		EBTNodeResult::Type AbortedResult = EBTNodeResult::Aborted;
+		DeactivateUpTo(InstanceStack[0].RootNode, 0, AbortedResult);
+	}
+
 	// clear current state, don't touch debugger data
 	for (int32 InstanceIndex = 0; InstanceIndex < InstanceStack.Num(); InstanceIndex++)
 	{
@@ -297,6 +304,7 @@ void UBehaviorTreeComponent::StopTree(EBTStopMode::Type StopMode)
 
 	InstanceStack.Reset();
 	TaskMessageObservers.Reset();
+	SearchData.Reset();
 	ExecutionRequest = FBTNodeExecutionInfo();
 	PendingExecution = FBTPendingExecutionInfo();
 	ActiveInstanceIdx = 0;
diff --git a/Engine/Source/Runtime/AIModule/Private/BehaviorTree/BehaviorTreeTypes.cpp b/Engine/Source/Runtime/AIModule/Private/BehaviorTree/BehaviorTreeTypes.cpp
--- a/Engine/Source/Runtime/AIModule/Private/BehaviorTree/BehaviorTreeTypes.cpp
+++ b/Engine/Source/Runtime/AIModule/Private/BehaviorTree/BehaviorTreeTypes.cpp
@@ -277,6 +277,15 @@ void FBehaviorTreeSearchData::AssignSearchId()
 	NextSearchId++;
 }
 
+void FBehaviorTreeSearchData::Reset()
+{
+	PendingUpdates.Reset();
+	SearchStart = FBTNodeIndex();
+	SearchEnd = FBTNodeIndex();
+	bSearchInProgress = false;
+	bPostponeSearch = false;
+}
+
 //----------------------------------------------------------------------//
 // FBlackboardKeySelector
 //----------------------------------------------------------------------//
diff --git a/Engine/Source/Runtime/AIModule/Private/BehaviorTree/Composites/BTComposite_SimpleParallel.cpp b/Engine/Source/Runtime/AIModule/Private/BehaviorTree/Composites/BTComposite_SimpleParallel.cpp
--- a/Engine/Source/Runtime/AIModule/Private/BehaviorTree/Composites/BTComposite_SimpleParallel.cpp
+++ b/Engine/Source/Runtime/AIModule/Private/BehaviorTree/Composites/BTComposite_SimpleParallel.cpp
@@ -11,6 +11,7 @@ UBTComposite_SimpleParallel::UBTComposite_SimpleParallel(const FObjectInitialize
 	NodeName = "Simple Parallel";
 	bUseChildExecutionNotify = true;
 	bUseNodeDeactivationNotify = true;
+	bUseDecoratorsDeactivationCheck = true;
 	OnNextChild.BindUObject(this, &UBTComposite_SimpleParallel::GetNextChildHandler);
 }
 
@@ -73,6 +74,10 @@ void UBTComposite_SimpleParallel::NotifyChildExecution(UBehaviorTreeComponent& O
 		else if (MyMemory->bMainTaskIsActive)
 		{
 			MyMemory->bMainTaskIsActive = false;
+			
+			// notify decorators on main task, ignore observers updates in FakeSearchData - they are not allowed by parallel composite
+			FBehaviorTreeSearchData FakeSearchData(OwnerComp);
+			NotifyDecoratorsOnDeactivation(FakeSearchData, ChildIdx, NodeResult);
 
 			const int32 MyInstanceIdx = OwnerComp.FindInstanceContainingNode(this);
 
@@ -142,6 +147,22 @@ void UBTComposite_SimpleParallel::NotifyNodeDeactivation(FBehaviorTreeSearchData
 	}
 }
 
+bool UBTComposite_SimpleParallel::CanNotifyDecoratorsOnDeactivation(FBehaviorTreeSearchData& SearchData, int32 ChildIdx, EBTNodeResult::Type& NodeResult) const
+{
+	// don't pass deactivation to decorators when parallel is just switching to background tree
+	// decorators on running main task will be notified when their task finishes execution
+	if (ChildIdx == EBTParallelChild::MainTask)
+	{
+		FBTParallelMemory* MyMemory = GetNodeMemory<FBTParallelMemory>(SearchData);
+		if (MyMemory->bMainTaskIsActive)
+		{
+			return false;
+		}
+	}
+
+	return true;
+}
+
 uint16 UBTComposite_SimpleParallel::GetInstanceMemorySize() const
 {
 	return sizeof(FBTParallelMemory);
diff --git a/Engine/Source/Runtime/AppFramework/Private/Widgets/Colors/SColorPicker.cpp b/Engine/Source/Runtime/AppFramework/Private/Widgets/Colors/SColorPicker.cpp
--- a/Engine/Source/Runtime/AppFramework/Private/Widgets/Colors/SColorPicker.cpp
+++ b/Engine/Source/Runtime/AppFramework/Private/Widgets/Colors/SColorPicker.cpp
@@ -1664,7 +1664,7 @@ bool OpenColorPicker(const FColorPickerArgs& Args)
 	const FVector2D CursorPos = FSlateApplication::Get().GetCursorPos();
 	FSlateRect Anchor(CursorPos.X, CursorPos.Y, CursorPos.X, CursorPos.Y);
 
-	FVector2D AdjustedSummonLocation = FSlateApplication::Get().CalculatePopupWindowPosition( Anchor, SColorPicker::DEFAULT_WINDOW_SIZE, Orient_Horizontal );
+	FVector2D AdjustedSummonLocation = FSlateApplication::Get().CalculatePopupWindowPosition( Anchor, SColorPicker::DEFAULT_WINDOW_SIZE, FVector2D::ZeroVector, Orient_Horizontal );
 
 	// Only override the color picker window creation behavior if we are not creating a modal color picker
 	const bool bOverrideNonModalCreation = (SColorPicker::OnColorPickerNonModalCreateOverride.IsBound() && !Args.bIsModal);
diff --git a/Engine/Source/Runtime/AssetRegistry/Public/AssetBundleData.h b/Engine/Source/Runtime/AssetRegistry/Public/AssetBundleData.h
--- a/Engine/Source/Runtime/AssetRegistry/Public/AssetBundleData.h
+++ b/Engine/Source/Runtime/AssetRegistry/Public/AssetBundleData.h
@@ -11,7 +11,11 @@
 USTRUCT()
 struct ASSETREGISTRY_API FAssetBundleEntry
 {
-	GENERATED_USTRUCT_BODY()
+	GENERATED_BODY()
+
+	/** Declare constructors inline so this can be a header only class */
+	FORCEINLINE FAssetBundleEntry() {}
+	FORCEINLINE ~FAssetBundleEntry() {}
 
 	/** Asset this bundle is saved within. This is empty for global bundles, or in the saved bundle info */
 	UPROPERTY()
@@ -25,9 +29,6 @@ struct ASSETREGISTRY_API FAssetBundleEntry
 	UPROPERTY()
 	TArray<FStringAssetReference> BundleAssets;
 
-	FAssetBundleEntry()	
-	{}
-
 	FAssetBundleEntry(const FAssetBundleEntry& OldEntry)
 		: BundleScope(OldEntry.BundleScope), BundleName(OldEntry.BundleName), BundleAssets(OldEntry.BundleAssets)
 	{
@@ -51,7 +52,11 @@ struct ASSETREGISTRY_API FAssetBundleEntry
 USTRUCT()
 struct ASSETREGISTRY_API FAssetBundleData
 {
-	GENERATED_USTRUCT_BODY()
+	GENERATED_BODY()
+
+	/** Declare constructors inline so this can be a header only class */
+	FORCEINLINE FAssetBundleData() {}
+	FORCEINLINE ~FAssetBundleData() {}
 
 	/** List of bundles defined */
 	UPROPERTY()
diff --git a/Engine/Source/Runtime/Core/Private/Linux/LinuxApplication.cpp b/Engine/Source/Runtime/Core/Private/Linux/LinuxApplication.cpp
--- a/Engine/Source/Runtime/Core/Private/Linux/LinuxApplication.cpp
+++ b/Engine/Source/Runtime/Core/Private/Linux/LinuxApplication.cpp
@@ -1512,12 +1512,6 @@ void FLinuxApplication::OnMouseCursorLock( bool bLockEnabled )
 	}
 }
 
-
-bool FLinuxApplication::TryCalculatePopupWindowPosition( const FPlatformRect& InAnchor, const FVector2D& InSize, const EPopUpOrientation::Type Orientation, /*OUT*/ FVector2D* const CalculatedPopUpPosition ) const
-{
-	return false;
-}
-
 void FDisplayMetrics::GetDisplayMetrics(FDisplayMetrics& OutDisplayMetrics)
 {
 	int NumDisplays = 0;
diff --git a/Engine/Source/Runtime/Core/Private/Windows/WindowsApplication.cpp b/Engine/Source/Runtime/Core/Private/Windows/WindowsApplication.cpp
--- a/Engine/Source/Runtime/Core/Private/Windows/WindowsApplication.cpp
+++ b/Engine/Source/Runtime/Core/Private/Windows/WindowsApplication.cpp
@@ -410,11 +410,6 @@ void FWindowsApplication::SetHighPrecisionMouseMode( const bool Enable, const TS
 	::RegisterRawInputDevices( &RawInputDevice, 1, sizeof( RAWINPUTDEVICE ) );
 }
 
-bool FWindowsApplication::TryCalculatePopupWindowPosition( const FPlatformRect& InAnchor, const FVector2D& InSize, const EPopUpOrientation::Type Orientation, /*OUT*/ FVector2D* const CalculatedPopUpPosition ) const
-{
-	return false;
-}
-
 FPlatformRect FWindowsApplication::GetWorkArea( const FPlatformRect& CurrentWindow ) const
 {
 	RECT WindowsWindowDim;
diff --git a/Engine/Source/Runtime/Core/Public/GenericPlatform/GenericApplication.h b/Engine/Source/Runtime/Core/Public/GenericPlatform/GenericApplication.h
--- a/Engine/Source/Runtime/Core/Public/GenericPlatform/GenericApplication.h
+++ b/Engine/Source/Runtime/Core/Public/GenericPlatform/GenericApplication.h
@@ -465,7 +465,7 @@ public:
 		return OutRect;
 	}
 
-	virtual bool TryCalculatePopupWindowPosition( const FPlatformRect& InAnchor, const FVector2D& InSize, const EPopUpOrientation::Type Orientation, /*OUT*/ FVector2D* const CalculatedPopUpPosition ) const { return false; }
+	virtual bool TryCalculatePopupWindowPosition( const FPlatformRect& InAnchor, const FVector2D& InSize, const FVector2D& ProposedPlacement, const EPopUpOrientation::Type Orientation, /*OUT*/ FVector2D* const CalculatedPopUpPosition ) const { return false; }
 
 	DECLARE_EVENT_OneParam(GenericApplication, FOnDisplayMetricsChanged, const FDisplayMetrics&);
 	
diff --git a/Engine/Source/Runtime/Core/Public/Linux/LinuxApplication.h b/Engine/Source/Runtime/Core/Public/Linux/LinuxApplication.h
--- a/Engine/Source/Runtime/Core/Public/Linux/LinuxApplication.h
+++ b/Engine/Source/Runtime/Core/Public/Linux/LinuxApplication.h
@@ -63,8 +63,6 @@ public:
 
 	virtual FPlatformRect GetWorkArea( const FPlatformRect& CurrentWindow ) const override;
 
-	virtual bool TryCalculatePopupWindowPosition( const FPlatformRect& InAnchor, const FVector2D& InSize, const EPopUpOrientation::Type Orientation, /*OUT*/ FVector2D* const CalculatedPopUpPosition ) const override;
-
 	virtual EWindowTransparency GetWindowTransparencySupport() const override
 	{
 		return EWindowTransparency::PerWindow;
diff --git a/Engine/Source/Runtime/Core/Public/UObject/FrameworkObjectVersion.h b/Engine/Source/Runtime/Core/Public/UObject/FrameworkObjectVersion.h
--- a/Engine/Source/Runtime/Core/Public/UObject/FrameworkObjectVersion.h
+++ b/Engine/Source/Runtime/Core/Public/UObject/FrameworkObjectVersion.h
@@ -94,9 +94,12 @@ struct CORE_API FFrameworkObjectVersion
 		// Represent a pins container type as an enum not 3 independent booleans
 		EdGraphPinContainerType,
 
-		// Switch asset pins to store as string instead of hard object referece
+		// Switch asset pins to store as string instead of hard object reference
 		ChangeAssetPinsToString,
 
+		// Fix Local Variables so that the properties are correctly flagged as blueprint visible
+		LocalVariablesBlueprintVisible,
+
 		// -----<new versions can be added above this line>-------------------------------------------------
 		VersionPlusOne,
 		LatestVersion = VersionPlusOne - 1
diff --git a/Engine/Source/Runtime/Core/Public/Windows/WindowsApplication.h b/Engine/Source/Runtime/Core/Public/Windows/WindowsApplication.h
--- a/Engine/Source/Runtime/Core/Public/Windows/WindowsApplication.h
+++ b/Engine/Source/Runtime/Core/Public/Windows/WindowsApplication.h
@@ -355,7 +355,6 @@ public:
 	virtual FModifierKeysState GetModifierKeys() const override;
 	virtual bool IsCursorDirectlyOverSlateWindow() const override;
 	virtual FPlatformRect GetWorkArea( const FPlatformRect& CurrentWindow ) const override;
-	virtual bool TryCalculatePopupWindowPosition( const FPlatformRect& InAnchor, const FVector2D& InSize, const EPopUpOrientation::Type Orientation, /*OUT*/ FVector2D* const CalculatedPopUpPosition ) const override;
 	virtual void GetInitialDisplayMetrics( FDisplayMetrics& OutDisplayMetrics ) const override;
 	virtual EWindowTitleAlignment::Type GetWindowTitleAlignment() const override;
 	virtual EWindowTransparency GetWindowTransparencySupport() const override;
diff --git a/Engine/Source/Runtime/CoreUObject/Private/Misc/PackageName.cpp b/Engine/Source/Runtime/CoreUObject/Private/Misc/PackageName.cpp
--- a/Engine/Source/Runtime/CoreUObject/Private/Misc/PackageName.cpp
+++ b/Engine/Source/Runtime/CoreUObject/Private/Misc/PackageName.cpp
@@ -286,6 +286,16 @@ FString FPackageName::InternalFilenameToLongPackageName(const FString& InFilenam
 	if (!bIsValidLongPackageName)
 	{
 		Filename = IFileManager::Get().ConvertToRelativePath(*Filename);
+		if (InFilename.Len() > 0 && InFilename[InFilename.Len() - 1] == '/')
+		{
+			// If InFilename ends in / but converted doesn't, add the / back
+			bool bEndsInSlash = Filename.Len() > 0 && Filename[Filename.Len() - 1] == '/';
+
+			if (!bEndsInSlash)
+			{
+				Filename += TEXT("/");
+			}
+		}
 	}
 
 	FString PackageName         = FPaths::GetBaseFilename(Filename);
@@ -310,9 +320,10 @@ bool FPackageName::TryConvertFilenameToLongPackageName(const FString& InFilename
 	FString LongPackageName = InternalFilenameToLongPackageName(InFilename);
 
 	// we don't support loading packages from outside of well defined places
-	const bool bContainsDot = LongPackageName.Contains(TEXT("."), ESearchCase::CaseSensitive);
-	const bool bContainsBackslash = LongPackageName.Contains(TEXT("\\"), ESearchCase::CaseSensitive);
-	const bool bContainsColon = LongPackageName.Contains(TEXT(":"), ESearchCase::CaseSensitive);
+	int32 CharacterIndex;
+	const bool bContainsDot = LongPackageName.FindChar(TEXT('.'), CharacterIndex);
+	const bool bContainsBackslash = LongPackageName.FindChar(TEXT('\\'), CharacterIndex);
+	const bool bContainsColon = LongPackageName.FindChar(TEXT(':'), CharacterIndex);
 	const bool bResult = !(bContainsDot || bContainsBackslash || bContainsColon);
 
 	if (bResult)
@@ -447,16 +458,7 @@ bool FPackageName::SplitLongPackageName(const FString& InLongPackageName, FStrin
 
 FString FPackageName::GetLongPackageAssetName(const FString& InLongPackageName)
 {
-	int32 IndexOfLastSlash = INDEX_NONE;
-	if (InLongPackageName.FindLastChar('/', IndexOfLastSlash))
-	{
-		return InLongPackageName.Mid(IndexOfLastSlash + 1);
-	}
-	else
-	{
-		return InLongPackageName;
-	}
-
+	return GetShortName(InLongPackageName);
 }
 
 bool FPackageName::DoesPackageNameContainInvalidCharacters(const FString& InLongPackageName, FText* OutReason /*= NULL*/)
diff --git a/Engine/Source/Runtime/CoreUObject/Private/Misc/StringAssetReference.cpp b/Engine/Source/Runtime/CoreUObject/Private/Misc/StringAssetReference.cpp
--- a/Engine/Source/Runtime/CoreUObject/Private/Misc/StringAssetReference.cpp
+++ b/Engine/Source/Runtime/CoreUObject/Private/Misc/StringAssetReference.cpp
@@ -158,6 +158,10 @@ bool FStringAssetReference::ImportTextItem(const TCHAR*& Buffer, int32 PortFlags
 	{
 		if (*Buffer == TCHAR('\''))
 		{
+			// A ' token likely means we're looking at an asset string in the form "Texture2d'/Game/UI/HUD/Actions/Barrel'" and we need to read and append the path part
+			// We have to skip over the first ' as UPropertyHelpers::ReadToken doesn't read single-quoted strings correctly, but does read a path correctly
+			Buffer++; // Skip the leading '
+			ImportedPath.Reset();
 			NewBuffer = UPropertyHelpers::ReadToken(Buffer, ImportedPath, 1);
 			if (!NewBuffer)
 			{
diff --git a/Engine/Source/Runtime/CoreUObject/Private/Serialization/AsyncLoading.cpp b/Engine/Source/Runtime/CoreUObject/Private/Serialization/AsyncLoading.cpp
--- a/Engine/Source/Runtime/CoreUObject/Private/Serialization/AsyncLoading.cpp
+++ b/Engine/Source/Runtime/CoreUObject/Private/Serialization/AsyncLoading.cpp
@@ -3108,20 +3108,19 @@ void FAsyncPackage::EventDrivenCreateExport(int32 LocalExportIndex)
 					// Do this for all subobjects created in the native constructor.
 					if (!Export.Object->HasAnyFlags(RF_LoadCompleted))
 					{
-							UE_LOG(LogStreaming, VeryVerbose, TEXT("Note2: %s was constructed during load and is an export and so needs loading."), *Export.Object->GetFullName());
-							UE_CLOG(!Export.Object->HasAllFlags(RF_WillBeLoaded), LogStreaming, Fatal, TEXT("%s was found in memory and is an export but does not have all load flags."), *Export.Object->GetFullName());
-							if(Export.Object->HasAnyFlags(RF_ClassDefaultObject))
-							{
-								// never call PostLoadSubobjects on class default objects, this matches the behavior of the old linker where
-								// StaticAllocateObject prevents setting of RF_NeedPostLoad and RF_NeedPostLoadSubobjects, but FLinkerLoad::Preload
-								// assigns RF_NeedPostLoad for blueprint CDOs:
-								Export.Object->SetFlags(RF_NeedLoad | RF_NeedPostLoad | RF_WasLoaded);
-							}
-							else
-							{
-								Export.Object->SetFlags(RF_NeedLoad | RF_NeedPostLoad | RF_NeedPostLoadSubobjects | RF_WasLoaded);
-							}
-							Export.Object->ClearFlags(RF_WillBeLoaded);
+						UE_LOG(LogStreaming, VeryVerbose, TEXT("Note2: %s was constructed during load and is an export and so needs loading."), *Export.Object->GetFullName());
+						UE_CLOG(!Export.Object->HasAllFlags(RF_WillBeLoaded), LogStreaming, Fatal, TEXT("%s was found in memory and is an export but does not have all load flags."), *Export.Object->GetFullName());
+						if(Export.Object->HasAnyFlags(RF_ClassDefaultObject))
+						{
+							// never call PostLoadSubobjects on class default objects, this matches the behavior of the old linker where
+							// StaticAllocateObject prevents setting of RF_NeedPostLoad and RF_NeedPostLoadSubobjects, but FLinkerLoad::Preload
+							// assigns RF_NeedPostLoad for blueprint CDOs:
+							Export.Object->SetFlags(RF_NeedLoad | RF_NeedPostLoad | RF_WasLoaded);
+						}
+						else
+						{
+							Export.Object->SetFlags(RF_NeedLoad | RF_NeedPostLoad | RF_NeedPostLoadSubobjects | RF_WasLoaded);
+						}
 					}
 				}
 			}
diff --git a/Engine/Source/Runtime/CoreUObject/Private/UObject/CoreRedirects.cpp b/Engine/Source/Runtime/CoreUObject/Private/UObject/CoreRedirects.cpp
--- a/Engine/Source/Runtime/CoreUObject/Private/UObject/CoreRedirects.cpp
+++ b/Engine/Source/Runtime/CoreUObject/Private/UObject/CoreRedirects.cpp
@@ -719,7 +719,6 @@ bool FCoreRedirects::ReadRedirectsFromIni(const FString& IniName)
 		ConfigKeyMap.Add(TEXT("FunctionRedirects"), ECoreRedirectFlags::Type_Function);
 		ConfigKeyMap.Add(TEXT("PropertyRedirects"), ECoreRedirectFlags::Type_Property);
 		ConfigKeyMap.Add(TEXT("PackageRedirects"), ECoreRedirectFlags::Type_Package);
-		ConfigKeyMap.Add(TEXT("DelegateRedirects"), ECoreRedirectFlags::Type_Delegate);
 
 		RegisterNativeRedirects();
 
diff --git a/Engine/Source/Runtime/CoreUObject/Private/UObject/GarbageCollection.cpp b/Engine/Source/Runtime/CoreUObject/Private/UObject/GarbageCollection.cpp
--- a/Engine/Source/Runtime/CoreUObject/Private/UObject/GarbageCollection.cpp
+++ b/Engine/Source/Runtime/CoreUObject/Private/UObject/GarbageCollection.cpp
@@ -21,7 +21,6 @@
 #include "UObject/GCScopeLock.h"
 #include "HAL/ExceptionHandling.h"
 #include "UObject/UObjectClusters.h"
-#include "CoreDelegates.h"
 
 /*-----------------------------------------------------------------------------
    Garbage collection.
@@ -95,190 +94,10 @@ bool IsGarbageCollectionLocked()
 	return GGarbageCollectionGuardCritical.IsAsyncLocked();
 }
 
-/**
- * Pool for reducing GC allocations
- */
-class FGCArrayPool
-{
-public:
-
-	/**
-	* Gets the singleton instance of the FObjectArrayPool
-	* @return Pool singleton.
-	*/
-	FORCEINLINE static FGCArrayPool& Get()
-	{
-		static FAutoConsoleCommandWithOutputDevice GCDumpPoolCommand(
-			TEXT("gc.DumpPoolStats"),
-			TEXT("Dumps count and size of GC Pools"),
-			FConsoleCommandWithOutputDeviceDelegate::CreateStatic(&FGCArrayPool::DumpStats)
-		);
-
-		static FGCArrayPool* Singleton = nullptr;
-
-		if (!Singleton)
-		{
-			Singleton = new FGCArrayPool();
-			FCoreDelegates::GetMemoryTrimDelegate().AddRaw(Singleton, &FGCArrayPool::Cleanup);
-		}
-		return *Singleton;
-	}
-
-	/**
-	* Gets an event from the pool or creates one if necessary.
-	*
-	* @return The array.
-	* @see ReturnToPool
-	*/
-	FORCEINLINE TArray<UObject*>* GetArrayFromPool()
-	{
-		TArray<UObject*>* Result = Pool.Pop();
-		if (!Result)
-		{
-			Result = new TArray<UObject*>();
-		}
-		check(Result);
-#if UE_BUILD_DEBUG
-		NumberOfUsedArrays.Increment();
-#endif // UE_BUILD_DEBUG
-		return Result;
-	}
-
-	/**
-	* Returns an array to the pool.
-	*
-	* @param Array The array to return.
-	* @see GetArrayFromPool
-	*/
-	FORCEINLINE void ReturnToPool(TArray<UObject*>* Array)
-	{
-#if UE_BUILD_DEBUG
-		const int32 CheckUsedArrays = NumberOfUsedArrays.Decrement();
-		checkSlow(CheckUsedArrays >= 0);
-#endif // UE_BUILD_DEBUG
-		check(Array);
-		Array->Reset();
-		Pool.Push(Array);
-	}
-
-	/** Performs memory cleanup */
-	void Cleanup()
-	{
-#if UE_BUILD_DEBUG
-		const int32 CheckUsedArrays = NumberOfUsedArrays.GetValue();
-		checkSlow(CheckUsedArrays == 0);
-#endif // UE_BUILD_DEBUG
-
-		uint32 FreedMemory = 0;
-		TArray< TArray<UObject*>* > AllArrays;
-		Pool.PopAll(AllArrays);
-		for (TArray<UObject*>* Array : AllArrays)
-		{
-			FreedMemory += Array->GetAllocatedSize();
-			delete Array;
-		}
-		UE_LOG(LogGarbage, Log, TEXT("Freed %ub from %d GC array pools."), FreedMemory, AllArrays.Num());
-	}
-
-	/**
-	* Writes out info about the makeup of the pool. called by 'gc.DumpPoolStats'
-	*
-	* @param Array The array to return.
-	* @see GetArrayFromPool
-	*/
-	static void DumpStats(FOutputDevice& OutputDevice)
-	{
-		FGCArrayPool& Instance = Get();
-				
-		TArray<TArray<UObject*>*> PoppedItems;
-
-		TMap<int32, int32> Buckets;
-
-		int32 TotalSize = 0;
-		int32 MaxSize = 0;
-		int32 TotalItems = 0;
-
-		do
-		{
-			TArray<UObject*>* Item = Instance.Pool.Pop();
-
-			if (Item)
-			{
-				PoppedItems.Push(Item);
-
-				// Inc our bucket
-				Buckets.FindOrAdd(Item->Max()) += 1;
-
-				TotalSize += Item->Max();
-				TotalItems++;
-			}
-			else
-			{
-				break;
-			}
-
-		} while (true);
-
-		// return everything to the pool
-		while (PoppedItems.Num())
-		{
-			Instance.Pool.Push(PoppedItems.Pop());
-		}
-
-		// One of these lists is used by the main GC and is huge, so remove it so that it doesn't
-		// pollute the stats and we can accurately see what the task pools are using.
-		int32 TotalSizeKB =(TotalSize * sizeof(UObject*)) / 1024;
-
-		OutputDevice.Logf(TEXT("GCPoolStats: %d Pools totaling %d KB. Avg: Objs=%d, Size=%d KB."),
-			TotalItems,
-			TotalSizeKB,
-			TotalSize / TotalItems,
-			TotalSizeKB / TotalItems);
-
-		//long form output...
-		TArray<int32> Keys;
-
-		Buckets.GetKeys(Keys);
-
-		Keys.Sort([&](int32 lhs, int32 rhs) {
-			return lhs > rhs;
-		});
-
-		for (int Key : Keys)
-		{
-			const int32 Value = Buckets[Key];
-			int32 ItemSize = (Key * sizeof(UObject*)) / 1024;
-			OutputDevice.Logf(TEXT("\t%d\t\t(%d Items @ %d KB = %d KB)"), Key, Value, ItemSize, Value * ItemSize);
-		}
-	}
-	
-#if UE_BUILD_DEBUG
-	void CheckLeaks()
-	{
-		// This function is called after GC has finished so at this point there should be no
-		// arrays used by GC and all should be returned to the pool
-		const int32 LeakedGCPoolArrays = NumberOfUsedArrays.GetValue();
-		checkSlow(LeakedGCPoolArrays == 0);
-	}
-#endif
-
-private:
-
-	/** Holds the collection of recycled arrays. */
-	TLockFreePointerListLIFO< TArray<UObject*> > Pool;
-
-#if UE_BUILD_DEBUG
-	/** Number of arrays currently acquired from the pool by GC */
-	FThreadSafeCounter NumberOfUsedArrays;
-#endif // UE_BUILD_DEBUG
-};
-
-void CleanupClusterArrayPools();
 /** Called on shutdown to free GC memory */
 void CleanupGCArrayPools()
 {
 	FGCArrayPool::Get().Cleanup();
-	CleanupClusterArrayPools();
 }
 
 /**
@@ -442,7 +261,6 @@ public:
 							if (ReferencedMutableObjectItem->ThisThreadAtomicallyClearedRFUnreachable())
 							{
 								// Needs doing because this is either a normal unclustered object (clustered objects are never unreachable) or a cluster root
-								ReferencedMutableObjectItem->ThisThreadAtomicallyClearedFlag(EInternalObjectFlags::NoStrongReference);
 								ObjectsToSerialize.Add(static_cast<UObject*>(ReferencedMutableObjectItem->Object));
 
 								// So is this a cluster root maybe?
@@ -464,7 +282,6 @@ public:
 									// The root is also maybe unreachable so process it and all the referenced clusters
 									if (ReferencedMutableObjectsClusterRootItem->ThisThreadAtomicallyClearedRFUnreachable())
 									{
-										ReferencedMutableObjectsClusterRootItem->ThisThreadAtomicallyClearedFlag(EInternalObjectFlags::NoStrongReference);
 										MarkReferencedClustersAsReachable(ReferencedMutableObjectsClusterRootItem->GetClusterIndex(), ObjectsToSerialize);
 									}
 								}
@@ -483,7 +300,7 @@ public:
 					if (ReferencedMutableObjectItem->IsUnreachable())
 					{
 						// Needs doing because this is either a normal unclustered object (clustered objects are never unreachable) or a cluster root
-						ReferencedMutableObjectItem->ClearFlags(EInternalObjectFlags::NoStrongReference | EInternalObjectFlags::Unreachable);
+						ReferencedMutableObjectItem->ClearFlags(EInternalObjectFlags::Unreachable);
 						ObjectsToSerialize.Add(static_cast<UObject*>(ReferencedMutableObjectItem->Object));
 
 						// So is this a cluster root?
@@ -501,7 +318,7 @@ public:
 						FUObjectItem* ReferencedMutableObjectsClusterRootItem = GUObjectArray.IndexToObjectUnsafeForGC(ReferencedMutableObjectItem->GetOwnerIndex());
 						if (ReferencedMutableObjectsClusterRootItem->IsUnreachable())
 						{
-							ReferencedMutableObjectsClusterRootItem->ClearFlags(EInternalObjectFlags::NoStrongReference | EInternalObjectFlags::Unreachable);
+							ReferencedMutableObjectsClusterRootItem->ClearFlags(EInternalObjectFlags::Unreachable);
 							MarkReferencedClustersAsReachable(ReferencedMutableObjectsClusterRootItem->GetClusterIndex(), ObjectsToSerialize);
 						}
 					}
@@ -539,12 +356,12 @@ public:
 					{
 						if (ReferencedClusterRootObjectItem->IsUnreachable())
 						{
-							ReferencedClusterRootObjectItem->ThisThreadAtomicallyClearedFlag(EInternalObjectFlags::NoStrongReference | EInternalObjectFlags::Unreachable);
+							ReferencedClusterRootObjectItem->ThisThreadAtomicallyClearedFlag( EInternalObjectFlags::Unreachable);
 						}
 					}
 					else
 					{
-						ReferencedClusterRootObjectItem->ClearFlags(EInternalObjectFlags::NoStrongReference | EInternalObjectFlags::Unreachable);
+						ReferencedClusterRootObjectItem->ClearFlags(EInternalObjectFlags::Unreachable);
 					}
 				}
 				else
@@ -580,7 +397,7 @@ public:
 	 * @param ReferencingObject UObject which owns the reference (can be NULL)
 	 * @param bAllowReferenceElimination	Whether to allow NULL'ing the reference if RF_PendingKill is set
 	*/
-	FORCEINLINE void HandleObjectReference(TArray<UObject*>& ObjectsToSerialize, const UObject * const ReferencingObject, UObject*& Object, const bool bAllowReferenceElimination, const bool bStrongReference = true)
+	FORCEINLINE void HandleObjectReference(TArray<UObject*>& ObjectsToSerialize, const UObject * const ReferencingObject, UObject*& Object, const bool bAllowReferenceElimination)
 	{
 		// Disregard NULL objects and perform very fast check to see whether object is part of permanent
 		// object pool and should therefore be disregarded. The check doesn't touch the object and is
@@ -680,13 +497,8 @@ public:
 				checkSlow(RootObjectItem->HasAnyFlags(EInternalObjectFlags::ClusterRoot));
 				if (bParallel)
 				{
-					if (RootObjectItem->ThisThreadAtomicallyClearedRFUnreachable() || 
-						(bStrongReference && RootObjectItem->ThisThreadAtomicallyClearedFlag(EInternalObjectFlags::NoStrongReference)))
+					if (RootObjectItem->ThisThreadAtomicallyClearedRFUnreachable())
 					{
-						if (bStrongReference)
-						{
-							RootObjectItem->ThisThreadAtomicallyClearedFlag(EInternalObjectFlags::NoStrongReference);
-						}
 						// Make sure all referenced clusters are marked as reachable too
 						MarkReferencedClustersAsReachable(RootObjectItem->GetClusterIndex(), ObjectsToSerialize);
 					}
@@ -694,21 +506,11 @@ public:
 				else if (RootObjectItem->IsUnreachable())
 				{
 					RootObjectItem->ClearFlags(EInternalObjectFlags::Unreachable);
-					if (bStrongReference)
-					{
-						RootObjectItem->ClearFlags(EInternalObjectFlags::NoStrongReference);
-					}
 					// Make sure all referenced clusters are marked as reachable too
 					MarkReferencedClustersAsReachable(RootObjectItem->GetClusterIndex(), ObjectsToSerialize);
 				}
 			}
 		}
-
-		// The second condition seems to improve perf for multithreaded GC
-		if (bStrongReference && ObjectItem->HasAnyFlags(EInternalObjectFlags::NoStrongReference))
-		{
-			ObjectItem->ThisThreadAtomicallyClearedFlag(EInternalObjectFlags::NoStrongReference);
-		}
 #if PERF_DETAILED_PER_CLASS_GC_STATS
 		GCurrentObjectRegularObjectRefs++;
 #endif
@@ -767,17 +569,15 @@ template <bool bParallel>
 class FGCCollector : public FReferenceCollector
 {
 	FGCReferenceProcessor<bParallel>& ReferenceProcessor;
-	TArray<UObject*>& ObjectArray;
+	FGCArrayStruct& ObjectArrayStruct;
 	bool bAllowEliminatingReferences;
-	bool bShouldHandleAsWeakRef;
 
 public:
 
-	FGCCollector(FGCReferenceProcessor<bParallel>& InProcessor, TArray<UObject*>& InObjectArray)
+	FGCCollector(FGCReferenceProcessor<bParallel>& InProcessor, FGCArrayStruct& InObjectArrayStruct)
 		: ReferenceProcessor(InProcessor)
-		, ObjectArray(InObjectArray)
+		, ObjectArrayStruct(InObjectArrayStruct)
 		, bAllowEliminatingReferences(true)
-		, bShouldHandleAsWeakRef(false)
 	{
 	}
 
@@ -792,7 +592,7 @@ public:
 				ReferencingProperty ? *ReferencingProperty->GetFullName() : TEXT("NULL"));
 		}
 #endif
-		ReferenceProcessor.HandleObjectReference(ObjectArray, const_cast<UObject*>(ReferencingObject), Object, bAllowEliminatingReferences, !bShouldHandleAsWeakRef);
+		ReferenceProcessor.HandleObjectReference(ObjectArrayStruct.ObjectsToSerialize, const_cast<UObject*>(ReferencingObject), Object, bAllowEliminatingReferences);
 	}
 
 	virtual void HandleObjectReference(UObject*& Object, const UObject* ReferencingObject, const UProperty* ReferencingProperty) override
@@ -821,9 +621,10 @@ public:
 		bAllowEliminatingReferences = bAllow;
 	}
 
-	virtual void SetShouldHandleAsWeakRef(bool bWeakRef) override
+	virtual void MarkWeakObjectReferenceForClearing(UObject** WeakReference) override
 	{
-		bShouldHandleAsWeakRef = bWeakRef;
+		// Track this references for later destruction if necessary. These should be relatively rare
+		ObjectArrayStruct.WeakReferences.Add(WeakReference);
 	}
 };
 
@@ -986,7 +787,7 @@ public:
 				}
 				else
 				{
-					ObjectItem->SetFlags(EInternalObjectFlags::Unreachable | EInternalObjectFlags::NoStrongReference);
+					ObjectItem->SetFlags(EInternalObjectFlags::Unreachable);
 				}
 			}
 		}
@@ -1007,7 +808,7 @@ public:
 					// if it is reachable via keep flags we will do this below (or maybe already have)
 					if (RootObjectItem->IsUnreachable()) 
 					{
-						RootObjectItem->ClearFlags(EInternalObjectFlags::Unreachable | EInternalObjectFlags::NoStrongReference);
+						RootObjectItem->ClearFlags(EInternalObjectFlags::Unreachable);
 						// Make sure all referenced clusters are marked as reachable too
 						FGCReferenceProcessorSinglethreaded::MarkReferencedClustersAsReachable(RootObjectItem->GetClusterIndex(), ObjectsToSerialize);
 					}
@@ -1033,7 +834,8 @@ public:
 		DECLARE_SCOPE_CYCLE_COUNTER(TEXT("FRealtimeGC::PerformReachabilityAnalysis"), STAT_FArchiveRealtimeGC_PerformReachabilityAnalysis, STATGROUP_GC);
 
 		/** Growing array of objects that require serialization */
-		TArray<UObject*>& ObjectsToSerialize = *FGCArrayPool::Get().GetArrayFromPool();
+		FGCArrayStruct* ArrayStruct = FGCArrayPool::Get().GetArrayStructFromPool();
+		TArray<UObject*>& ObjectsToSerialize = ArrayStruct->ObjectsToSerialize;
 
 		// Reset object count.
 		GObjectCountDuringLastMarkPhase = 0;
@@ -1052,15 +854,15 @@ public:
 		{
 			FGCReferenceProcessorMultithreaded ReferenceProcessor;
 			TFastReferenceCollector<true, FGCReferenceProcessorMultithreaded, FGCCollectorMultithreaded, FGCArrayPool> ReferenceCollector(ReferenceProcessor, FGCArrayPool::Get());
-			ReferenceCollector.CollectReferences(ObjectsToSerialize);
+			ReferenceCollector.CollectReferences(*ArrayStruct);
 		}
 		else
 		{
 			FGCReferenceProcessorSinglethreaded ReferenceProcessor;
 			TFastReferenceCollector<false, FGCReferenceProcessorSinglethreaded, FGCCollectorSinglethreaded, FGCArrayPool> ReferenceCollector(ReferenceProcessor, FGCArrayPool::Get());
-			ReferenceCollector.CollectReferences(ObjectsToSerialize);
+			ReferenceCollector.CollectReferences(*ArrayStruct);
 		}
-		FGCArrayPool::Get().ReturnToPool(&ObjectsToSerialize);
+		FGCArrayPool::Get().ReturnToPool(ArrayStruct);
 
 #if UE_BUILD_DEBUG
 		FGCArrayPool::Get().CheckLeaks();
@@ -1629,6 +1431,9 @@ void CollectGarbageInternal(EObjectFlags KeepFlags, bool bPerformFullPurge)
 
 			FCoreUObjectDelegates::PreGarbageCollectConditionalBeginDestroy.Broadcast();
 
+			// This nulls all weak references to unreachable objects
+			FGCArrayPool::Get().ClearWeakReferences();
+
 			// Unhash all unreachable objects.
 			const double StartTime = FPlatformTime::Seconds();
 			int32 ClustersRemoved = 0;
@@ -1652,13 +1457,12 @@ void CollectGarbageInternal(EObjectFlags KeepFlags, bool bPerformFullPurge)
 						UE_LOG(LogGarbage, Log, TEXT("Destroying cluster (%d) %s"), ObjectItem->GetClusterIndex(), *Object->GetFullName());
 #endif
 						// Nuke the entire cluster
-						ObjectItem->ClearFlags(EInternalObjectFlags::ClusterRoot | EInternalObjectFlags::NoStrongReference);
+						ObjectItem->ClearFlags(EInternalObjectFlags::ClusterRoot);
 						const int32 ClusterRootIndex = It.GetIndex();
 						FUObjectCluster& Cluster = GUObjectClusters[ObjectItem->GetClusterIndex()];
 						for (int32 ClusterObjectIndex : Cluster.Objects)
 						{
 							FUObjectItem* ClusterObjectItem = GUObjectArray.IndexToObjectUnsafeForGC(ClusterObjectIndex);
-							ClusterObjectItem->ClearFlags(EInternalObjectFlags::NoStrongReference);
 							ClusterObjectItem->SetOwnerIndex(0);
 
 							if (!ClusterObjectItem->HasAnyFlags(EInternalObjectFlags::ReachableInCluster))
@@ -1684,22 +1488,8 @@ void CollectGarbageInternal(EObjectFlags KeepFlags, bool bPerformFullPurge)
 					FScopedCBDProfile Profile(Object);
 					Object->ConditionalBeginDestroy();
 				}
-				else if (ObjectItem->IsNoStrongReference())
-				{
-					ObjectItem->ClearNoStrongReference();
-					ObjectItem->SetPendingKill();
-					if (ObjectItem->HasAnyFlags(EInternalObjectFlags::ClusterRoot))
-					{
-						// If a cluster root has no strong reference then none of its object had a strong reference either, mark all as pending kill
-						FUObjectCluster& Cluster = GUObjectClusters[ObjectItem->GetClusterIndex()];
-						for (int32 ClusterObjectIndex : Cluster.Objects)
-						{
-							FUObjectItem* ClusterObjectItem = GUObjectArray.IndexToObjectUnsafeForGC(ClusterObjectIndex);
-							ClusterObjectItem->SetPendingKill();
-						}
-					}
-				}
 			}
+
 			UE_LOG(LogGarbage, Log, TEXT("%f ms for unhashing unreachable objects. Clusters removed: %d.   Items %d Cluster Items %d"), (FPlatformTime::Seconds() - StartTime) * 1000, ClustersRemoved, Items, ClusterItems);
 			FCoreUObjectDelegates::PostGarbageCollectConditionalBeginDestroy.Broadcast();
 		}
diff --git a/Engine/Source/Runtime/CoreUObject/Private/UObject/PropertyArray.cpp b/Engine/Source/Runtime/CoreUObject/Private/UObject/PropertyArray.cpp
--- a/Engine/Source/Runtime/CoreUObject/Private/UObject/PropertyArray.cpp
+++ b/Engine/Source/Runtime/CoreUObject/Private/UObject/PropertyArray.cpp
@@ -160,6 +160,7 @@ void UArrayProperty::SerializeItem( FArchive& Ar, void* Value, void const* Defau
 
 		const FCustomPropertyListNode* CustomPropertyList = Ar.ArCustomPropertyList;
 		const FCustomPropertyListNode* PropertyNode = CustomPropertyList;
+		FSerializedPropertyScope SerializedProperty(Ar, Inner, this);
 		while (PropertyNode && i < n && !bSerializeRemainingItems)
 		{
 			if (PropertyNode->Property != Inner)
@@ -197,6 +198,7 @@ void UArrayProperty::SerializeItem( FArchive& Ar, void* Value, void const* Defau
 		Ar.ArUseCustomPropertyList = false;
 
 		// Serialize each item until we get to the end of the array
+		FSerializedPropertyScope SerializedProperty(Ar, Inner, this);
 		while (i < n)
 		{
 #if WITH_EDITOR
diff --git a/Engine/Source/Runtime/CoreUObject/Private/UObject/PropertyAssetObject.cpp b/Engine/Source/Runtime/CoreUObject/Private/UObject/PropertyAssetObject.cpp
--- a/Engine/Source/Runtime/CoreUObject/Private/UObject/PropertyAssetObject.cpp
+++ b/Engine/Source/Runtime/CoreUObject/Private/UObject/PropertyAssetObject.cpp
@@ -65,75 +65,45 @@ void UAssetObjectProperty::ExportTextItem( FString& ValueStr, const void* Proper
 {
 	FAssetPtr& AssetPtr = *(FAssetPtr*)PropertyValue;
 
-	FStringAssetReference ID;
+	FStringAssetReference AssetReference;
 	UObject *Object = AssetPtr.Get();
 
 	if (Object)
 	{
 		// Use object in case name has changed.
-		ID = FStringAssetReference(Object);
+		AssetReference = FStringAssetReference(Object);
 	}
 	else
 	{
-		ID = AssetPtr.GetUniqueID();
+		AssetReference = AssetPtr.GetUniqueID();
 	}
 
 	if (0 != (PortFlags & PPF_ExportCpp))
 	{
-		ValueStr += FString::Printf(TEXT("FStringAssetReference(TEXT(\"%s\"))"), *ID.ToString().ReplaceCharWithEscapedChar());
+		ValueStr += FString::Printf(TEXT("FStringAssetReference(TEXT(\"%s\"))"), *AssetReference.ToString().ReplaceCharWithEscapedChar());
 		return;
 	}
 
-	if (!ID.ToString().IsEmpty())
-	{
-		ValueStr += ID.ToString();
-	}
-	else
-	{
-		ValueStr += TEXT("None");
-	}
+	AssetReference.ExportTextItem(ValueStr, AssetReference, Parent, PortFlags, ExportRootScope);
 }
 
 const TCHAR* UAssetObjectProperty::ImportText_Internal( const TCHAR* InBuffer, void* Data, int32 PortFlags, UObject* Parent, FOutputDevice* ErrorText ) const
 {
 	FAssetPtr& AssetPtr = *(FAssetPtr*)Data;
 
-	FString NewPath;
-	const TCHAR* Buffer = InBuffer;
-	const TCHAR* NewBuffer = UPropertyHelpers::ReadToken( Buffer, NewPath, 1 );
-	if( !NewBuffer )
-	{
-		return NULL;
-	}
-	Buffer = NewBuffer;
-	if( NewPath==TEXT("None") )
+	FStringAssetReference AssetReference;
+
+	if (AssetReference.ImportTextItem(InBuffer, PortFlags, Parent, ErrorText))
 	{
-		AssetPtr = NULL;
+		AssetPtr = AssetReference;
+		return InBuffer;
 	}
+
 	else
 	{
-		if( *Buffer == TCHAR('\'') )
-		{
-			// A ' token likely means we're looking at an asset string in the form "Texture2d'/Game/UI/HUD/Actions/Barrel'" and we need to read and append the path part
-			// We have to skip over the first ' as UPropertyHelpers::ReadToken doesn't read single-quoted strings correctly, but does read a path correctly
-			NewPath += *Buffer++; // Append the leading '
-			NewBuffer = UPropertyHelpers::ReadToken( Buffer, NewPath, 1 );
-			if( !NewBuffer )
-			{
-				return NULL;
-			}
-			Buffer = NewBuffer;
-			if( *Buffer != TCHAR('\'') )
-			{
-				return NULL;
-			}
-			NewPath += *Buffer++; // Append the trailing '
-		}
-		FStringAssetReference ID(NewPath);
-		AssetPtr = ID;
+		AssetPtr = nullptr;
+		return nullptr;
 	}
-
-	return Buffer;
 }
 
 bool UAssetObjectProperty::ConvertFromType(const FPropertyTag& Tag, FArchive& Ar, uint8* Data, UStruct* DefaultsStruct, bool& bOutAdvanceProperty)
diff --git a/Engine/Source/Runtime/CoreUObject/Private/UObject/PropertyMap.cpp b/Engine/Source/Runtime/CoreUObject/Private/UObject/PropertyMap.cpp
--- a/Engine/Source/Runtime/CoreUObject/Private/UObject/PropertyMap.cpp
+++ b/Engine/Source/Runtime/CoreUObject/Private/UObject/PropertyMap.cpp
@@ -187,7 +187,7 @@ void UMapProperty::LinkInternal(FArchive& Ar)
 {
 	check(KeyProp && ValueProp);
 
-	if (auto* MyLinker = GetLinker())
+	if (FLinkerLoad* MyLinker = GetLinker())
 	{
 		MyLinker->Preload(this);
 	}
@@ -279,6 +279,7 @@ void UMapProperty::SerializeItem(FArchive& Ar, void* Value, const void* Defaults
 			TempKeyStorage = (uint8*)FMemory::Malloc(MapLayout.SetLayout.Size);
 			KeyProp->InitializeValue(TempKeyStorage);
 
+			FSerializedPropertyScope SerializedProperty(Ar, KeyProp, this);
 			for (; NumKeysToRemove; --NumKeysToRemove)
 			{
 				// Read key into temporary storage
@@ -307,8 +308,11 @@ void UMapProperty::SerializeItem(FArchive& Ar, void* Value, const void* Defaults
 		for (; NumEntries; --NumEntries)
 		{
 			// Read key into temporary storage
-			KeyProp->SerializeItem(Ar, TempKeyStorage);
-
+			{
+				FSerializedPropertyScope SerializedProperty(Ar, KeyProp, this);
+				KeyProp->SerializeItem(Ar, TempKeyStorage);
+			}
+			
 			// Add a new default value if the key doesn't currently exist in the map
 			int32 NextPairIndex = MapHelper.FindMapIndexWithKey(TempKeyStorage);
 			if (NextPairIndex == INDEX_NONE)
@@ -322,7 +326,10 @@ void UMapProperty::SerializeItem(FArchive& Ar, void* Value, const void* Defaults
 			KeyProp->CopyCompleteValue_InContainer(NextPairPtr, TempKeyStorage);
 
 			// Deserialize value
-			ValueProp->SerializeItem(Ar, NextPairPtr + MapLayout.ValueOffset);
+			{
+				FSerializedPropertyScope SerializedProperty(Ar, ValueProp, this);
+				ValueProp->SerializeItem(Ar, NextPairPtr + MapLayout.ValueOffset);
+			}
 		}
 
 		MapHelper.Rehash();
@@ -356,9 +363,12 @@ void UMapProperty::SerializeItem(FArchive& Ar, void* Value, const void* Defaults
 		// Write out the missing keys
 		int32 MissingKeysNum = Indices.Num();
 		Ar << MissingKeysNum;
-		for (int32 Index : Indices)
 		{
-			KeyProp->SerializeItem(Ar, DefaultsHelper.GetPairPtr(Index));
+			FSerializedPropertyScope SerializedProperty(Ar, KeyProp, this);
+			for (int32 Index : Indices)
+			{
+				KeyProp->SerializeItem(Ar, DefaultsHelper.GetPairPtr(Index));
+			}
 		}
 
 		// Write out differences from defaults
@@ -388,8 +398,14 @@ void UMapProperty::SerializeItem(FArchive& Ar, void* Value, const void* Defaults
 			{
 				uint8* ValuePairPtr = MapHelper.GetPairPtrWithoutCheck(Index);
 
-				KeyProp  ->SerializeItem(Ar, ValuePairPtr);
-				ValueProp->SerializeItem(Ar, ValuePairPtr + MapLayout.ValueOffset);
+				{
+					FSerializedPropertyScope SerializedProperty(Ar, KeyProp, this);
+					KeyProp->SerializeItem(Ar, ValuePairPtr);
+				}
+				{
+					FSerializedPropertyScope SerializedProperty(Ar, ValueProp, this);
+					ValueProp->SerializeItem(Ar, ValuePairPtr + MapLayout.ValueOffset);
+				}
 			}
 		}
 		else
@@ -402,8 +418,14 @@ void UMapProperty::SerializeItem(FArchive& Ar, void* Value, const void* Defaults
 				{
 					uint8* ValuePairPtr = MapHelper.GetPairPtrWithoutCheck(Index);
 
-					KeyProp  ->SerializeItem(Ar, ValuePairPtr);
-					ValueProp->SerializeItem(Ar, ValuePairPtr + MapLayout.ValueOffset);
+					{
+						FSerializedPropertyScope SerializedProperty(Ar, KeyProp, this);
+						KeyProp->SerializeItem(Ar, ValuePairPtr);
+					}
+					{
+						FSerializedPropertyScope SerializedProperty(Ar, ValueProp, this);
+						ValueProp->SerializeItem(Ar, ValuePairPtr + MapLayout.ValueOffset);
+					}
 
 					--Num;
 				}
@@ -490,7 +512,7 @@ void UMapProperty::ExportTextItem(FString& ValueStr, const void* PropertyValue,
 	}
 
 	uint8* StructDefaults = nullptr;
-	if (auto* StructValueProp = dynamic_cast<UStructProperty*>(ValueProp))
+	if (UStructProperty* StructValueProp = dynamic_cast<UStructProperty*>(ValueProp))
 	{
 		checkSlow(StructValueProp->Struct);
 
@@ -805,7 +827,7 @@ void UMapProperty::InstanceSubobjects(void* Data, void const* DefaultData, UObje
 
 bool UMapProperty::SameType(const UProperty* Other) const
 {
-	auto* MapProp = (UMapProperty*)Other;
+	UMapProperty* MapProp = (UMapProperty*)Other;
 	return Super::SameType(Other) && KeyProp && ValueProp && KeyProp->SameType(MapProp->KeyProp) && ValueProp->SameType(MapProp->ValueProp);
 }
 
diff --git a/Engine/Source/Runtime/CoreUObject/Private/UObject/PropertySet.cpp b/Engine/Source/Runtime/CoreUObject/Private/UObject/PropertySet.cpp
--- a/Engine/Source/Runtime/CoreUObject/Private/UObject/PropertySet.cpp
+++ b/Engine/Source/Runtime/CoreUObject/Private/UObject/PropertySet.cpp
@@ -251,6 +251,7 @@ void USetProperty::SerializeItem(FArchive& Ar, void* Value, const void* Defaults
 			TempElementStorage = (uint8*)FMemory::Malloc(SetLayout.Size);
 			ElementProp->InitializeValue(TempElementStorage);
 
+			FSerializedPropertyScope SerializedProperty(Ar, ElementProp, this);
 			for (; NumElementsToRemove; --NumElementsToRemove)
 			{
 				// Read key into temporary storage
@@ -275,6 +276,7 @@ void USetProperty::SerializeItem(FArchive& Ar, void* Value, const void* Defaults
 			ElementProp->InitializeValue(TempElementStorage);
 		}
 
+		FSerializedPropertyScope SerializedProperty(Ar, ElementProp, this);
 		// Read remaining items into container
 		for (; Num; --Num)
 		{
@@ -323,9 +325,12 @@ void USetProperty::SerializeItem(FArchive& Ar, void* Value, const void* Defaults
 		// Write out the removed elements
 		int32 RemovedElementsNum = Indices.Num();
 		Ar << RemovedElementsNum;
-		for (int32 Index : Indices)
 		{
-			ElementProp->SerializeItem(Ar, DefaultsHelper.GetElementPtr(Index));
+			FSerializedPropertyScope SerializedProperty(Ar, ElementProp, this);
+			for (int32 Index : Indices)
+			{
+				ElementProp->SerializeItem(Ar, DefaultsHelper.GetElementPtr(Index));
+			}
 		}
 
 		// Write out added elements
@@ -351,6 +356,8 @@ void USetProperty::SerializeItem(FArchive& Ar, void* Value, const void* Defaults
 			// Write out differences from defaults
 			int32 Num = Indices.Num();
 			Ar << Num;
+
+			FSerializedPropertyScope SerializedProperty(Ar, ElementProp, this);
 			for (int32 Index : Indices)
 			{
 				uint8* ElementPtr = SetHelper.GetElementPtrWithoutCheck(Index);
@@ -362,6 +369,8 @@ void USetProperty::SerializeItem(FArchive& Ar, void* Value, const void* Defaults
 		{
 			int32 Num = SetHelper.Num();
 			Ar << Num;
+
+			FSerializedPropertyScope SerializedProperty(Ar, ElementProp, this);
 			for (int32 Index = 0; Num; ++Index)
 			{
 				if (SetHelper.IsValidIndex(Index))
diff --git a/Engine/Source/Runtime/CoreUObject/Private/UObject/UObjectClusters.cpp b/Engine/Source/Runtime/CoreUObject/Private/UObject/UObjectClusters.cpp
--- a/Engine/Source/Runtime/CoreUObject/Private/UObject/UObjectClusters.cpp
+++ b/Engine/Source/Runtime/CoreUObject/Private/UObject/UObjectClusters.cpp
@@ -182,7 +182,7 @@ void FUObjectClusterContainer::DissolveClusterAndMarkObjectsAsUnreachable(const
 		ClusterObjectItem->SetOwnerIndex(0);
 		if (ClusterObjectIndex < CurrentIndex)
 		{
-			ClusterObjectItem->SetFlags(EInternalObjectFlags::Unreachable | EInternalObjectFlags::NoStrongReference);
+			ClusterObjectItem->SetFlags(EInternalObjectFlags::Unreachable);
 		}
 	}
 
@@ -201,7 +201,7 @@ void FUObjectClusterContainer::DissolveClusterAndMarkObjectsAsUnreachable(const
 		{
 			if (ReferencedByClusterRootIndex < CurrentIndex)
 			{
-				ReferencedByClusterRootItem->SetFlags(EInternalObjectFlags::Unreachable | EInternalObjectFlags::NoStrongReference);
+				ReferencedByClusterRootItem->SetFlags(EInternalObjectFlags::Unreachable);
 			}
 			DissolveClusterAndMarkObjectsAsUnreachable(CurrentIndex, ReferencedByClusterRootItem);
 		}
@@ -378,106 +378,6 @@ static FAutoConsoleCommand FindStaleClustersCommand(
 
 #endif // !UE_BUILD_SHIPPING
 
-/**
-* Pool for reducing allocations when constructing clusters
-*/
-class FClusterArrayPool
-{
-public:
-
-	/**
-	* Gets the singleton instance of the FObjectArrayPool
-	* @return Pool singleton.
-	*/
-	FORCEINLINE static FClusterArrayPool& Get()
-	{
-		static FClusterArrayPool Singleton;
-		return Singleton;
-	}
-
-	/**
-	* Gets an event from the pool or creates one if necessary.
-	*
-	* @return The array.
-	* @see ReturnToPool
-	*/
-	FORCEINLINE TArray<UObject*>* GetArrayFromPool()
-	{
-		TArray<UObject*>* Result = Pool.Pop();
-		if (!Result)
-		{
-			Result = new TArray<UObject*>();
-		}
-		check(Result);
-#if UE_BUILD_DEBUG
-		NumberOfUsedArrays.Increment();
-#endif // UE_BUILD_DEBUG
-		return Result;
-	}
-
-	/**
-	* Returns an array to the pool.
-	*
-	* @param Array The array to return.
-	* @see GetArrayFromPool
-	*/
-	FORCEINLINE void ReturnToPool(TArray<UObject*>* Array)
-	{
-#if UE_BUILD_DEBUG
-		const int32 CheckUsedArrays = NumberOfUsedArrays.Decrement();
-		checkSlow(CheckUsedArrays >= 0);
-#endif // UE_BUILD_DEBUG
-		check(Array);
-		Array->Reset();
-		Pool.Push(Array);
-	}
-
-	/** Performs memory cleanup */
-	void Cleanup()
-	{
-#if UE_BUILD_DEBUG
-		const int32 CheckUsedArrays = NumberOfUsedArrays.GetValue();
-		checkSlow(CheckUsedArrays == 0);
-#endif // UE_BUILD_DEBUG
-
-		uint32 FreedMemory = 0;
-		TArray< TArray<UObject*>* > AllArrays;
-		Pool.PopAll(AllArrays);
-		for (TArray<UObject*>* Array : AllArrays)
-		{
-			FreedMemory += Array->GetAllocatedSize();
-			delete Array;
-		}
-		UE_LOG(LogObj, Log, TEXT("Freed %ub from %d cluster array pools."), FreedMemory, AllArrays.Num());
-	}
-
-#if UE_BUILD_DEBUG
-	void CheckLeaks()
-	{
-		// This function is called after cluster has been created so at this point there should be no
-		// arrays used by cluster creation code and all should be returned to the pool
-		const int32 LeakedPoolArrays = NumberOfUsedArrays.GetValue();
-		checkSlow(LeakedPoolArrays == 0);
-	}
-#endif
-
-private:
-
-	/** Holds the collection of recycled arrays. */
-	TLockFreePointerListLIFO< TArray<UObject*> > Pool;
-
-#if UE_BUILD_DEBUG
-	/** Number of arrays currently acquired from the pool by clusters */
-	FThreadSafeCounter NumberOfUsedArrays;
-#endif // UE_BUILD_DEBUG
-};
-
-/** Called on shutdown to free cluster memory */
-void CleanupClusterArrayPools()
-{
-	FClusterArrayPool::Get().Cleanup();
-}
-
 /**
  * Handles UObject references found by TFastReferenceCollector
  */
@@ -679,24 +579,24 @@ template <class TProcessor>
 class TClusterCollector : public FReferenceCollector
 {
 	TProcessor& Processor;
-	TArray<UObject*>& ObjectArray;
+	FGCArrayStruct& ObjectArrayStruct;
 
 public:
-	TClusterCollector(TProcessor& InProcessor, TArray<UObject*>& InObjectArray)
+	TClusterCollector(TProcessor& InProcessor, FGCArrayStruct& InObjectArrayStruct)
 		: Processor(InProcessor)
-		, ObjectArray(InObjectArray)
+		, ObjectArrayStruct(InObjectArrayStruct)
 	{
 	}
 	virtual void HandleObjectReference(UObject*& Object, const UObject* ReferencingObject, const UProperty* ReferencingProperty) override
 	{
-		Processor.HandleTokenStreamObjectReference(ObjectArray, const_cast<UObject*>(ReferencingObject), Object, INDEX_NONE, false);
+		Processor.HandleTokenStreamObjectReference(ObjectArrayStruct.ObjectsToSerialize, const_cast<UObject*>(ReferencingObject), Object, INDEX_NONE, false);
 	}
 	virtual void HandleObjectReferences(UObject** InObjects, const int32 ObjectNum, const UObject* ReferencingObject, const UProperty* InReferencingProperty) override
 	{
 		for (int32 ObjectIndex = 0; ObjectIndex < ObjectNum; ++ObjectIndex)
 		{
 			UObject*& Object = InObjects[ObjectIndex];
-			Processor.HandleTokenStreamObjectReference(ObjectArray, const_cast<UObject*>(ReferencingObject), Object, INDEX_NONE, false);
+			Processor.HandleTokenStreamObjectReference(ObjectArrayStruct.ObjectsToSerialize, const_cast<UObject*>(ReferencingObject), Object, INDEX_NONE, false);
 		}
 	}
 	virtual bool IsIgnoringArchetypeRef() const override
@@ -735,13 +635,14 @@ void UObjectBaseUtility::AddToCluster(UObjectBaseUtility* ClusterRootOrObjectFro
 		if (!bAddAsMutableObject)
 		{
 			FClusterReferenceProcessor Processor(ClusterRootIndex, *Cluster);
-			TFastReferenceCollector<false, FClusterReferenceProcessor, TClusterCollector<FClusterReferenceProcessor>, FClusterArrayPool, true> ReferenceCollector(Processor, FClusterArrayPool::Get());
-			TArray<UObject*> ObjectsToProcess;
+			TFastReferenceCollector<false, FClusterReferenceProcessor, TClusterCollector<FClusterReferenceProcessor>, FGCArrayPool, true> ReferenceCollector(Processor, FGCArrayPool::Get());
+			FGCArrayStruct ArrayStruct;
+			TArray<UObject*>& ObjectsToProcess = ArrayStruct.ObjectsToSerialize;
 			UObject* ThisObject = static_cast<UObject*>(this);
 			Processor.HandleTokenStreamObjectReference(ObjectsToProcess, static_cast<UObject*>(ClusterRootOrObjectFromCluster), ThisObject, INDEX_NONE, true);
 			if (ObjectsToProcess.Num())
 			{
-				ReferenceCollector.CollectReferences(ObjectsToProcess);
+				ReferenceCollector.CollectReferences(ArrayStruct);
 			}
 		}
 		else
@@ -793,12 +694,13 @@ void UObjectBaseUtility::CreateCluster()
 
 	// Collect all objects referenced by cluster root and by all objects it's referencing
 	FClusterReferenceProcessor Processor(InternalIndex, Cluster);
-	TFastReferenceCollector<false, FClusterReferenceProcessor, TClusterCollector<FClusterReferenceProcessor>, FClusterArrayPool, true> ReferenceCollector(Processor, FClusterArrayPool::Get());
-	TArray<UObject*> ObjectsToProcess;
+	TFastReferenceCollector<false, FClusterReferenceProcessor, TClusterCollector<FClusterReferenceProcessor>, FGCArrayPool, true> ReferenceCollector(Processor, FGCArrayPool::Get());
+	FGCArrayStruct ArrayStruct;
+	TArray<UObject*>& ObjectsToProcess = ArrayStruct.ObjectsToSerialize;
 	ObjectsToProcess.Add(static_cast<UObject*>(this));
-	ReferenceCollector.CollectReferences(ObjectsToProcess);
+	ReferenceCollector.CollectReferences(ArrayStruct);
 #if UE_BUILD_DEBUG
-	FClusterArrayPool::Get().CheckLeaks();
+	FGCArrayPool::Get().CheckLeaks();
 #endif
 
 	if (Cluster.Objects.Num())
@@ -995,9 +897,10 @@ bool VerifyClusterAssumptions(UObject* ClusterRootObject)
 {
 	// Collect all objects referenced by cluster root and by all objects it's referencing
 	FClusterVerifyReferenceProcessor Processor(ClusterRootObject);
-	TFastReferenceCollector<false, FClusterVerifyReferenceProcessor, TClusterCollector<FClusterVerifyReferenceProcessor>, FClusterArrayPool> ReferenceCollector(Processor, FClusterArrayPool::Get());
-	TArray<UObject*> ObjectsToProcess;
+	TFastReferenceCollector<false, FClusterVerifyReferenceProcessor, TClusterCollector<FClusterVerifyReferenceProcessor>, FGCArrayPool> ReferenceCollector(Processor, FGCArrayPool::Get());
+	FGCArrayStruct ArrayStruct;
+	TArray<UObject*>& ObjectsToProcess = ArrayStruct.ObjectsToSerialize;
 	ObjectsToProcess.Add(ClusterRootObject);
-	ReferenceCollector.CollectReferences(ObjectsToProcess);
+	ReferenceCollector.CollectReferences(ArrayStruct);
 	return Processor.NoExternalReferencesFound();
 }
diff --git a/Engine/Source/Runtime/CoreUObject/Private/UObject/UObjectGlobals.cpp b/Engine/Source/Runtime/CoreUObject/Private/UObject/UObjectGlobals.cpp
--- a/Engine/Source/Runtime/CoreUObject/Private/UObject/UObjectGlobals.cpp
+++ b/Engine/Source/Runtime/CoreUObject/Private/UObject/UObjectGlobals.cpp
@@ -1235,12 +1235,28 @@ UPackage* LoadPackageInternal(UPackage* InOuter, const TCHAR* InLongPackageNameO
 		};
 
 #if WITH_EDITORONLY_DATA
-		if (!(LoadFlags & (LOAD_IsVerifying|LOAD_EditorOnly)) &&
-			(!ImportLinker || !ImportLinker->GetSerializedProperty() || !ImportLinker->GetSerializedProperty()->IsEditorOnlyProperty()))
+		if (!(LoadFlags & (LOAD_IsVerifying|LOAD_EditorOnly)))
 		{
-			// If this package hasn't been loaded as part of import verification and there's no import linker or the
-			// currently serialized property is not editor-only mark this package as runtime.
-			Result->SetLoadedByEditorPropertiesOnly(false);
+			bool bIsEditorOnly = false;
+			UProperty* SerializingProperty = ImportLinker ? ImportLinker->GetSerializedProperty() : nullptr;
+			
+			// Check property parent chain
+			while (SerializingProperty)
+			{
+				if (SerializingProperty->IsEditorOnlyProperty())
+				{
+					bIsEditorOnly = true;
+					break;
+				}
+				SerializingProperty = Cast<UProperty>(SerializingProperty->GetOuter());
+			}
+
+			if (!bIsEditorOnly)
+			{
+				// If this package hasn't been loaded as part of import verification and there's no import linker or the
+				// currently serialized property is not editor-only mark this package as runtime.
+				Result->SetLoadedByEditorPropertiesOnly(false);
+			}
 		}
 #endif
 
