Fix shadow variable warning when compiling natively on Linux with Clang 3.8.

diff --git a/Engine/Source/Programs/UnrealHeaderTool/Private/CodeGenerator.cpp b/Engine/Source/Programs/UnrealHeaderTool/Private/CodeGenerator.cpp
--- a/Engine/Source/Programs/UnrealHeaderTool/Private/CodeGenerator.cpp
+++ b/Engine/Source/Programs/UnrealHeaderTool/Private/CodeGenerator.cpp
@@ -1576,9 +1576,9 @@ void FNativeClassHeaderGenerator::OutputProperty(FOutputDevice& OutputDevice, TA
 	FString PropVariableName = FString::Printf(TEXT("NewProp_%s"), *PropName);
 
 	// Helper to handle the creation of the underlying properties if they're enum properties
-	auto HandleUnderlyingEnumProperty = [this, &PropertyNamesAndPointers, &OutputDevice, Spaces](UProperty* Prop, const FString& OuterName)
+	auto HandleUnderlyingEnumProperty = [this, &PropertyNamesAndPointers, &OutputDevice, Spaces](UProperty* LocalProp, const FString& OuterName)
 	{
-		if (UEnumProperty* EnumProp = Cast<UEnumProperty>(Prop))
+		if (UEnumProperty* EnumProp = Cast<UEnumProperty>(LocalProp))
 		{
 			FString PropVarName = OuterName + TEXT("_Underlying");
 			PropertyNew(OutputDevice, EnumProp->UnderlyingProp, TEXT("0"), *PropVarName, Spaces);
