Copying //UE4/Dev-Framework to //UE4/Dev-Main (Source: //UE4/Dev-Framework @ 3582324)

diff --git a/Engine/Source/Programs/UnrealHeaderTool/Private/HeaderParser.cpp b/Engine/Source/Programs/UnrealHeaderTool/Private/HeaderParser.cpp
--- a/Engine/Source/Programs/UnrealHeaderTool/Private/HeaderParser.cpp
+++ b/Engine/Source/Programs/UnrealHeaderTool/Private/HeaderParser.cpp
@@ -68,6 +68,10 @@ enum {MAX_ARRAY_SIZE=2048};
 
 static const FName NAME_ToolTip(TEXT("ToolTip"));
 EGeneratedCodeVersion FHeaderParser::DefaultGeneratedCodeVersion = EGeneratedCodeVersion::V1;
+TArray<FString> FHeaderParser::StructsWithNoPrefix;
+TArray<FString> FHeaderParser::StructsWithTPrefix;
+TArray<FString> FHeaderParser::DelegateParameterCountStrings;
+TMap<FString, FString> FHeaderParser::TypeRedirectMap;
 TMap<UClass*, ClassDefinitionRange> ClassDefinitionRanges;
 /**
  * Dirty hack global variable to allow different result codes passed through
@@ -820,19 +824,19 @@ namespace
 				return Result;
 			}
 
-			case CPT_AssetObjectReference:
+			case CPT_SoftObjectReference:
 				check(VarProperty.PropertyClass);
 
 				if (VarProperty.PropertyClass->IsChildOf(UClass::StaticClass()))
 				{
-					UAssetClassProperty* Result = new (EC_InternalUseOnlyConstructor, Scope, Name, ObjectFlags) UAssetClassProperty(FObjectInitializer());
+					USoftClassProperty* Result = new (EC_InternalUseOnlyConstructor, Scope, Name, ObjectFlags) USoftClassProperty(FObjectInitializer());
 					Result->MetaClass     = VarProperty.MetaClass;
 					Result->PropertyClass = VarProperty.PropertyClass;
 					return Result;
 				}
 				else
 				{
-					UAssetObjectProperty* Result = new (EC_InternalUseOnlyConstructor, Scope, Name, ObjectFlags) UAssetObjectProperty(FObjectInitializer());
+					USoftObjectProperty* Result = new (EC_InternalUseOnlyConstructor, Scope, Name, ObjectFlags) USoftObjectProperty(FObjectInitializer());
 					Result->PropertyClass = VarProperty.PropertyClass;
 					return Result;
 				}
@@ -1037,7 +1041,7 @@ namespace
 
 		const bool bSupportedType = Property->IsA<UInterfaceProperty>()
 			|| Property->IsA<UClassProperty>()
-			|| Property->IsA<UAssetObjectProperty>()
+			|| Property->IsA<USoftObjectProperty>()
 			|| Property->IsA<UObjectProperty>()
 			|| Property->IsA<UFloatProperty>()
 			|| Property->IsA<UIntProperty>()
@@ -1097,6 +1101,8 @@ FClass* FHeaderParser::GetQualifiedClass(const FClasses& AllClasses, const TCHAR
 	FToken Token;
 	if (GetIdentifier(Token))
 	{
+		RedirectTypeIdentifier(Token);
+
 		FCString::Strncat( ClassName, Token.Identifier, ARRAY_COUNT(ClassName) );
 	}
 
@@ -2043,12 +2049,16 @@ UScriptStruct* FHeaderParser::CompileStructDeclaration(FClasses& AllClasses)
 		FToken ParentScope, ParentName;
 		if (GetIdentifier( ParentScope ))
 		{
+			RedirectTypeIdentifier(ParentScope);
+
 			TSharedPtr<FScope> StructScope = Scope;
 			FString ParentStructNameInScript = FString(ParentScope.Identifier);
 			if (MatchSymbol(TEXT(".")))
 			{
 				if (GetIdentifier(ParentName))
 				{
+					RedirectTypeIdentifier(ParentName);
+
 					ParentStructNameInScript = FString(ParentName.Identifier);
 					FString ParentNameStripped = GetClassNameWithPrefixRemoved(ParentScope.Identifier);
 					FClass* StructClass = AllClasses.FindClass(*ParentNameStripped);
@@ -3623,6 +3633,8 @@ void FHeaderParser::GetVarType(
 		FError::Throwf(TEXT("%s: Missing variable type"), GetHintText(VariableCategory));
 	}
 
+	RedirectTypeIdentifier(VarType);
+
 	if ( VarType.Matches(TEXT("int8")) )
 	{
 		VarProperty = FPropertyBase(CPT_Int8);
@@ -4076,8 +4088,8 @@ void FHeaderParser::GetVarType(
 			UClass* TempClass = NULL;
 
 			const bool bIsLazyPtrTemplate        = VarType.Matches(TEXT("TLazyObjectPtr"));
-			const bool bIsAssetPtrTemplate       = VarType.Matches(TEXT("TAssetPtr"));
-			const bool bIsAssetClassTemplate     = VarType.Matches(TEXT("TAssetSubclassOf"));
+			const bool bIsSoftObjectPtrTemplate  = VarType.Matches(TEXT("TSoftObjectPtr"));
+			const bool bIsSoftClassPtrTemplate   = VarType.Matches(TEXT("TSoftClassPtr"));
 			const bool bIsWeakPtrTemplate        = VarType.Matches(TEXT("TWeakObjectPtr"));
 			const bool bIsAutoweakPtrTemplate    = VarType.Matches(TEXT("TAutoWeakObjectPtr"));
 			const bool bIsScriptInterfaceWrapper = VarType.Matches(TEXT("TScriptInterface"));
@@ -4085,7 +4097,7 @@ void FHeaderParser::GetVarType(
 
 			bool bIsWeak     = false;
 			bool bIsLazy     = false;
-			bool bIsAsset    = false;
+			bool bIsSoft     = false;
 			bool bWeakIsAuto = false;
 
 			if (VarType.Matches(TEXT("TSubclassOf")))
@@ -4097,12 +4109,12 @@ void FHeaderParser::GetVarType(
 				TempClass = UInterface::StaticClass();
 				Flags |= CPF_UObjectWrapper;
 			}
-			else if (bIsAssetClassTemplate)
+			else if (bIsSoftClassPtrTemplate)
 			{
 				TempClass = UClass::StaticClass();
-				bIsAsset = true;
+				bIsSoft = true;
 			}
-			else if (bIsLazyPtrTemplate || bIsWeakPtrTemplate || bIsAutoweakPtrTemplate || bIsScriptInterfaceWrapper || bIsAssetPtrTemplate || bIsSubobjectPtrTemplate)
+			else if (bIsLazyPtrTemplate || bIsWeakPtrTemplate || bIsAutoweakPtrTemplate || bIsScriptInterfaceWrapper || bIsSoftObjectPtrTemplate || bIsSubobjectPtrTemplate)
 			{
 				RequireSymbol(TEXT("<"), VarType.Identifier);
 
@@ -4116,6 +4128,8 @@ void FHeaderParser::GetVarType(
 				FToken InnerClass;
 				if (GetIdentifier(InnerClass))
 				{
+					RedirectTypeIdentifier(InnerClass);
+
 					TempClass = AllClasses.FindScriptClass(InnerClass.Identifier);
 					if (TempClass == nullptr)
 					{
@@ -4135,9 +4149,9 @@ void FHeaderParser::GetVarType(
 					{
 						bIsWeak = true;
 					}
-					else if (bIsAssetPtrTemplate)
+					else if (bIsSoftObjectPtrTemplate)
 					{
-						bIsAsset = true;
+						bIsSoft = true;
 					}
 					else if (bIsSubobjectPtrTemplate)
 					{
@@ -4163,7 +4177,7 @@ void FHeaderParser::GetVarType(
 				bHandledType = true;
 
 				bool bAllowWeak = !(Disallow & CPF_AutoWeak); // if it is not allowing anything, force it strong. this is probably a function arg
-				VarProperty = FPropertyBase( TempClass, bAllowWeak && bIsWeak, bWeakIsAuto, bIsLazy, bIsAsset );
+				VarProperty = FPropertyBase(TempClass, bAllowWeak && bIsWeak, bWeakIsAuto, bIsLazy, bIsSoft);
 				if (TempClass->IsChildOf(UClass::StaticClass()))
 				{
 					if ( MatchSymbol(TEXT("<")) )
@@ -4180,6 +4194,8 @@ void FHeaderParser::GetVarType(
 							FError::Throwf(TEXT("'class': Missing class limitor"));
 						}
 
+						RedirectTypeIdentifier(Limitor);
+
 						VarProperty.MetaClass = AllClasses.FindScriptClassOrThrow(Limitor.Identifier);
 
 						RequireSymbol( TEXT(">"), TEXT("'class limitor'"), ESymbolParseOption::CloseTemplateBracket );
@@ -4199,9 +4215,9 @@ void FHeaderParser::GetVarType(
 						FError::Throwf(TEXT("Class variables cannot be lazy, they are always strong."));
 					}
 
-					if (bIsAssetPtrTemplate)
+					if (bIsSoftObjectPtrTemplate)
 					{
-						FError::Throwf(TEXT("Class variables cannot be stored in TAssetPtr, use TAssetSubclassOf instead."));
+						FError::Throwf(TEXT("Class variables cannot be stored in TSoftObjectPtr, use TSoftClassPtr instead."));
 					}
 				}
 
@@ -4384,7 +4400,7 @@ void FHeaderParser::GetVarType(
 		}
 	}
 
-	if ( VarProperty.IsObject() && VarProperty.Type != CPT_AssetObjectReference && VarProperty.MetaClass == nullptr && (VarProperty.PropertyFlags&CPF_Config) != 0 )
+	if ( VarProperty.IsObject() && VarProperty.Type != CPT_SoftObjectReference && VarProperty.MetaClass == nullptr && (VarProperty.PropertyFlags&CPF_Config) != 0 )
 	{
 		FError::Throwf(TEXT("Not allowed to use 'config' with object variables"));
 	}
@@ -5726,6 +5742,18 @@ bool FHeaderParser::IsValidDelegateDeclaration(const FToken& Token) const
 	return (Token.TokenType == TOKEN_Identifier) && TokenStr.StartsWith(TEXT("DECLARE_DYNAMIC_"));
 }
 
+// Modify token to fix redirected types if needed
+void FHeaderParser::RedirectTypeIdentifier(FToken& Token) const
+{
+	check(Token.TokenType == TOKEN_Identifier);
+
+	FString* FoundRedirect = TypeRedirectMap.Find(Token.Identifier);
+	if (FoundRedirect)
+	{
+		Token.SetIdentifier(**FoundRedirect);
+	}
+}
+
 // Parse the parameter list of a function or delegate declaration
 void FHeaderParser::ParseParameterList(FClasses& AllClasses, UFunction* Function, bool bExpectCommaBeforeName, TMap<FName, FString>* MetaData)
 {
@@ -6893,7 +6921,7 @@ void FHeaderParser::ValidatePropertyIsDeprecatedIfNecessary(FPropertyBase& VarPr
 
 	// check to see if we have a UObjectProperty using a deprecated class.
 	// PropertyClass is part of a union, so only check PropertyClass if this token represents an object property
-	if ( (VarProperty.Type == CPT_ObjectReference || VarProperty.Type == CPT_WeakObjectReference || VarProperty.Type == CPT_LazyObjectReference || VarProperty.Type == CPT_AssetObjectReference) && VarProperty.PropertyClass != NULL
+	if ( (VarProperty.Type == CPT_ObjectReference || VarProperty.Type == CPT_WeakObjectReference || VarProperty.Type == CPT_LazyObjectReference || VarProperty.Type == CPT_SoftObjectReference) && VarProperty.PropertyClass != NULL
 		&&	VarProperty.PropertyClass->HasAnyClassFlags(CLASS_Deprecated)	// and the object class being used has been deprecated
 		&& (VarProperty.PropertyFlags&CPF_Deprecated) == 0					// and this property isn't marked deprecated as well
 		&& (OuterPropertyType == NULL || !(OuterPropertyType->PropertyFlags & CPF_Deprecated)) ) // and this property isn't in an array that was marked deprecated either
@@ -7658,34 +7686,53 @@ FHeaderParser::FHeaderParser(FFeedbackContext* InWarn, const FManifestModule& In
 
 	FScriptLocation::Compiler = this;
 
-	// This should be moved to some sort of config
-	StructsWithNoPrefix.Add("uint64");
-	StructsWithNoPrefix.Add("uint32");
-	StructsWithNoPrefix.Add("double");
-
-	StructsWithTPrefix.Add("IndirectArray");
-	StructsWithTPrefix.Add("BitArray");
-	StructsWithTPrefix.Add("SparseArray");
-	StructsWithTPrefix.Add("Set");
-	StructsWithTPrefix.Add("Map");
-	StructsWithTPrefix.Add("MultiMap");
-	StructsWithTPrefix.Add("SharedPtr");
-
-	// List of legal delegate parameter counts
-	DelegateParameterCountStrings.Add(TEXT("_OneParam"));
-	DelegateParameterCountStrings.Add(TEXT("_TwoParams"));
-	DelegateParameterCountStrings.Add(TEXT("_ThreeParams"));
-	DelegateParameterCountStrings.Add(TEXT("_FourParams"));
-	DelegateParameterCountStrings.Add(TEXT("_FiveParams"));
-	DelegateParameterCountStrings.Add(TEXT("_SixParams"));
-	DelegateParameterCountStrings.Add(TEXT("_SevenParams"));
-	DelegateParameterCountStrings.Add(TEXT("_EightParams"));
-	DelegateParameterCountStrings.Add(TEXT("_NineParams"));
-
-	FString Version;
-	if (GConfig->GetString(TEXT("GeneratedCodeVersion"), TEXT("UnrealHeaderTool"), Version, GEngineIni))
-	{
-		DefaultGeneratedCodeVersion = ToGeneratedCodeVersion(Version);
+	static bool bConfigOptionsInitialized = false;
+
+	if (!bConfigOptionsInitialized)
+	{
+		// Read Ini options, GConfig must exist by this point
+		check(GConfig);
+
+		const FName TypeRedirectsKey(TEXT("TypeRedirects"));
+		const FName StructsWithNoPrefixKey(TEXT("StructsWithNoPrefix"));
+		const FName StructsWithTPrefixKey(TEXT("StructsWithTPrefix"));
+		const FName DelegateParameterCountStringsKey(TEXT("DelegateParameterCountStrings"));
+		const FName GeneratedCodeVersionKey(TEXT("GeneratedCodeVersion"));
+
+		FConfigSection* ConfigSection = GConfig->GetSectionPrivate(TEXT("UnrealHeaderTool"), false, true, GEngineIni);
+		if (ConfigSection)
+		{
+			for (FConfigSection::TIterator It(*ConfigSection); It; ++It)
+			{
+				if (It.Key() == TypeRedirectsKey)
+				{
+					FString OldType;
+					FString NewType;
+
+					FParse::Value(*It.Value().GetValue(), TEXT("OldType="), OldType);
+					FParse::Value(*It.Value().GetValue(), TEXT("NewType="), NewType);
+
+					TypeRedirectMap.Add(OldType, NewType);
+				}
+				else if (It.Key() == StructsWithNoPrefixKey)
+				{
+					StructsWithNoPrefix.Add(It.Value().GetValue());
+				}
+				else if (It.Key() == StructsWithTPrefixKey)
+				{
+					StructsWithTPrefix.Add(It.Value().GetValue());
+				}
+				else if (It.Key() == DelegateParameterCountStringsKey)
+				{
+					DelegateParameterCountStrings.Add(It.Value().GetValue());
+				}
+				else if (It.Key() == GeneratedCodeVersionKey)
+				{
+					DefaultGeneratedCodeVersion = ToGeneratedCodeVersion(It.Value().GetValue());
+				}
+			}
+		}
+		bConfigOptionsInitialized = true;
 	}
 }
 
diff --git a/Engine/Source/Programs/UnrealHeaderTool/Private/HeaderParser.h b/Engine/Source/Programs/UnrealHeaderTool/Private/HeaderParser.h
--- a/Engine/Source/Programs/UnrealHeaderTool/Private/HeaderParser.h
+++ b/Engine/Source/Programs/UnrealHeaderTool/Private/HeaderParser.h
@@ -389,21 +389,23 @@ protected:
 	////////////////////////////////////////////////////
 
 	// Special parsed struct names that do not require a prefix
-	TArray<FString> StructsWithNoPrefix;
+	static TArray<FString> StructsWithNoPrefix;
 	
 	// Special parsed struct names that have a 'T' prefix
-	TArray<FString> StructsWithTPrefix;
+	static TArray<FString> StructsWithTPrefix;
 
 	// Mapping from 'human-readable' macro substring to # of parameters for delegate declarations
 	// Index 0 is 1 parameter, Index 1 is 2, etc...
-	TArray<FString> DelegateParameterCountStrings;
+	static TArray<FString> DelegateParameterCountStrings;
+
+	// Types that have been renamed, treat the old deprecated name as the new name for code generation
+	static TMap<FString, FString> TypeRedirectMap;
 
 	// List of all used identifiers for net service function declarations (every function must be unique)
 	TMap<int32, FString> UsedRPCIds;
 	// List of all net service functions with undeclared response functions 
 	TMap<int32, FString> RPCsNeedingHookup;
 
-protected:
 	// Constructor.
 	explicit FHeaderParser(FFeedbackContext* InWarn, const FManifestModule& InModule);
 
@@ -424,6 +426,9 @@ protected:
 	// Parse the parameter list of a function or delegate declaration
 	void ParseParameterList(FClasses& AllClasses, UFunction* Function, bool bExpectCommaBeforeName = false, TMap<FName, FString>* MetaData = NULL);
 
+	// Modify token to fix redirected types if needed
+	void RedirectTypeIdentifier(FToken& Token) const;
+
 public:
 	// Throws if a specifier value wasn't provided
 	static void RequireSpecifierValue(const FPropertySpecifier& Specifier, bool bRequireExactlyOne = false);
diff --git a/Engine/Source/Programs/UnrealHeaderTool/Private/ParserHelper.cpp b/Engine/Source/Programs/UnrealHeaderTool/Private/ParserHelper.cpp
--- a/Engine/Source/Programs/UnrealHeaderTool/Private/ParserHelper.cpp
+++ b/Engine/Source/Programs/UnrealHeaderTool/Private/ParserHelper.cpp
@@ -122,7 +122,7 @@ const TCHAR* FPropertyBase::GetPropertyTypeText( EPropertyType Type )
 		CASE_TEXT(CPT_String);
 		CASE_TEXT(CPT_Text);
 		CASE_TEXT(CPT_MulticastDelegate);
-		CASE_TEXT(CPT_AssetObjectReference);
+		CASE_TEXT(CPT_SoftObjectReference);
 		CASE_TEXT(CPT_WeakObjectReference);
 		CASE_TEXT(CPT_LazyObjectReference);
 		CASE_TEXT(CPT_Map);
diff --git a/Engine/Source/Programs/UnrealHeaderTool/Private/ParserHelper.h b/Engine/Source/Programs/UnrealHeaderTool/Private/ParserHelper.h
--- a/Engine/Source/Programs/UnrealHeaderTool/Private/ParserHelper.h
+++ b/Engine/Source/Programs/UnrealHeaderTool/Private/ParserHelper.h
@@ -189,7 +189,7 @@ public:
 	{
 	}
 
-	explicit FPropertyBase(UClass* InClass, bool bIsWeak = false, bool bWeakIsAuto = false, bool bIsLazy = false, bool bIsAsset = false)
+	explicit FPropertyBase(UClass* InClass, bool bIsWeak = false, bool bWeakIsAuto = false, bool bIsLazy = false, bool bIsSoft = false)
 	: Type                (CPT_ObjectReference)
 	, ArrayType           (EArrayType::None)
 	, PropertyFlags       (0)
@@ -213,9 +213,9 @@ public:
 		{
 			Type = CPT_LazyObjectReference;
 		}
-		else if (bIsAsset)
+		else if (bIsSoft)
 		{
-			Type = CPT_AssetObjectReference;
+			Type = CPT_SoftObjectReference;
 		}
 		else if (bIsWeak)
 		{
@@ -384,16 +384,16 @@ public:
 			*this = FPropertyBase(CPT_LazyObjectReference);
 			PropertyClass = Cast<ULazyObjectProperty>(Property)->PropertyClass;
 		}
-		else if( ClassOfProperty==UAssetClassProperty::StaticClass() )
+		else if( ClassOfProperty==USoftClassProperty::StaticClass() )
 		{
-			*this = FPropertyBase(CPT_AssetObjectReference);
-			PropertyClass = Cast<UAssetClassProperty>(Property)->PropertyClass;
-			MetaClass = Cast<UAssetClassProperty>(Property)->MetaClass;
+			*this = FPropertyBase(CPT_SoftObjectReference);
+			PropertyClass = Cast<USoftClassProperty>(Property)->PropertyClass;
+			MetaClass = Cast<USoftClassProperty>(Property)->MetaClass;
 		}
-		else if( ClassOfProperty==UAssetObjectProperty::StaticClass() )
+		else if( ClassOfProperty==USoftObjectProperty::StaticClass() )
 		{
-			*this = FPropertyBase(CPT_AssetObjectReference);
-			PropertyClass = Cast<UAssetObjectProperty>(Property)->PropertyClass;
+			*this = FPropertyBase(CPT_SoftObjectReference);
+			PropertyClass = Cast<USoftObjectProperty>(Property)->PropertyClass;
 		}
 		else if( ClassOfProperty==UNameProperty::StaticClass() )
 		{
@@ -448,7 +448,7 @@ public:
 	 */
 	bool IsObject() const
 	{
-		return Type == CPT_ObjectReference || Type == CPT_Interface || Type == CPT_WeakObjectReference || Type == CPT_LazyObjectReference || Type == CPT_AssetObjectReference;
+		return Type == CPT_ObjectReference || Type == CPT_Interface || Type == CPT_WeakObjectReference || Type == CPT_LazyObjectReference || Type == CPT_SoftObjectReference;
 	}
 
 	bool IsContainer() const
@@ -518,7 +518,7 @@ public:
 				return false;
 			}
 		}
-		else if ((Type == CPT_ObjectReference || Type == CPT_WeakObjectReference || Type == CPT_LazyObjectReference || Type == CPT_AssetObjectReference) && Other.Type != CPT_Interface && (PropertyFlags & CPF_ReturnParm))
+		else if ((Type == CPT_ObjectReference || Type == CPT_WeakObjectReference || Type == CPT_LazyObjectReference || Type == CPT_SoftObjectReference) && Other.Type != CPT_Interface && (PropertyFlags & CPF_ReturnParm))
 		{
 			bReverseClassChainCheck = false;
 		}
@@ -830,6 +830,14 @@ public:
 
 	// Setters.
 	
+	void SetIdentifier( const TCHAR* InString)
+	{
+		InitToken(CPT_None);
+		TokenType = TOKEN_Identifier;
+		FCString::Strncpy(Identifier, InString, NAME_SIZE);
+		TokenName = FName(Identifier, FNAME_Find);
+	}
+
 	void SetConstInt64( int64 InInt64 )
 	{
 		(FPropertyBase&)*this = FPropertyBase(CPT_Int64);
@@ -866,7 +874,7 @@ public:
 		*(FName *)NameBytes = InName;
 		TokenType		= TOKEN_Const;
 	}
-	void SetConstString( TCHAR* InString, int32 MaxLength=MAX_STRING_CONST_SIZE )
+	void SetConstString( const TCHAR* InString, int32 MaxLength=MAX_STRING_CONST_SIZE )
 	{
 		check(MaxLength>0);
 		(FPropertyBase&)*this = FPropertyBase(CPT_String);
diff --git a/Engine/Source/Runtime/AIModule/Classes/AISystem.h b/Engine/Source/Runtime/AIModule/Classes/AISystem.h
--- a/Engine/Source/Runtime/AIModule/Classes/AISystem.h
+++ b/Engine/Source/Runtime/AIModule/Classes/AISystem.h
@@ -4,7 +4,7 @@
 
 #include "CoreMinimal.h"
 #include "UObject/ObjectMacros.h"
-#include "Misc/StringClassReference.h"
+#include "UObject/SoftObjectPath.h"
 #include "Engine/World.h"
 #include "AI/AISystemBase.h"
 #include "Math/RandomStream.h"
@@ -29,10 +29,10 @@ class AIMODULE_API UAISystem : public UAISystemBase
 
 protected:
 	UPROPERTY(globalconfig, EditAnywhere, Category = "AISystem", meta = (MetaClass = "AIPerceptionSystem", DisplayName = "Perception System Class"))
-	FStringClassReference PerceptionSystemClassName;
+	FSoftClassPath PerceptionSystemClassName;
 
 	UPROPERTY(globalconfig, EditAnywhere, Category = "AISystem", meta = (MetaClass = "AIHotSpotManager", DisplayName = "AIHotSpotManager Class"))
-	FStringClassReference HotSpotManagerClassName;
+	FSoftClassPath HotSpotManagerClassName;
 
 public:
 	/** Default AI movement's acceptance radius used to determine whether 
diff --git a/Engine/Source/Runtime/AIModule/Private/Perception/AIPerceptionSystem.cpp b/Engine/Source/Runtime/AIModule/Private/Perception/AIPerceptionSystem.cpp
--- a/Engine/Source/Runtime/AIModule/Private/Perception/AIPerceptionSystem.cpp
+++ b/Engine/Source/Runtime/AIModule/Private/Perception/AIPerceptionSystem.cpp
@@ -280,7 +280,7 @@ void UAIPerceptionSystem::AgeStimuli()
 UAIPerceptionSystem* UAIPerceptionSystem::GetCurrent(UObject* WorldContextObject)
 {
 	UWorld* World = Cast<UWorld>(WorldContextObject);
-	if (World == nullptr)
+	if (World == nullptr && WorldContextObject != nullptr)
 	{
 		World = GEngine->GetWorldFromContextObject(WorldContextObject, EGetWorldErrorMode::LogAndReturnNull);
 	}
diff --git a/Engine/Source/Runtime/AssetRegistry/Private/AssetBundleData.cpp b/Engine/Source/Runtime/AssetRegistry/Private/AssetBundleData.cpp
--- a/Engine/Source/Runtime/AssetRegistry/Private/AssetBundleData.cpp
+++ b/Engine/Source/Runtime/AssetRegistry/Private/AssetBundleData.cpp
@@ -9,7 +9,7 @@ bool FAssetBundleData::SetFromAssetData(const FAssetData& AssetData)
 	FString TagValue;
 
 	// Register that we're reading string assets for a specific package
-	FStringAssetReferenceSerializationScope SerializationScope(AssetData.PackageName, FAssetBundleData::StaticStruct()->GetFName(), EStringAssetReferenceCollectType::AlwaysCollect);
+	FSoftObjectPathSerializationScope SerializationScope(AssetData.PackageName, FAssetBundleData::StaticStruct()->GetFName(), ESoftObjectPathCollectType::AlwaysCollect);
 
 	if (AssetData.GetTagValue(FAssetBundleData::StaticStruct()->GetFName(), TagValue))
 	{
@@ -44,7 +44,7 @@ FAssetBundleEntry* FAssetBundleData::FindEntry(const FPrimaryAssetId& SearchScop
 	return nullptr;
 }
 
-void FAssetBundleData::AddBundleAsset(FName BundleName, const FStringAssetReference& AssetPath)
+void FAssetBundleData::AddBundleAsset(FName BundleName, const FSoftObjectPath& AssetPath)
 {
 	if (!AssetPath.IsValid())
 	{
@@ -61,11 +61,11 @@ void FAssetBundleData::AddBundleAsset(FName BundleName, const FStringAssetRefere
 	FoundEntry->BundleAssets.AddUnique(AssetPath);
 }
 
-void FAssetBundleData::AddBundleAssets(FName BundleName, const TArray<FStringAssetReference>& AssetPaths)
+void FAssetBundleData::AddBundleAssets(FName BundleName, const TArray<FSoftObjectPath>& AssetPaths)
 {
 	FAssetBundleEntry* FoundEntry = FindEntry(FPrimaryAssetId(), BundleName);
 
-	for (const FStringAssetReference& Path : AssetPaths)
+	for (const FSoftObjectPath& Path : AssetPaths)
 	{
 		if (Path.IsValid())
 		{
@@ -80,7 +80,7 @@ void FAssetBundleData::AddBundleAssets(FName BundleName, const TArray<FStringAss
 	}
 }
 
-void FAssetBundleData::SetBundleAssets(FName BundleName, TArray<FStringAssetReference>&& AssetPaths)
+void FAssetBundleData::SetBundleAssets(FName BundleName, TArray<FSoftObjectPath>&& AssetPaths)
 {
 	FAssetBundleEntry* FoundEntry = FindEntry(FPrimaryAssetId(), BundleName);
 
diff --git a/Engine/Source/Runtime/AssetRegistry/Private/AssetDataGatherer.cpp b/Engine/Source/Runtime/AssetRegistry/Private/AssetDataGatherer.cpp
--- a/Engine/Source/Runtime/AssetRegistry/Private/AssetDataGatherer.cpp
+++ b/Engine/Source/Runtime/AssetRegistry/Private/AssetDataGatherer.cpp
@@ -13,7 +13,7 @@
 
 namespace AssetDataGathererConstants
 {
-	static const int32 CacheSerializationVersion = 11;
+	static const int32 CacheSerializationVersion = 12;
 	static const int32 MaxFilesToDiscoverBeforeFlush = 2500;
 	static const int32 MaxFilesToGatherBeforeFlush = 250;
 	static const int32 MaxFilesToProcessBeforeCacheWrite = 50000;
diff --git a/Engine/Source/Runtime/AssetRegistry/Private/AssetRegistry.cpp b/Engine/Source/Runtime/AssetRegistry/Private/AssetRegistry.cpp
--- a/Engine/Source/Runtime/AssetRegistry/Private/AssetRegistry.cpp
+++ b/Engine/Source/Runtime/AssetRegistry/Private/AssetRegistry.cpp
@@ -1220,7 +1220,7 @@ void UAssetRegistryImpl::AssetDeleted(UObject* DeletedAsset)
 		if (bInitialSearchCompleted && AssetDataDeleted.IsRedirector())
 		{
 			// Need to remove from GRedirectCollector
-			GRedirectCollector.RemoveAssetPathRedirection(AssetDataDeleted.ObjectPath.ToString());
+			GRedirectCollector.RemoveAssetPathRedirection(AssetDataDeleted.ObjectPath);
 		}
 #endif
 
@@ -1753,24 +1753,15 @@ void UAssetRegistryImpl::DependencyDataGathered(const double TickStartTime, TBac
 			PackageDependencies.Add(AssetReference, EAssetRegistryDependencyType::Hard);
 		}
 
-		for (const FString& StringAssetReference : Result.StringAssetReferencesMap)
+		for (FName SoftPackageName : Result.SoftPackageReferenceList)
 		{
-			// Possibly resolve ini:name references before adding to dependency list
-			const FString* IniFilename = GetIniFilenameFromObjectsReference(StringAssetReference);
-			const FName AssetReference = IniFilename ? *ResolveIniObjectsReference(StringAssetReference, IniFilename) : *StringAssetReference;
-
 			// Already processed?
-			if (PackageDependencies.Contains(AssetReference))
+			if (PackageDependencies.Contains(SoftPackageName))
 			{
 				continue;
 			}
 
-			if (FPackageName::IsShortPackageName(AssetReference))
-			{
-				UE_LOG(LogAssetRegistry, Warning, TEXT("Package with string asset reference with short asset path: %s. This is unsupported, can couse errors and be slow on loading. Please resave the package to fix this."), *Result.PackageName.ToString());
-			}
-
-			PackageDependencies.Add(AssetReference, EAssetRegistryDependencyType::Soft);
+			PackageDependencies.Add(SoftPackageName, EAssetRegistryDependencyType::Soft);
 		}
 
 		for (const TPair<FPackageIndex, TArray<FName>>& SearchableNameList : Result.SearchableNamesMap)
@@ -2241,7 +2232,7 @@ void UAssetRegistryImpl::UpdateRedirectCollector()
 
 		if (Destination != AssetData->ObjectPath)
 		{
-			GRedirectCollector.AddAssetPathRedirection(AssetData->ObjectPath.ToString(), Destination.ToString());
+			GRedirectCollector.AddAssetPathRedirection(AssetData->ObjectPath, Destination);
 		}
 	}
 }
diff --git a/Engine/Source/Runtime/AssetRegistry/Private/AssetRegistryHelpers.cpp b/Engine/Source/Runtime/AssetRegistry/Private/AssetRegistryHelpers.cpp
--- a/Engine/Source/Runtime/AssetRegistry/Private/AssetRegistryHelpers.cpp
+++ b/Engine/Source/Runtime/AssetRegistry/Private/AssetRegistryHelpers.cpp
@@ -42,9 +42,9 @@ bool UAssetRegistryHelpers::IsRedirector(const FAssetData& InAssetData)
 	return InAssetData.IsRedirector();
 }
 
-FStringAssetReference UAssetRegistryHelpers::ToStringReference(const FAssetData& InAssetData) 
+FSoftObjectPath UAssetRegistryHelpers::ToSoftObjectPath(const FAssetData& InAssetData) 
 {
-	return InAssetData.ToStringReference();
+	return InAssetData.ToSoftObjectPath();
 }
 
 UClass* UAssetRegistryHelpers::GetClass(const FAssetData& InAssetData)
diff --git a/Engine/Source/Runtime/AssetRegistry/Private/PackageDependencyData.h b/Engine/Source/Runtime/AssetRegistry/Private/PackageDependencyData.h
--- a/Engine/Source/Runtime/AssetRegistry/Private/PackageDependencyData.h
+++ b/Engine/Source/Runtime/AssetRegistry/Private/PackageDependencyData.h
@@ -32,7 +32,7 @@ public:
 	{
 		Ar << PackageName;
 		Ar << ImportMap;
-		Ar << StringAssetReferencesMap;
+		Ar << SoftPackageReferenceList;
 		Ar << SearchableNamesMap;
 		
 		PackageData.SerializeForCache(Ar);
diff --git a/Engine/Source/Runtime/AssetRegistry/Private/PackageReader.cpp b/Engine/Source/Runtime/AssetRegistry/Private/PackageReader.cpp
--- a/Engine/Source/Runtime/AssetRegistry/Private/PackageReader.cpp
+++ b/Engine/Source/Runtime/AssetRegistry/Private/PackageReader.cpp
@@ -329,7 +329,7 @@ bool FPackageReader::ReadDependencyData(FPackageDependencyData& OutDependencyDat
 
 	SerializeNameMap();
 	SerializeImportMap(OutDependencyData.ImportMap);
-	SerializeStringAssetReferencesMap(OutDependencyData.StringAssetReferencesMap);
+	SerializeSoftPackageReferenceList(OutDependencyData.SoftPackageReferenceList);
 	SerializeSearchableNamesMap(OutDependencyData);
 
 	return true;
@@ -378,44 +378,38 @@ void FPackageReader::SerializeExportMap(TArray<FObjectExport>& OutExportMap)
 	}
 }
 
-void FPackageReader::SerializeStringAssetReferencesMap(TArray<FString>& OutStringAssetReferencesMap)
+void FPackageReader::SerializeSoftPackageReferenceList(TArray<FName>& OutSoftPackageReferenceList)
 {
-	if (UE4Ver() >= VER_UE4_ADD_STRING_ASSET_REFERENCES_MAP && PackageFileSummary.StringAssetReferencesOffset > 0 && PackageFileSummary.StringAssetReferencesCount > 0)
+	if (UE4Ver() >= VER_UE4_ADD_STRING_ASSET_REFERENCES_MAP && PackageFileSummary.SoftPackageReferencesOffset > 0 && PackageFileSummary.SoftPackageReferencesCount > 0)
 	{
-		Seek(PackageFileSummary.StringAssetReferencesOffset);
+		Seek(PackageFileSummary.SoftPackageReferencesOffset);
 
-		if (UE4Ver() < VER_UE4_KEEP_ONLY_PACKAGE_NAMES_IN_STRING_ASSET_REFERENCES_MAP)
+		if (UE4Ver() < VER_UE4_ADDED_SOFT_OBJECT_PATH)
 		{
-			for (int32 ReferenceIdx = 0; ReferenceIdx < PackageFileSummary.StringAssetReferencesCount; ++ReferenceIdx)
+			for (int32 ReferenceIdx = 0; ReferenceIdx < PackageFileSummary.SoftPackageReferencesCount; ++ReferenceIdx)
 			{
-				FString Buf;
-				*this << Buf;
+				FString PackageName;
+				*this << PackageName;
 
-				if (GetIniFilenameFromObjectsReference(Buf) != nullptr)
+				if (UE4Ver() < VER_UE4_KEEP_ONLY_PACKAGE_NAMES_IN_STRING_ASSET_REFERENCES_MAP)
 				{
-					OutStringAssetReferencesMap.AddUnique(MoveTemp(Buf));
-				}
-				else
-				{
-					FString NormalizedPath = FPackageName::GetNormalizedObjectPath(MoveTemp(Buf));
-					if (!NormalizedPath.IsEmpty())
+					PackageName = FPackageName::GetNormalizedObjectPath(PackageName);
+					if (!PackageName.IsEmpty())
 					{
-						OutStringAssetReferencesMap.AddUnique(
-							FPackageName::ObjectPathToPackageName(
-								NormalizedPath
-							)
-						);
+						PackageName = FPackageName::ObjectPathToPackageName(PackageName);
 					}
 				}
+
+				OutSoftPackageReferenceList.Add(FName(*PackageName));
 			}
 		}
 		else
 		{
-			for (int32 ReferenceIdx = 0; ReferenceIdx < PackageFileSummary.StringAssetReferencesCount; ++ReferenceIdx)
+			for (int32 ReferenceIdx = 0; ReferenceIdx < PackageFileSummary.SoftPackageReferencesCount; ++ReferenceIdx)
 			{
-				FString Buf;
-				*this << Buf;
-				OutStringAssetReferencesMap.Add(MoveTemp(Buf));
+				FName PackageName;
+				*this << PackageName;
+				OutSoftPackageReferenceList.Add(PackageName);
 			}
 		}
 	}
diff --git a/Engine/Source/Runtime/AssetRegistry/Private/PackageReader.h b/Engine/Source/Runtime/AssetRegistry/Private/PackageReader.h
--- a/Engine/Source/Runtime/AssetRegistry/Private/PackageReader.h
+++ b/Engine/Source/Runtime/AssetRegistry/Private/PackageReader.h
@@ -47,7 +47,7 @@ public:
 	void SerializeNameMap();
 	void SerializeImportMap(TArray<FObjectImport>& OutImportMap);
 	void SerializeExportMap(TArray<FObjectExport>& OutExportMap);
-	void SerializeStringAssetReferencesMap(TArray<FString>& OutStringAssetReferencesMap);
+	void SerializeSoftPackageReferenceList(TArray<FName>& OutSoftPackageReferenceList);
 	void SerializeSearchableNamesMap(FPackageDependencyData& OutDependencyData);
 
 	/** Returns flags the asset package was saved with */
diff --git a/Engine/Source/Runtime/AssetRegistry/Public/AssetBundleData.h b/Engine/Source/Runtime/AssetRegistry/Public/AssetBundleData.h
--- a/Engine/Source/Runtime/AssetRegistry/Public/AssetBundleData.h
+++ b/Engine/Source/Runtime/AssetRegistry/Public/AssetBundleData.h
@@ -3,8 +3,8 @@
 #pragma once
 
 #include "CoreMinimal.h"
-#include "PrimaryAssetId.h"
-#include "AssetPtr.h"
+#include "UObject/PrimaryAssetId.h"
+#include "UObject/SoftObjectPtr.h"
 #include "AssetBundleData.generated.h"
 
 /** A struct representing a single AssetBundle */
@@ -27,7 +27,7 @@ struct ASSETREGISTRY_API FAssetBundleEntry
 
 	/** List of string assets contained in this bundle */
 	UPROPERTY()
-	TArray<FStringAssetReference> BundleAssets;
+	TArray<FSoftObjectPath> BundleAssets;
 
 	FAssetBundleEntry(const FAssetBundleEntry& OldEntry)
 		: BundleScope(OldEntry.BundleScope), BundleName(OldEntry.BundleName), BundleAssets(OldEntry.BundleAssets)
@@ -90,19 +90,19 @@ struct ASSETREGISTRY_API FAssetBundleData
 	FAssetBundleEntry* FindEntry(const FPrimaryAssetId& SearchScope, FName SearchName);
 
 	/** Adds or updates an entry with the given BundleName -> Path. Scope is empty and will be filled in later */
-	void AddBundleAsset(FName BundleName, const FStringAssetReference& AssetPath);
+	void AddBundleAsset(FName BundleName, const FSoftObjectPath& AssetPath);
 
 	template< typename T >
-	void AddBundleAsset(FName BundleName, const TAssetPtr<T>& AssetPtr)
+	void AddBundleAsset(FName BundleName, const TSoftObjectPtr<T>& SoftObjectPtr)
 	{
-		AddBundleAsset(BundleName, AssetPtr.ToStringReference());
+		AddBundleAsset(BundleName, SoftObjectPtr.ToSoftObjectPath());
 	}
 
 	/** Adds multiple assets at once */
-	void AddBundleAssets(FName BundleName, const TArray<FStringAssetReference>& AssetPaths);
+	void AddBundleAssets(FName BundleName, const TArray<FSoftObjectPath>& AssetPaths);
 
 	/** A fast set of asset bundle assets, will destroy copied in path list */
-	void SetBundleAssets(FName BundleName, TArray<FStringAssetReference>&& AssetPaths);
+	void SetBundleAssets(FName BundleName, TArray<FSoftObjectPath>&& AssetPaths);
 
 	/** Resets the data to defaults */
 	void Reset();
diff --git a/Engine/Source/Runtime/AssetRegistry/Public/AssetData.h b/Engine/Source/Runtime/AssetRegistry/Public/AssetData.h
--- a/Engine/Source/Runtime/AssetRegistry/Public/AssetData.h
+++ b/Engine/Source/Runtime/AssetRegistry/Public/AssetData.h
@@ -6,7 +6,7 @@
 #include "UObject/ObjectMacros.h"
 #include "UObject/Object.h"
 #include "UObject/Class.h"
-#include "Misc/StringAssetReference.h"
+#include "UObject/SoftObjectPath.h"
 #include "UObject/Package.h"
 #include "UObject/ObjectRedirector.h"
 #include "Misc/PackageName.h"
@@ -111,6 +111,7 @@ public:
 			}
 
 			const UPackage* Outermost = InAsset->GetOutermost();
+			const UObject* Outer = InAsset->GetOuter();
 
 			PackageName = Outermost->GetFName();
 			PackagePath = FName(*FPackageName::GetLongPackagePath(Outermost->GetName()));
@@ -240,10 +241,16 @@ public:
 		return FoundClass;
 	}
 
-	/** Convert to a StringAssetReference for loading */
-	FStringAssetReference ToStringReference() const
+	/** Convert to a SoftObjectPath for loading */
+	FSoftObjectPath ToSoftObjectPath() const
 	{
-		return FStringAssetReference(ObjectPath.ToString());
+		return FSoftObjectPath(ObjectPath.ToString());
+	}
+
+	DEPRECATED(4.18, "ToStringReference was renamed to ToSoftObjectPath")
+	FSoftObjectPath ToStringReference() const
+	{
+		return ToSoftObjectPath();
 	}
 	
 	/** Gets primary asset id of this data */
diff --git a/Engine/Source/Runtime/AssetRegistry/Public/AssetRegistryHelpers.h b/Engine/Source/Runtime/AssetRegistry/Public/AssetRegistryHelpers.h
--- a/Engine/Source/Runtime/AssetRegistry/Public/AssetRegistryHelpers.h
+++ b/Engine/Source/Runtime/AssetRegistry/Public/AssetRegistryHelpers.h
@@ -54,9 +54,9 @@ public:
 	UFUNCTION(BlueprintPure, Category = "Asset Data")
 	static FString GetFullName(const FAssetData& InAssetData);
 
-	/** Convert to a StringAssetReference for loading */
+	/** Convert to a SoftObjectPath for loading */
 	UFUNCTION(BlueprintPure, Category = "Asset Data")
-	static FStringAssetReference ToStringReference(const FAssetData& InAssetData);
+	static FSoftObjectPath ToSoftObjectPath(const FAssetData& InAssetData);
 
 	UFUNCTION(BlueprintPure, Category = "Asset Data")
 	static UClass* GetClass(const FAssetData& InAssetData);
diff --git a/Engine/Source/Runtime/AutomationMessages/Classes/AutomationWorkerMessages.h b/Engine/Source/Runtime/AutomationMessages/Classes/AutomationWorkerMessages.h
--- a/Engine/Source/Runtime/AutomationMessages/Classes/AutomationWorkerMessages.h
+++ b/Engine/Source/Runtime/AutomationMessages/Classes/AutomationWorkerMessages.h
@@ -167,10 +167,10 @@ struct FAutomationWorkerRequestTests
 
 
 /**
- * Implements a message that is sent in response to FAutomationWorkerRequestTests.
+ * A single test reply, used by FAutomationWorkerRequestTestsReplyComplete
  */
 USTRUCT()
-struct FAutomationWorkerRequestTestsReply
+struct FAutomationWorkerSingleTestReply
 {
 	GENERATED_USTRUCT_BODY()
 
@@ -204,16 +204,11 @@ struct FAutomationWorkerRequestTestsReply
 	UPROPERTY(EditAnywhere, Category="Message")
 	uint32 NumParticipantsRequired;
 
-	/** Holds the total number of tests returned. */
-	UPROPERTY(EditAnywhere, Category="Message")
-	int32 TotalNumTests;
-
 	/** Default constructor. */
-	FAutomationWorkerRequestTestsReply() { }
+	FAutomationWorkerSingleTestReply() { }
 
 	/** Creates and initializes a new instance. */
-	FAutomationWorkerRequestTestsReply(const FAutomationTestInfo& InTestInfo, const int32& InTotalNumTests)
-		: TotalNumTests(InTotalNumTests)
+	FAutomationWorkerSingleTestReply(const FAutomationTestInfo& InTestInfo)
 	{
 		DisplayName = InTestInfo.GetDisplayName();
 		FullTestPath = InTestInfo.GetFullTestPath();
@@ -245,11 +240,15 @@ struct FAutomationWorkerRequestTestsReply
 
 
 /**
-*/
+ * Returns list of all tests
+ */
 USTRUCT()
 struct FAutomationWorkerRequestTestsReplyComplete
 {
 	GENERATED_USTRUCT_BODY()
+
+	UPROPERTY(EditAnywhere, Category = "Message")
+	TArray<FAutomationWorkerSingleTestReply> Tests;
 };
 
 
diff --git a/Engine/Source/Runtime/AutomationWorker/Private/AutomationWorkerModule.cpp b/Engine/Source/Runtime/AutomationWorker/Private/AutomationWorkerModule.cpp
--- a/Engine/Source/Runtime/AutomationWorker/Private/AutomationWorkerModule.cpp
+++ b/Engine/Source/Runtime/AutomationWorker/Private/AutomationWorkerModule.cpp
@@ -259,11 +259,13 @@ void FAutomationWorkerModule::ReportTestComplete()
 
 void FAutomationWorkerModule::SendTests( const FMessageAddress& ControllerAddress )
 {
+	FAutomationWorkerRequestTestsReplyComplete* Reply = new FAutomationWorkerRequestTestsReplyComplete();
 	for( int32 TestIndex = 0; TestIndex < TestInfo.Num(); TestIndex++ )
 	{
-		MessageEndpoint->Send(new FAutomationWorkerRequestTestsReply(TestInfo[TestIndex], TestInfo.Num()), ControllerAddress);
+		Reply->Tests.Emplace(FAutomationWorkerSingleTestReply(TestInfo[TestIndex]));
 	}
-	MessageEndpoint->Send(new FAutomationWorkerRequestTestsReplyComplete(), ControllerAddress);
+
+	MessageEndpoint->Send(Reply, ControllerAddress);
 }
 
 
@@ -364,11 +366,6 @@ void FAutomationWorkerModule::HandleRequestTestsMessage( const FAutomationWorker
 void FAutomationWorkerModule::HandlePreTestingEvent()
 {
 #if WITH_ENGINE
-	if (!GIsEditor && GEngine->GameViewport)
-	{
-		GEngine->GameViewport->OnScreenshotCaptured().AddRaw(this, &FAutomationWorkerModule::HandleScreenShotCaptured);
-	}
-	//Register the editor screen shot callback
 	FAutomationTestFramework::Get().OnScreenshotCaptured().BindRaw(this, &FAutomationWorkerModule::HandleScreenShotCapturedWithName);
 #endif
 }
@@ -377,12 +374,7 @@ void FAutomationWorkerModule::HandlePreTestingEvent()
 void FAutomationWorkerModule::HandlePostTestingEvent()
 {
 #if WITH_ENGINE
-	if (!GIsEditor && GEngine->GameViewport)
-	{
-		GEngine->GameViewport->OnScreenshotCaptured().RemoveAll(this);
-	}
-	//Register the editor screen shot callback
-	FAutomationTestFramework::Get().OnScreenshotCaptured().BindRaw(this, &FAutomationWorkerModule::HandleScreenShotCapturedWithName);
+	FAutomationTestFramework::Get().OnScreenshotCaptured().Unbind();
 #endif
 }
 
@@ -404,16 +396,6 @@ void FAutomationWorkerModule::HandlePerformanceDataRetrieved(const FAutomationWo
 }
 
 #if WITH_ENGINE
-void FAutomationWorkerModule::HandleScreenShotCaptured(int32 Width, int32 Height, const TArray<FColor>& Bitmap)
-{
-	FAutomationScreenshotData Data;
-	Data.Width = Width;
-	Data.Height = Height;
-	Data.Path = FScreenshotRequest::GetFilename();
-
-	HandleScreenShotCapturedWithName(Bitmap, Data);
-}
-
 void FAutomationWorkerModule::HandleScreenShotCapturedWithName(const TArray<FColor>& RawImageData, const FAutomationScreenshotData& Data)
 {
 	int32 NewHeight = Data.Height;
diff --git a/Engine/Source/Runtime/AutomationWorker/Private/AutomationWorkerModule.h b/Engine/Source/Runtime/AutomationWorker/Private/AutomationWorkerModule.h
--- a/Engine/Source/Runtime/AutomationWorker/Private/AutomationWorkerModule.h
+++ b/Engine/Source/Runtime/AutomationWorker/Private/AutomationWorkerModule.h
@@ -120,7 +120,6 @@ private:
 
 #if WITH_ENGINE
 	/** Invoked when we have screen shot to send. */
-	void HandleScreenShotCaptured(int32 Width, int32 Height, const TArray<FColor>& Bitmap);
 	void HandleScreenShotCapturedWithName(const TArray<FColor>& RawImageData, const FAutomationScreenshotData& Data);
 #endif
 
diff --git a/Engine/Source/Runtime/Core/Private/GenericPlatform/GenericPlatformMath.cpp b/Engine/Source/Runtime/Core/Private/GenericPlatform/GenericPlatformMath.cpp
--- a/Engine/Source/Runtime/Core/Private/GenericPlatform/GenericPlatformMath.cpp
+++ b/Engine/Source/Runtime/Core/Private/GenericPlatform/GenericPlatformMath.cpp
@@ -85,6 +85,11 @@ void FGenericPlatformMath::FmodReportError(float X, float Y)
 #if WITH_DEV_AUTOMATION_TESTS
 extern float TheCompilerDoesntKnowThisIsAlwaysZero;
 
+// Disabled for UE-48236
+#if PLATFORM_PS4
+PRAGMA_DISABLE_OPTIMIZATION
+#endif
+
 void FGenericPlatformMath::AutoTest() 
 {
 	check(IsNaN(sqrtf(-1.0f)));
@@ -146,4 +151,9 @@ void FGenericPlatformMath::AutoTest()
 	}
 }
 
+// Disabled for UE-48236
+#if PLATFORM_PS4
+PRAGMA_ENABLE_OPTIMIZATION
+#endif
+
 #endif //WITH_DEV_AUTOMATION_TESTS
diff --git a/Engine/Source/Runtime/Core/Private/Misc/CoreDelegates.cpp b/Engine/Source/Runtime/Core/Private/Misc/CoreDelegates.cpp
--- a/Engine/Source/Runtime/Core/Private/Misc/CoreDelegates.cpp
+++ b/Engine/Source/Runtime/Core/Private/Misc/CoreDelegates.cpp
@@ -94,8 +94,6 @@ FCoreDelegates::FStatDisableAll FCoreDelegates::StatDisableAll;
 FCoreDelegates::FApplicationLicenseChange FCoreDelegates::ApplicationLicenseChange;
 FCoreDelegates::FPlatformChangedLaptopMode FCoreDelegates::PlatformChangedLaptopMode;
 
-FCoreDelegates::FLoadStringAssetReferenceInCook FCoreDelegates::LoadStringAssetReferenceInCook;
-
 FCoreDelegates::FVRHeadsetRecenter FCoreDelegates::VRHeadsetTrackingInitializingAndNeedsHMDToBeTrackedDelegate;
 FCoreDelegates::FVRHeadsetRecenter FCoreDelegates::VRHeadsetTrackingInitializedDelegate;
 FCoreDelegates::FVRHeadsetRecenter FCoreDelegates::VRHeadsetRecenter;
diff --git a/Engine/Source/Runtime/Core/Private/Serialization/Archive.cpp b/Engine/Source/Runtime/Core/Private/Serialization/Archive.cpp
--- a/Engine/Source/Runtime/Core/Private/Serialization/Archive.cpp
+++ b/Engine/Source/Runtime/Core/Private/Serialization/Archive.cpp
@@ -269,30 +269,30 @@ FArchive& FArchive::operator<<( FText& Value )
 	return *this;
 }
 
-FArchive& FArchive::operator<<( class FLazyObjectPtr& LazyObjectPtr )
+FArchive& FArchive::operator<<(struct FLazyObjectPtr& Value)
 {
-	// The base FArchive does not implement this method. Use FArchiveUOBject instead.
+	// The base FArchive does not implement this method. Use FArchiveUObject instead.
 	UE_LOG(LogSerialization, Fatal, TEXT("FArchive does not support FLazyObjectPtr serialization. Use FArchiveUObject instead."));
 	return *this;
 }
 
-FArchive& FArchive::operator<<( class FAssetPtr& AssetPtr )
+FArchive& FArchive::operator<<(struct FSoftObjectPtr& Value)
 {
-	// The base FArchive does not implement this method. Use FArchiveUOBject instead.
-	UE_LOG(LogSerialization, Fatal, TEXT("FArchive does not support FAssetPtr serialization. Use FArchiveUObject instead."));
+	// The base FArchive does not implement this method. Use FArchiveUObject instead.
+	UE_LOG(LogSerialization, Fatal, TEXT("FArchive does not support FSoftObjectPtr serialization. Use FArchiveUObject instead."));
 	return *this;
 }
 
-FArchive& FArchive::operator<<(struct FStringAssetReference& Value)
+FArchive& FArchive::operator<<(struct FSoftObjectPath& Value)
 {
-	// The base FArchive does not implement this method. Use FArchiveUOBject instead.
-	UE_LOG(LogSerialization, Fatal, TEXT("FArchive does not support FAssetPtr serialization. Use FArchiveUObject instead."));
+	// The base FArchive does not implement this method. Use FArchiveUObject instead.
+	UE_LOG(LogSerialization, Fatal, TEXT("FArchive does not support FSoftObjectPath serialization. Use FArchiveUObject instead."));
 	return *this;
 }
 
 FArchive& FArchive::operator<<(struct FWeakObjectPtr& Value)
 {
-	// The base FArchive does not implement this method. Use FArchiveUOBject instead.
+	// The base FArchive does not implement this method. Use FArchiveUObject instead.
 	UE_LOG(LogSerialization, Fatal, TEXT("FArchive does not support FWeakObjectPtr serialization. Use FArchiveUObject instead."));
 	return *this;
 }
diff --git a/Engine/Source/Runtime/Core/Private/Tests/HAL/PlatformTest.cpp b/Engine/Source/Runtime/Core/Private/Tests/HAL/PlatformTest.cpp
--- a/Engine/Source/Runtime/Core/Private/Tests/HAL/PlatformTest.cpp
+++ b/Engine/Source/Runtime/Core/Private/Tests/HAL/PlatformTest.cpp
@@ -40,9 +40,13 @@ struct TestC : public TestA, TestB
 	int i;
 };
 
-
 IMPLEMENT_SIMPLE_AUTOMATION_TEST(FPlatformVerificationTest, "System.Core.HAL.Platform Verification", EAutomationTestFlags::EditorContext | EAutomationTestFlags::ClientContext | EAutomationTestFlags::SmokeFilter)
 
+// Disabled for UE-48236
+#if PLATFORM_PS4
+PRAGMA_DISABLE_OPTIMIZATION
+#endif
+
 bool FPlatformVerificationTest::RunTest (const FString& Parameters)
 {
 	PTRINT Offset1 = VTABLE_OFFSET(TestC, TestB);
@@ -108,4 +112,9 @@ bool FPlatformVerificationTest::RunTest (const FString& Parameters)
 	return true;
 }
 
+// Disabled for UE-48236
+#if PLATFORM_PS4
+PRAGMA_ENABLE_OPTIMIZATION
+#endif
+
 #endif //WITH_DEV_AUTOMATION_TESTS
diff --git a/Engine/Source/Runtime/Core/Private/Tests/Math/UnrealMathTest.cpp b/Engine/Source/Runtime/Core/Private/Tests/Math/UnrealMathTest.cpp
--- a/Engine/Source/Runtime/Core/Private/Tests/Math/UnrealMathTest.cpp
+++ b/Engine/Source/Runtime/Core/Private/Tests/Math/UnrealMathTest.cpp
@@ -683,6 +683,10 @@ void SetScratch( float X, float Y, float Z, float W, float U = 0.0f )
 	GScratch[4] = U;
 }
 
+// Disabled for UE-48236
+#if PLATFORM_PS4
+PRAGMA_DISABLE_OPTIMIZATION
+#endif
 
 IMPLEMENT_SIMPLE_AUTOMATION_TEST(FVectorRegisterAbstractionTest, "System.Core.Math.Vector Register Abstraction Test", EAutomationTestFlags::ApplicationContextMask | EAutomationTestFlags::SmokeFilter)
 
@@ -1114,7 +1118,6 @@ bool FVectorRegisterAbstractionTest::RunTest(const FString& Parameters)
 	V2 = TestVectorTransformVector(V0, &M1);
 	LogTest( TEXT("VectorTransformVector"), TestVectorsEqual( V1, V2 ) );
 
-
 	// NaN / Inf tests
 	const float NaN = sqrtf(-1.0f);
 	LogTest(TEXT("VectorContainsNaNOrInfinite true"), VectorContainsNaNOrInfinite(MakeVectorRegister(NaN, NaN, NaN, NaN)));
@@ -1485,6 +1488,12 @@ bool FVectorRegisterAbstractionTest::RunTest(const FString& Parameters)
 	return true;
 }
 
+// Disabled for UE-48236
+#if PLATFORM_PS4
+PRAGMA_ENABLE_OPTIMIZATION
+#endif
+
+
 IMPLEMENT_SIMPLE_AUTOMATION_TEST(FInterpolationFunctionTests, "System.Core.Math.Interpolation Function Test", EAutomationTestFlags::ApplicationContextMask | EAutomationTestFlags::SmokeFilter)
 bool FInterpolationFunctionTests::RunTest(const FString&)
 {
diff --git a/Engine/Source/Runtime/Core/Private/Tests/Misc/PathsTest.cpp b/Engine/Source/Runtime/Core/Private/Tests/Misc/PathsTest.cpp
--- a/Engine/Source/Runtime/Core/Private/Tests/Misc/PathsTest.cpp
+++ b/Engine/Source/Runtime/Core/Private/Tests/Misc/PathsTest.cpp
@@ -10,6 +10,9 @@
 
 IMPLEMENT_SIMPLE_AUTOMATION_TEST(FPathTests, "System.Core.Misc.Paths", EAutomationTestFlags::ApplicationContextMask | EAutomationTestFlags::SmokeFilter)
 
+// MSVC 2015 generates bad code here due to the number of string literals, in a cooked build it was crashing on RunCollapseRelativeDirectoriesTest(TEXT("C:/."), TEXT("C:/."));
+PRAGMA_DISABLE_OPTIMIZATION
+
 bool FPathTests::RunTest( const FString& Parameters )
 {
 	// Directory collapsing
@@ -157,4 +160,6 @@ bool FPathTests::RunTest( const FString& Parameters )
 	return true;
 }
 
+PRAGMA_ENABLE_OPTIMIZATION
+
 #endif //WITH_DEV_AUTOMATION_TESTS
diff --git a/Engine/Source/Runtime/Core/Public/Misc/CoreDelegates.h b/Engine/Source/Runtime/Core/Public/Misc/CoreDelegates.h
--- a/Engine/Source/Runtime/Core/Public/Misc/CoreDelegates.h
+++ b/Engine/Source/Runtime/Core/Public/Misc/CoreDelegates.h
@@ -267,12 +267,6 @@ public:
 	DECLARE_MULTICAST_DELEGATE_OneParam(FPlatformChangedLaptopMode, EConvertibleLaptopMode);
 	static FPlatformChangedLaptopMode PlatformChangedLaptopMode;
 
-	DECLARE_DELEGATE_RetVal_OneParam(bool, FLoadStringAssetReferenceInCook, const FString&);
-	static FLoadStringAssetReferenceInCook LoadStringAssetReferenceInCook;
-
-	DECLARE_DELEGATE_RetVal_OneParam(bool, FStringAssetReferenceLoaded, const FName&);
-	static FStringAssetReferenceLoaded StringAssetReferenceLoaded;
-
 	/** Sent when the platform needs the user to fix headset tracking on startup (PS4 Morpheus only) */
 	DECLARE_MULTICAST_DELEGATE(FVRHeadsetTrackingInitializingAndNeedsHMDToBeTrackedDelegate);
 	static FVRHeadsetTrackingInitializingAndNeedsHMDToBeTrackedDelegate VRHeadsetTrackingInitializingAndNeedsHMDToBeTrackedDelegate;
diff --git a/Engine/Source/Runtime/Core/Public/Misc/CoreMiscDefines.h b/Engine/Source/Runtime/Core/Public/Misc/CoreMiscDefines.h
--- a/Engine/Source/Runtime/Core/Public/Misc/CoreMiscDefines.h
+++ b/Engine/Source/Runtime/Core/Public/Misc/CoreMiscDefines.h
@@ -14,7 +14,6 @@
 
 #define LOCALIZED_SEEKFREE_SUFFIX	TEXT("_LOC")
 #define PLAYWORLD_PACKAGE_PREFIX TEXT("UEDPIE")
-#define PLAYWORLD_CONSOLE_BASE_PACKAGE_PREFIX TEXT("UED")
 
 #ifndef WITH_EDITORONLY_DATA
 	#if !PLATFORM_CAN_SUPPORT_EDITORONLY_DATA || UE_SERVER || PLATFORM_IOS
diff --git a/Engine/Source/Runtime/Core/Public/Serialization/Archive.h b/Engine/Source/Runtime/Core/Public/Serialization/Archive.h
--- a/Engine/Source/Runtime/Core/Public/Serialization/Archive.h
+++ b/Engine/Source/Runtime/Core/Public/Serialization/Archive.h
@@ -226,25 +226,25 @@ public:
 	 * @param Value The value to serialize.
 	 * @return This instance.
 	 */
-	virtual FArchive& operator<<(class FLazyObjectPtr& Value);
+	virtual FArchive& operator<<(struct FLazyObjectPtr& Value);
 	
 	/**
 	 * Serializes asset pointer from or into this archive.
 	 *
-	 * Most of the time, FAssetPtrs are serialized as UObject *, but some archives need to override this.
+	 * Most of the time, FSoftObjectPtr are serialized as UObject *, but some archives need to override this.
 	 *
 	 * @param Value The asset pointer to serialize.
 	 * @return This instance.
 	 */
-	virtual FArchive& operator<<(class FAssetPtr& Value);
+	virtual FArchive& operator<<(struct FSoftObjectPtr& Value);
 
 	/**
-	 * Serializes string asset reference from or into this archive.
+	 * Serializes soft object paths from or into this archive.
 	 *
 	 * @param Value String asset reference to serialize.
 	 * @return This instance.
 	 */
-	virtual FArchive& operator<<(struct FStringAssetReference& Value);
+	virtual FArchive& operator<<(struct FSoftObjectPath& Value);
 
 	/**
 	* Serializes FWeakObjectPtr value from or into this archive.
diff --git a/Engine/Source/Runtime/Core/Public/Serialization/ArchiveProxy.h b/Engine/Source/Runtime/Core/Public/Serialization/ArchiveProxy.h
--- a/Engine/Source/Runtime/Core/Public/Serialization/ArchiveProxy.h
+++ b/Engine/Source/Runtime/Core/Public/Serialization/ArchiveProxy.h
@@ -42,7 +42,7 @@ public:
 		return *this;
 	}
 
-	virtual FArchive& operator<<(struct FStringAssetReference& Value) override
+	virtual FArchive& operator<<(struct FSoftObjectPath& Value) override
 	{
 		InnerArchive << Value;
 		return *this;
diff --git a/Engine/Source/Runtime/Core/Public/UObject/FrameworkObjectVersion.h b/Engine/Source/Runtime/Core/Public/UObject/FrameworkObjectVersion.h
--- a/Engine/Source/Runtime/Core/Public/UObject/FrameworkObjectVersion.h
+++ b/Engine/Source/Runtime/Core/Public/UObject/FrameworkObjectVersion.h
@@ -100,6 +100,10 @@ struct CORE_API FFrameworkObjectVersion
 		// Fix Local Variables so that the properties are correctly flagged as blueprint visible
 		LocalVariablesBlueprintVisible,
 
+		// Stopped serializing UField_Next so that UFunctions could be serialized in dependently of a UClass
+		// in order to allow us to do all UFunction loading in a single pass (after classes and CDOs are created):
+		RemoveUField_Next,
+
 		// -----<new versions can be added above this line>-------------------------------------------------
 		VersionPlusOne,
 		LatestVersion = VersionPlusOne - 1
diff --git a/Engine/Source/Runtime/Core/Public/UObject/ObjectVersion.h b/Engine/Source/Runtime/Core/Public/UObject/ObjectVersion.h
--- a/Engine/Source/Runtime/Core/Public/UObject/ObjectVersion.h
+++ b/Engine/Source/Runtime/Core/Public/UObject/ObjectVersion.h
@@ -623,6 +623,8 @@ enum EUnrealEngineObjectUE4Version
 	VER_UE4_SKYLIGHT_MOBILE_IRRADIANCE_MAP,
 	// Added flag to control sweep behavior while walking in UCharacterMovementComponent.
 	VER_UE4_ADDED_SWEEP_WHILE_WALKING_FLAG,
+	// StringAssetReference changed to SoftObjectPath and swapped to serialize as a name+string instead of a string
+	VER_UE4_ADDED_SOFT_OBJECT_PATH,
 
 	// -----<new versions can be added before this line>-------------------------------------------------
 	// - this needs to be the last line (see note below)
diff --git a/Engine/Source/Runtime/Core/Public/UObject/UObjectHierarchyFwd.h b/Engine/Source/Runtime/Core/Public/UObject/UObjectHierarchyFwd.h
--- a/Engine/Source/Runtime/Core/Public/UObject/UObjectHierarchyFwd.h
+++ b/Engine/Source/Runtime/Core/Public/UObject/UObjectHierarchyFwd.h
@@ -25,7 +25,8 @@ class					UClassProperty;
 class					UInterfaceProperty;
 class					UWeakObjectProperty;
 class					ULazyObjectProperty;
-class					UAssetObjectProperty;
+class					USoftObjectProperty;
+class						USoftClassProperty;
 class				UNameProperty;
 class				UStructProperty;
 class				UStrProperty;
diff --git a/Engine/Source/Runtime/Core/Public/UObject/UnrealNames.h b/Engine/Source/Runtime/Core/Public/UObject/UnrealNames.h
--- a/Engine/Source/Runtime/Core/Public/UObject/UnrealNames.h
+++ b/Engine/Source/Runtime/Core/Public/UObject/UnrealNames.h
@@ -10,7 +10,7 @@
 
 /** Index of highest hardcoded name to be replicated by index by the networking code
  * @warning: changing this number or making any change to the list of hardcoded names with index
- * less than this value breaks network compatibility (update GEngineMinNetVersion)
+ * less than this value breaks network compatibility, which by default checks for the same changelist
  * @note: names with a greater value than this can still be replicated, but they are sent as
  * strings instead of an efficient index
  */
diff --git a/Engine/Source/Runtime/Core/Public/UObject/UnrealNames.inl b/Engine/Source/Runtime/Core/Public/UObject/UnrealNames.inl
--- a/Engine/Source/Runtime/Core/Public/UObject/UnrealNames.inl
+++ b/Engine/Source/Runtime/Core/Public/UObject/UnrealNames.inl
@@ -8,10 +8,10 @@ REGISTER_NAME(1,ByteProperty)
 REGISTER_NAME(2,IntProperty)
 REGISTER_NAME(3,BoolProperty)
 REGISTER_NAME(4,FloatProperty)
-REGISTER_NAME(5,ObjectProperty)
+REGISTER_NAME(5,ObjectProperty) // ClassProperty shares the same tag
 REGISTER_NAME(6,NameProperty)
 REGISTER_NAME(7,DelegateProperty)
-REGISTER_NAME(8,ClassProperty)
+//REGISTER_NAME(8,Available)
 REGISTER_NAME(9,ArrayProperty)
 REGISTER_NAME(10,StructProperty)
 REGISTER_NAME(11,VectorProperty)
@@ -19,17 +19,17 @@ REGISTER_NAME(12,RotatorProperty)
 REGISTER_NAME(13,StrProperty)
 REGISTER_NAME(14,TextProperty)
 REGISTER_NAME(15,InterfaceProperty)
-REGISTER_NAME(16,MulticastDelegateProperty )
-REGISTER_NAME(17,WeakObjectProperty )
-REGISTER_NAME(18,LazyObjectProperty )
-REGISTER_NAME(19,AssetObjectProperty )
+REGISTER_NAME(16,MulticastDelegateProperty)
+//REGISTER_NAME(17,Available)
+REGISTER_NAME(18,LazyObjectProperty)
+REGISTER_NAME(19,SoftObjectProperty) // SoftClassProperty shares the same tag
 REGISTER_NAME(20,UInt64Property)
 REGISTER_NAME(21,UInt32Property)
 REGISTER_NAME(22,UInt16Property)
 REGISTER_NAME(23,Int64Property)
 REGISTER_NAME(25,Int16Property)
 REGISTER_NAME(26,Int8Property)
-REGISTER_NAME(27,AssetSubclassOfProperty)
+//REGISTER_NAME(27,Available)
 REGISTER_NAME(28,MapProperty)
 REGISTER_NAME(29,SetProperty)
 
diff --git a/Engine/Source/Runtime/CoreUObject/Private/Blueprint/BlueprintSupport.cpp b/Engine/Source/Runtime/CoreUObject/Private/Blueprint/BlueprintSupport.cpp
--- a/Engine/Source/Runtime/CoreUObject/Private/Blueprint/BlueprintSupport.cpp
+++ b/Engine/Source/Runtime/CoreUObject/Private/Blueprint/BlueprintSupport.cpp
@@ -22,6 +22,10 @@
 #include "UObject/StructScriptLoader.h"
 #include "UObject/UObjectThreadContext.h"
 
+#if USE_DEFERRED_DEPENDENCY_CHECK_VERIFICATION_TESTS
+#include "UObjectIterator.h"
+#endif
+
 DEFINE_LOG_CATEGORY_STATIC(LogBlueprintSupport, Log, All);
 
 // Flag to enable the new BlueprintCompilationManager:
@@ -686,8 +690,14 @@ bool FLinkerLoad::DeferPotentialCircularImport(const int32 Index)
 					if (DeferPotentialCircularImport(OuterImportIndex))
 					{
 						UObject* FuncOuter = ImportMap[OuterImportIndex].XObject;
-						Import.XObject = MakeImportPlaceholder<ULinkerPlaceholderFunction>(FuncOuter, *Import.ObjectName.ToString(), Index);
-						DEFERRED_DEPENDENCY_CHECK(dynamic_cast<ULinkerPlaceholderClass*>(FuncOuter) != nullptr);
+						// This is an ugly check to make sure we don't make a placeholder function for a missing native instance.
+						// We likely also need to avoid making placeholders for anything that's not outered to a ULinkerPlaceholderClass,
+						// but the DEFERRED_DEPENDENCY_CHECK may be out of date...
+						if(Cast<UClass>(FuncOuter))
+						{
+							Import.XObject = MakeImportPlaceholder<ULinkerPlaceholderFunction>(FuncOuter, *Import.ObjectName.ToString(), Index);
+							DEFERRED_DEPENDENCY_CHECK(dynamic_cast<ULinkerPlaceholderClass*>(FuncOuter) != nullptr);
+						}
 					}
 				}
 			}
diff --git a/Engine/Source/Runtime/CoreUObject/Private/CoreUObjectPrivatePCH.h b/Engine/Source/Runtime/CoreUObject/Private/CoreUObjectPrivatePCH.h
--- a/Engine/Source/Runtime/CoreUObject/Private/CoreUObjectPrivatePCH.h
+++ b/Engine/Source/Runtime/CoreUObject/Private/CoreUObjectPrivatePCH.h
@@ -231,9 +231,7 @@
 #include "Templates/Casts.h"
 #include "UObject/WeakObjectPtr.h"
 #include "UObject/LazyObjectPtr.h"
-#include "UObject/PersistentObjectPtr.h"
-#include "Misc/StringAssetReference.h"
-#include "UObject/AssetPtr.h"
+#include "UObject/SoftObjectPtr.h"
 #include "UObject/PropertyTag.h"
 #include "UObject/UnrealType.h"
 #include "UObject/CoreNetTypes.h"
diff --git a/Engine/Source/Runtime/CoreUObject/Private/Misc/RedirectCollector.cpp b/Engine/Source/Runtime/CoreUObject/Private/Misc/RedirectCollector.cpp
--- a/Engine/Source/Runtime/CoreUObject/Private/Misc/RedirectCollector.cpp
+++ b/Engine/Source/Runtime/CoreUObject/Private/Misc/RedirectCollector.cpp
@@ -11,29 +11,29 @@
 #include "Misc/PackageName.h"
 #include "UObject/LinkerLoad.h"
 #include "UObject/UObjectThreadContext.h"
+#include "UObject/SoftObjectPath.h"
 
 #if WITH_EDITOR
 
 DEFINE_LOG_CATEGORY_STATIC(LogRedirectors, Log, All);
 
-void FRedirectCollector::OnStringAssetReferenceLoaded(const FString& InString)
+void FRedirectCollector::OnSoftObjectPathLoaded(const FSoftObjectPath& InPath)
 {
-	FPackagePropertyPair ContainingPackageAndProperty;
-
-	if (InString.IsEmpty())
+	if (InPath.IsNull() || !GIsEditor)
 	{
-		// No need to track empty strings
+		// No need to track empty strings, or in standalone builds
 		return;
 	}
 
-	FStringAssetReferenceThreadContext& ThreadContext = FStringAssetReferenceThreadContext::Get();
+	FPackagePropertyPair ContainingPackageAndProperty;
+	FSoftObjectPathThreadContext& ThreadContext = FSoftObjectPathThreadContext::Get();
 
 	FName PackageName, PropertyName;
-	EStringAssetReferenceCollectType CollectType = EStringAssetReferenceCollectType::AlwaysCollect;
+	ESoftObjectPathCollectType CollectType = ESoftObjectPathCollectType::AlwaysCollect;
 
 	ThreadContext.GetSerializationOptions(PackageName, PropertyName, CollectType);
 
-	if (CollectType == EStringAssetReferenceCollectType::NeverCollect)
+	if (CollectType == ESoftObjectPathCollectType::NeverCollect)
 	{
 		// Do not track
 		return;
@@ -46,86 +46,51 @@ void FRedirectCollector::OnStringAssetReferenceLoaded(const FString& InString)
 		{
 			ContainingPackageAndProperty.SetProperty(PropertyName);
 		}
-		ContainingPackageAndProperty.SetReferencedByEditorOnlyProperty(CollectType == EStringAssetReferenceCollectType::EditorOnlyCollect);
+		ContainingPackageAndProperty.SetReferencedByEditorOnlyProperty(CollectType == ESoftObjectPathCollectType::EditorOnlyCollect);
 	}
 
 	FScopeLock ScopeLock(&CriticalSection);
 
-	StringAssetReferences.AddUnique(FName(*InString), ContainingPackageAndProperty);
+	SoftObjectPathMap.AddUnique(InPath.GetAssetPathName(), ContainingPackageAndProperty);
+}
+
+void FRedirectCollector::OnStringAssetReferenceLoaded(const FString& InString)
+{
+	FSoftObjectPath Path(InString);
+	OnSoftObjectPathLoaded(Path);
 }
 
 FString FRedirectCollector::OnStringAssetReferenceSaved(const FString& InString)
 {
 	FScopeLock ScopeLock(&CriticalSection);
 
-	FString* Found = AssetPathRedirectionMap.Find(InString);
-	if (Found)
-	{
-		return *Found;
-	}
-	return InString;
-}
+	FName Found = GetAssetPathRedirection(FName(*InString));
 
-#define REDIRECT_TIMERS 1
-#if REDIRECT_TIMERS
-struct FRedirectScopeTimer
-{
-	double& Time;
-	double StartTime;
-	FRedirectScopeTimer(double& InTime) : Time(InTime)
-	{
-		StartTime = FPlatformTime::Seconds();
-	}
-	~FRedirectScopeTimer()
+	if (Found != NAME_None)
 	{
-		Time += FPlatformTime::Seconds() - StartTime;
+		return Found.ToString();
 	}
-};
-
-double ResolveTimeTotal;
-double ResolveTimeLoad;
-double ResolveTimeDelegate;
-
-#define SCOPE_REDIRECT_TIMER(TimerName) FRedirectScopeTimer Timer##TimerName(TimerName);
-
-#define LOG_REDIRECT_TIMER( TimerName) UE_LOG(LogRedirectors, Display, TEXT("Timer %ls %f"), TEXT(#TimerName), TimerName);
-
-#define LOG_REDIRECT_TIMERS() \
-	LOG_REDIRECT_TIMER(ResolveTimeLoad); \
-	LOG_REDIRECT_TIMER(ResolveTimeDelegate);\
-	LOG_REDIRECT_TIMER(ResolveTimeTotal);\
-
-#else
-#define SCOPE_REDIRECT_TIMER(TimerName)
-#define LOG_REDIRECT_TIMERS()
-#endif
-
-void FRedirectCollector::LogTimers() const
-{
-	LOG_REDIRECT_TIMERS();
+	return InString;
 }
 
-void FRedirectCollector::ResolveStringAssetReference(FName FilterPackageFName, bool bProcessAlreadyResolvedPackages)
-{
-	SCOPE_REDIRECT_TIMER(ResolveTimeTotal);
-	
+void FRedirectCollector::ResolveAllSoftObjectPaths(FName FilterPackage)
+{	
 	FScopeLock ScopeLock(&CriticalSection);
 
 	TMultiMap<FName, FPackagePropertyPair> SkippedReferences;
-	SkippedReferences.Empty(StringAssetReferences.Num());
-	while ( StringAssetReferences.Num())
+	SkippedReferences.Empty(SoftObjectPathMap.Num());
+	while (SoftObjectPathMap.Num())
 	{
-
 		TMultiMap<FName, FPackagePropertyPair> CurrentReferences;
-		Swap(StringAssetReferences, CurrentReferences);
+		Swap(SoftObjectPathMap, CurrentReferences);
 
-		for (const auto& CurrentReference : CurrentReferences)
+		for (const TPair<FName, FPackagePropertyPair>& CurrentReference : CurrentReferences)
 		{
 			const FName& ToLoadFName = CurrentReference.Key;
 			const FPackagePropertyPair& RefFilenameAndProperty = CurrentReference.Value;
 
-			if ((FilterPackageFName != NAME_None) && // not using a filter
-				(FilterPackageFName != RefFilenameAndProperty.GetCachedPackageName()) && // this is the package we are looking for
+			if ((FilterPackage != NAME_None) && // not using a filter
+				(FilterPackage != RefFilenameAndProperty.GetCachedPackageName()) && // this is the package we are looking for
 				(RefFilenameAndProperty.GetCachedPackageName() != NAME_None) // if we have an empty package name then process it straight away
 				)
 			{
@@ -136,32 +101,11 @@ void FRedirectCollector::ResolveStringAssetReference(FName FilterPackageFName, b
 
 			const FString ToLoad = ToLoadFName.ToString();
 
-			if (FCoreDelegates::LoadStringAssetReferenceInCook.IsBound())
-			{
-				SCOPE_REDIRECT_TIMER(ResolveTimeDelegate);
-				if (FCoreDelegates::LoadStringAssetReferenceInCook.Execute(ToLoad) == false)
-				{
-					// Skip this reference
-					continue;
-				}
-			}
-			if ( (bProcessAlreadyResolvedPackages == false) && 
-				AlreadyRemapped.Contains(ToLoadFName) )
-			{
-				// don't load this package again because we already did it once before
-				// once is enough for somethings.
-				continue;
-			}
 			if (ToLoad.Len() > 0 )
 			{
-				SCOPE_REDIRECT_TIMER(ResolveTimeLoad);
-
 				UE_LOG(LogRedirectors, Verbose, TEXT("String Asset Reference '%s'"), *ToLoad);
 				UE_CLOG(RefFilenameAndProperty.GetProperty().ToString().Len(), LogRedirectors, Verbose, TEXT("    Referenced by '%s'"), *RefFilenameAndProperty.GetProperty().ToString());
 
-				// LoadObject will mark any failed loads as "KnownMissing", so since we want to print out the referencer if it was missing and not-known-missing
-				// then check here before attempting to load (TODO: Can we just not even do the load? Seems like we should be able to...)
-
 				int32 DotIndex = ToLoad.Find(TEXT("."));
 				FString PackageName = DotIndex != INDEX_NONE ? ToLoad.Left(DotIndex) : ToLoad;
 
@@ -175,17 +119,12 @@ void FRedirectCollector::ResolveStringAssetReference(FName FilterPackageFName, b
 
 				if (Loaded)
 				{
-					if (FCoreUObjectDelegates::PackageLoadedFromStringAssetReference.IsBound())
-					{
-						FCoreUObjectDelegates::PackageLoadedFromStringAssetReference.Broadcast(Loaded->GetOutermost()->GetFName());
-					}
 					FString Dest = Loaded->GetPathName();
 					UE_LOG(LogRedirectors, Verbose, TEXT("    Resolved to '%s'"), *Dest);
 					if (Dest != ToLoad)
 					{
-						AssetPathRedirectionMap.Add(ToLoad, Dest);
+						AssetPathRedirectionMap.Add(ToLoadFName, FName(*Dest));
 					}
-					AlreadyRemapped.Add(ToLoadFName);
 				}
 				else
 				{
@@ -193,23 +132,22 @@ void FRedirectCollector::ResolveStringAssetReference(FName FilterPackageFName, b
 					UE_LOG(LogRedirectors, Warning, TEXT("String Asset Reference '%s' was not found! (Referencer '%s')"), *ToLoad, *Referencer);
 				}
 			}
-
 		}
 	}
 
-	check(StringAssetReferences.Num() == 0);
+	check(SoftObjectPathMap.Num() == 0);
 	// Add any skipped references back into the map for the next time this is called
-	Swap(StringAssetReferences, SkippedReferences);
+	Swap(SoftObjectPathMap, SkippedReferences);
 	// we shouldn't have any references left if we decided to resolve them all
-	check((StringAssetReferences.Num() == 0) || (FilterPackageFName != NAME_None));
+	check((SoftObjectPathMap.Num() == 0) || (FilterPackage != NAME_None));
 }
 
-void FRedirectCollector::ProcessStringAssetReferencePackageList(FName FilterPackage, bool bGetEditorOnly, TSet<FName>& OutReferencedPackages)
+void FRedirectCollector::ProcessSoftObjectPathPackageList(FName FilterPackage, bool bGetEditorOnly, TSet<FName>& OutReferencedPackages)
 {
 	FScopeLock ScopeLock(&CriticalSection);
 
 	// Iterate map, remove all matching and potentially add to OutReferencedPackages
-	for (auto It = StringAssetReferences.CreateIterator(); It; ++It)
+	for (auto It = SoftObjectPathMap.CreateIterator(); It; ++It)
 	{
 		const FName& ToLoadFName = It->Key;
 		const FPackagePropertyPair& RefFilenameAndProperty = It->Value;
@@ -228,10 +166,10 @@ void FRedirectCollector::ProcessStringAssetReferencePackageList(FName FilterPack
 		}
 	}
 
-	StringAssetReferences.Compact();
+	SoftObjectPathMap.Compact();
 }
 
-void FRedirectCollector::AddAssetPathRedirection(const FString& OriginalPath, const FString& RedirectedPath)
+void FRedirectCollector::AddAssetPathRedirection(FName OriginalPath, FName RedirectedPath)
 {
 	FScopeLock ScopeLock(&CriticalSection);
 
@@ -239,23 +177,23 @@ void FRedirectCollector::AddAssetPathRedirection(const FString& OriginalPath, co
 	AssetPathRedirectionMap.Add(OriginalPath, RedirectedPath);
 }
 
-void FRedirectCollector::RemoveAssetPathRedirection(const FString& OriginalPath)
+void FRedirectCollector::RemoveAssetPathRedirection(FName OriginalPath)
 {
 	FScopeLock ScopeLock(&CriticalSection);
 
-	FString* Found = AssetPathRedirectionMap.Find(OriginalPath);
+	FName* Found = AssetPathRedirectionMap.Find(OriginalPath);
 
-	if (ensureMsgf(Found, TEXT("Cannot remove redirection from %s, it was not registered"), *OriginalPath))
+	if (ensureMsgf(Found, TEXT("Cannot remove redirection from %s, it was not registered"), *OriginalPath.ToString()))
 	{
 		AssetPathRedirectionMap.Remove(OriginalPath);
 	}
 }
 
-FString* FRedirectCollector::GetAssetPathRedirection(const FString& OriginalPath)
+FName FRedirectCollector::GetAssetPathRedirection(FName OriginalPath)
 {
 	FScopeLock ScopeLock(&CriticalSection);
 
-	return AssetPathRedirectionMap.Find(OriginalPath);
+	return AssetPathRedirectionMap.FindRef(OriginalPath);
 }
 
 FRedirectCollector GRedirectCollector;
diff --git a/Engine/Source/Runtime/CoreUObject/Private/Serialization/ArchiveUObject.cpp b/Engine/Source/Runtime/CoreUObject/Private/Serialization/ArchiveUObject.cpp
--- a/Engine/Source/Runtime/CoreUObject/Private/Serialization/ArchiveUObject.cpp
+++ b/Engine/Source/Runtime/CoreUObject/Private/Serialization/ArchiveUObject.cpp
@@ -20,7 +20,7 @@ DEFINE_LOG_CATEGORY_STATIC(LogArchiveUObject, Log, All);
  * don't need to bother serializing lazy object pointers.  However, serialization is required if you
  * want to load and save your object.
  */
-FArchive& FArchiveUObject::operator<<(class FLazyObjectPtr& LazyObjectPtr)
+FArchive& FArchiveUObject::operator<<(FLazyObjectPtr& Value)
 {
 	FArchive& Ar = *this;
 	// We never serialize our reference while the garbage collector is harvesting references
@@ -38,60 +38,60 @@ FArchive& FArchiveUObject::operator<<(class FLazyObjectPtr& LazyObjectPtr)
 			if (Ar.IsLoading())
 			{
 				// Reset before serializing to clear the internal weak pointer. 
-				LazyObjectPtr.Reset();
+				Value.Reset();
 			}
-			Ar << LazyObjectPtr.GetUniqueID();
+			Ar << Value.GetUniqueID();
 		}
 		else
 #endif
 		{
-			UObject* Object = LazyObjectPtr.Get();
+			UObject* Object = Value.Get();
 
 			Ar << Object;
 
 			if (IsLoading() || (Object && IsModifyingWeakAndStrongReferences()))
 			{
-				LazyObjectPtr = Object;
+				Value = Object;
 			}
 		}
 	}
 	return Ar;
 }
 
-FArchive& FArchiveUObject::operator<<( class FAssetPtr& AssetPtr )
+FArchive& FArchiveUObject::operator<<(FSoftObjectPtr& Value)
 {
 	FArchive& Ar = *this;
 	if (Ar.IsSaving() || Ar.IsLoading())
 	{
 		// Reset before serializing to clear the internal weak pointer. 
-		AssetPtr.ResetWeakPtr();
-		Ar << AssetPtr.GetUniqueID();
+		Value.ResetWeakPtr();
+		Ar << Value.GetUniqueID();
 	}
 	else if (!IsObjectReferenceCollector() || IsModifyingWeakAndStrongReferences())
 	{
 		// Treat this like a weak pointer object, as we are doing something like replacing references in memory
-		UObject* Object = AssetPtr.Get();
+		UObject* Object = Value.Get();
 
 		Ar << Object;
 
 		if (IsLoading() || (Object && IsModifyingWeakAndStrongReferences()))
 		{
-			AssetPtr = Object;
+			Value = Object;
 		}
 	}
 
 	return Ar;
 }
 
-FArchive& FArchiveUObject::operator<<(struct FStringAssetReference& Value)
+FArchive& FArchiveUObject::operator<<(FSoftObjectPath& Value)
 {
 	Value.SerializePath(*this);
 	return *this;
 }
 
-FArchive& FArchiveUObject::operator<<(FWeakObjectPtr& WeakObjectPtr)
+FArchive& FArchiveUObject::operator<<(FWeakObjectPtr& Value)
 {
-	WeakObjectPtr.Serialize(*this);
+	Value.Serialize(*this);
 	return *this;
 }
 
diff --git a/Engine/Source/Runtime/CoreUObject/Private/Serialization/AsyncLoading.cpp b/Engine/Source/Runtime/CoreUObject/Private/Serialization/AsyncLoading.cpp
--- a/Engine/Source/Runtime/CoreUObject/Private/Serialization/AsyncLoading.cpp
+++ b/Engine/Source/Runtime/CoreUObject/Private/Serialization/AsyncLoading.cpp
@@ -21,7 +21,7 @@
 #include "UObject/PackageFileSummary.h"
 #include "UObject/Linker.h"
 #include "UObject/CoreRedirects.h"
-#include "Misc/StringAssetReference.h"
+#include "UObject/SoftObjectPath.h"
 #include "UObject/LinkerLoad.h"
 #include "Serialization/DeferredMessageLog.h"
 #include "UObject/UObjectThreadContext.h"
@@ -3139,13 +3139,7 @@ void FAsyncPackage::EventDrivenCreateExport(int32 LocalExportIndex)
 
 				UObject* ActualObjectWithTheName = StaticFindObjectFastInternal(NULL, ThisParent, Export.ObjectName, true);
 
-				// if we require cooked data, attempt to find exports in memory first...which we always do for event driven loading
-				check(FPlatformProperties::RequiresCookedData()
-					|| IsAsyncLoading()
-					|| Export.bForcedExport
-					|| LinkerRoot->ShouldFindExportsInMemoryFirst()
-					);
-
+				// Always attempt to find object in memory first
 				if (ActualObjectWithTheName && (ActualObjectWithTheName->GetClass() == LoadClass))
 				{
 					Export.Object = ActualObjectWithTheName;
@@ -6360,7 +6354,7 @@ EAsyncPackageState::Type FAsyncPackage::PostLoadDeferredObjects(double InTickSta
 			::IsTimeLimitExceeded(InTickStartTime, bInUseTimeLimit, InOutTimeLimit, LastTypeOfWorkPerformed, LastObjectWorkWasPerformedOn);
 		}
 
-		FStringAssetReference::InvalidateTag();
+		FSoftObjectPath::InvalidateTag();
 		FUniqueObjectGuid::InvalidateTag();
 	}
 
diff --git a/Engine/Source/Runtime/CoreUObject/Private/Serialization/DuplicateDataReader.cpp b/Engine/Source/Runtime/CoreUObject/Private/Serialization/DuplicateDataReader.cpp
--- a/Engine/Source/Runtime/CoreUObject/Private/Serialization/DuplicateDataReader.cpp
+++ b/Engine/Source/Runtime/CoreUObject/Private/Serialization/DuplicateDataReader.cpp
@@ -3,7 +3,7 @@
 #include "CoreMinimal.h"
 #include "Serialization/ArchiveUObject.h"
 #include "UObject/LazyObjectPtr.h"
-#include "Misc/StringAssetReference.h"
+#include "UObject/SoftObjectPath.h"
 #include "UObject/PropertyPortFlags.h"
 #include "UObject/UnrealType.h"
 #include "Serialization/DuplicatedObject.h"
@@ -74,7 +74,7 @@ FArchive& FDuplicateDataReader::operator<<( UObject*& Object )
 	return *this;
 }
 
-FArchive& FDuplicateDataReader::operator<<( FLazyObjectPtr& LazyObjectPtr)
+FArchive& FDuplicateDataReader::operator<<(FLazyObjectPtr& LazyObjectPtr)
 {
 	FArchive& Ar = *this;
 	FUniqueObjectGuid ID;
@@ -89,18 +89,16 @@ FArchive& FDuplicateDataReader::operator<<( FLazyObjectPtr& LazyObjectPtr)
 	return Ar;
 }
 
-FArchive& FDuplicateDataReader::operator<<( FStringAssetReference& StringAssetReference)
+FArchive& FDuplicateDataReader::operator<<(FSoftObjectPath& SoftObjectPath)
 {
-	FArchiveUObject::operator<<(StringAssetReference);
+	FArchiveUObject::operator<<(SoftObjectPath);
 			
-	// Remap asset reference if necessary
+	// Remap soft reference if necessary
+	UObject* SourceObject = SoftObjectPath.ResolveObject();
+	FDuplicatedObject ObjectInfo = SourceObject ? DuplicatedObjectAnnotation.GetAnnotation(SourceObject) : FDuplicatedObject();
+	if (!ObjectInfo.IsDefault())
 	{
-		UObject* SourceObject = StringAssetReference.ResolveObject();
-		FDuplicatedObject ObjectInfo = SourceObject ? DuplicatedObjectAnnotation.GetAnnotation(SourceObject) : FDuplicatedObject();
-		if (!ObjectInfo.IsDefault())
-		{
-			StringAssetReference = FStringAssetReference::GetOrCreateIDForObject(ObjectInfo.DuplicatedObject);
-		}
+		SoftObjectPath = FSoftObjectPath::GetOrCreateIDForObject(ObjectInfo.DuplicatedObject);
 	}
 	
 	return *this;
diff --git a/Engine/Source/Runtime/CoreUObject/Private/Serialization/DuplicateDataWriter.cpp b/Engine/Source/Runtime/CoreUObject/Private/Serialization/DuplicateDataWriter.cpp
--- a/Engine/Source/Runtime/CoreUObject/Private/Serialization/DuplicateDataWriter.cpp
+++ b/Engine/Source/Runtime/CoreUObject/Private/Serialization/DuplicateDataWriter.cpp
@@ -2,8 +2,6 @@
 
 #include "CoreMinimal.h"
 #include "UObject/LazyObjectPtr.h"
-#include "Misc/StringAssetReference.h"
-#include "UObject/AssetPtr.h"
 #include "UObject/PropertyPortFlags.h"
 #include "Serialization/DuplicatedObject.h"
 #include "Serialization/DuplicatedDataWriter.h"
diff --git a/Engine/Source/Runtime/CoreUObject/Private/Serialization/ObjectReader.cpp b/Engine/Source/Runtime/CoreUObject/Private/Serialization/ObjectReader.cpp
--- a/Engine/Source/Runtime/CoreUObject/Private/Serialization/ObjectReader.cpp
+++ b/Engine/Source/Runtime/CoreUObject/Private/Serialization/ObjectReader.cpp
@@ -2,13 +2,12 @@
 
 #include "Serialization/ObjectReader.h"
 #include "UObject/LazyObjectPtr.h"
-#include "Misc/StringAssetReference.h"
-#include "UObject/AssetPtr.h"
+#include "UObject/SoftObjectPtr.h"
 
 ///////////////////////////////////////////////////////
 // FObjectReader
 
-FArchive& FObjectReader::operator<<( class FName& N )
+FArchive& FObjectReader::operator<<(FName& N)
 {
 	NAME_INDEX ComparisonIndex;
 	NAME_INDEX DisplayIndex;
@@ -21,35 +20,35 @@ FArchive& FObjectReader::operator<<( class FName& N )
 	return *this;
 }
 
-FArchive& FObjectReader::operator<<( class UObject*& Res )
+FArchive& FObjectReader::operator<<(UObject*& Res)
 {
 	ByteOrderSerialize(&Res, sizeof(Res));
 	return *this;
 }
 
-FArchive& FObjectReader::operator<<( class FLazyObjectPtr& LazyObjectPtr )
+FArchive& FObjectReader::operator<<(FLazyObjectPtr& Value)
 {
 	FArchive& Ar = *this;
 	FUniqueObjectGuid ID;
 	Ar << ID;
 
-	LazyObjectPtr = ID;
+	Value = ID;
 	return Ar;
 }
 
-FArchive& FObjectReader::operator<<( class FAssetPtr& AssetPtr )
+FArchive& FObjectReader::operator<<(FSoftObjectPtr& Value)
 {
-	AssetPtr.ResetWeakPtr();
-	return *this << AssetPtr.GetUniqueID();
+	Value.ResetWeakPtr();
+	return *this << Value.GetUniqueID();
 }
 
-FArchive& FObjectReader::operator<<(FStringAssetReference& Value)
+FArchive& FObjectReader::operator<<(FSoftObjectPath& Value)
 {
 	Value.SerializePath(*this);
 	return *this;
 }
 
-FArchive& FObjectReader::operator<< (struct FWeakObjectPtr& Value)
+FArchive& FObjectReader::operator<<(FWeakObjectPtr& Value)
 {
 	Value.Serialize(*this);
 	return *this;
diff --git a/Engine/Source/Runtime/CoreUObject/Private/Serialization/ObjectWriter.cpp b/Engine/Source/Runtime/CoreUObject/Private/Serialization/ObjectWriter.cpp
--- a/Engine/Source/Runtime/CoreUObject/Private/Serialization/ObjectWriter.cpp
+++ b/Engine/Source/Runtime/CoreUObject/Private/Serialization/ObjectWriter.cpp
@@ -2,8 +2,7 @@
 
 #include "Serialization/ObjectWriter.h"
 #include "UObject/LazyObjectPtr.h"
-#include "Misc/StringAssetReference.h"
-#include "UObject/AssetPtr.h"
+#include "UObject/SoftObjectPtr.h"
 
 ///////////////////////////////////////////////////////
 // FObjectWriter
@@ -25,25 +24,25 @@ FArchive& FObjectWriter::operator<<( class UObject*& Res )
 	return *this;
 }
 
-FArchive& FObjectWriter::operator<<( class FLazyObjectPtr& LazyObjectPtr )
+FArchive& FObjectWriter::operator<<(FLazyObjectPtr& Value)
 {
-	FUniqueObjectGuid ID = LazyObjectPtr.GetUniqueID();
+	FUniqueObjectGuid ID = Value.GetUniqueID();
 	return *this << ID;
 }
 
-FArchive& FObjectWriter::operator<<( class FAssetPtr& AssetPtr )
+FArchive& FObjectWriter::operator<<(FSoftObjectPtr& Value)
 {
-	AssetPtr.ResetWeakPtr();
-	return *this << AssetPtr.GetUniqueID();
+	Value.ResetWeakPtr();
+	return *this << Value.GetUniqueID();
 }
 
-FArchive& FObjectWriter::operator<<(FStringAssetReference& Value)
+FArchive& FObjectWriter::operator<<(FSoftObjectPath& Value)
 {
 	Value.SerializePath(*this);
 	return *this;
 }
 
-FArchive& FObjectWriter::operator<< (struct FWeakObjectPtr& Value)
+FArchive& FObjectWriter::operator<<(FWeakObjectPtr& Value)
 {
 	Value.Serialize(*this);
 	return *this;
diff --git a/Engine/Source/Runtime/CoreUObject/Private/UObject/Class.cpp b/Engine/Source/Runtime/CoreUObject/Private/UObject/Class.cpp
--- a/Engine/Source/Runtime/CoreUObject/Private/UObject/Class.cpp
+++ b/Engine/Source/Runtime/CoreUObject/Private/UObject/Class.cpp
@@ -35,6 +35,7 @@
 #include "UObject/PropertyHelper.h"
 #include "UObject/CoreRedirects.h"
 #include "Serialization/ArchiveScriptReferenceCollector.h"
+#include "UObject/FrameworkObjectVersion.h"
 
 // This flag enables some expensive class tree validation that is meant to catch mutations of 
 // the class tree outside of SetSuperStruct. It has been disabled because loading blueprints 
@@ -172,7 +173,11 @@ void UField::PostLoad()
 void UField::Serialize( FArchive& Ar )
 {
 	Super::Serialize( Ar );
-	Ar << Next;
+	Ar.UsingCustomVersion(FFrameworkObjectVersion::GUID);
+	if (Ar.CustomVer(FFrameworkObjectVersion::GUID) < FFrameworkObjectVersion::RemoveUField_Next)
+	{
+		Ar << Next;
+	}
 }
 
 void UField::AddCppProperty( UProperty* Property )
@@ -546,7 +551,10 @@ void UStruct::GetPreloadDependencies(TArray<UObject*>& OutDeps)
 
 	for (UField* Field = Children; Field; Field = Field->Next)
 	{
-		OutDeps.Add(Field);
+		if (!Cast<UFunction>(Field))
+		{
+			OutDeps.Add(Field);
+		}
 	}
 }
 
@@ -563,8 +571,10 @@ void UStruct::Link(FArchive& Ar, bool bRelinkExistingProperties)
 
 		for (UField* Field = Children; Field; Field = Field->Next)
 		{
-			// calling Preload here is required in order to load the value of Field->Next
-			Ar.Preload(Field);
+			if (!GEventDrivenLoaderEnabled || !Cast<UFunction>(Field))
+			{
+				Ar.Preload(Field);
+			}
 		}
 
 		int32 LoopNum = 1;
@@ -1181,7 +1191,43 @@ void UStruct::Serialize( FArchive& Ar )
 	Super::Serialize( Ar );
 
 	SerializeSuperStruct(Ar);
-	Ar << Children;
+	Ar.UsingCustomVersion(FFrameworkObjectVersion::GUID);
+	if (Ar.CustomVer(FFrameworkObjectVersion::GUID) < FFrameworkObjectVersion::RemoveUField_Next)
+	{
+		Ar << Children;
+	}
+	else
+	{
+		TArray<UField*> ChildArray;
+		if (Ar.IsLoading())
+		{
+			Ar << ChildArray;
+			if (ChildArray.Num())
+			{
+				for (int32 Index = 0; Index + 1 < ChildArray.Num(); Index++)
+				{
+					ChildArray[Index]->Next = ChildArray[Index + 1];
+				}
+				Children = ChildArray[0];
+				ChildArray[ChildArray.Num() - 1]->Next = nullptr;
+			}
+			else
+			{
+				Children = nullptr;
+			}
+		}
+		else
+		{
+			UField* Child = Children;
+			while (Child)
+			{
+				ChildArray.Add(Child);
+				Child = Child->Next;
+			}
+			Ar << ChildArray;
+		}
+	}
+
 
 	if (Ar.IsLoading())
 	{
@@ -2638,15 +2684,18 @@ UObject* UClass::CreateDefaultObject()
 				auto ClassLinker = GetLinker();
 				if (ClassLinker)
 				{
-					UField* FieldIt = Children;
-					while(FieldIt && (FieldIt->GetOuter() == this))
+					if (!GEventDrivenLoaderEnabled)
 					{
-						// If we've had cyclic dependencies between classes here, we might need to preload to ensure that we load the rest of the property chain
-						if( FieldIt->HasAnyFlags(RF_NeedLoad) )
+						UField* FieldIt = Children;
+						while (FieldIt && (FieldIt->GetOuter() == this))
 						{
-							ClassLinker->Preload(FieldIt);
+							// If we've had cyclic dependencies between classes here, we might need to preload to ensure that we load the rest of the property chain
+							if (FieldIt->HasAnyFlags(RF_NeedLoad))
+							{
+								ClassLinker->Preload(FieldIt);
+							}
+							FieldIt = FieldIt->Next;
 						}
-						FieldIt = FieldIt->Next;
 					}
 					
 					StaticLink(true);
@@ -4802,15 +4851,15 @@ UScriptStruct* TBaseStructure<FInt32Interval>::Get()
 	return ScriptStruct;
 }
 
-UScriptStruct* TBaseStructure<FStringAssetReference>::Get()
+UScriptStruct* TBaseStructure<FSoftObjectPath>::Get()
 {
-	static auto ScriptStruct = StaticGetBaseStructureInternal(TEXT("StringAssetReference"));
+	static auto ScriptStruct = StaticGetBaseStructureInternal(TEXT("SoftObjectPath"));
 	return ScriptStruct;
 }
 
-UScriptStruct* TBaseStructure<FStringClassReference>::Get()
+UScriptStruct* TBaseStructure<FSoftClassPath>::Get()
 {
-	static auto ScriptStruct = StaticGetBaseStructureInternal(TEXT("StringClassReference"));
+	static auto ScriptStruct = StaticGetBaseStructureInternal(TEXT("SoftClassPath"));
 	return ScriptStruct;
 }
 
diff --git a/Engine/Source/Runtime/CoreUObject/Private/UObject/CoreNet.cpp b/Engine/Source/Runtime/CoreUObject/Private/UObject/CoreNet.cpp
--- a/Engine/Source/Runtime/CoreUObject/Private/UObject/CoreNet.cpp
+++ b/Engine/Source/Runtime/CoreUObject/Private/UObject/CoreNet.cpp
@@ -350,7 +350,7 @@ FArchive& FNetBitWriter::operator<<( UObject*& Object )
 	return *this;
 }
 
-FArchive& FNetBitWriter::operator<<(FStringAssetReference& Value)
+FArchive& FNetBitWriter::operator<<(FSoftObjectPath& Value)
 {
 	FString Path = Value.ToString();
 
@@ -391,7 +391,7 @@ FArchive& FNetBitReader::operator<<( class FName& N )
 	return *this;
 }
 
-FArchive& FNetBitReader::operator<<(FStringAssetReference& Value)
+FArchive& FNetBitReader::operator<<(FSoftObjectPath& Value)
 {
 	FString Path = Value.ToString();
 
diff --git a/Engine/Source/Runtime/CoreUObject/Private/UObject/Linker.cpp b/Engine/Source/Runtime/CoreUObject/Private/UObject/Linker.cpp
--- a/Engine/Source/Runtime/CoreUObject/Private/UObject/Linker.cpp
+++ b/Engine/Source/Runtime/CoreUObject/Private/UObject/Linker.cpp
@@ -172,7 +172,7 @@ void FLinker::Serialize( FArchive& Ar )
 		Ar << ImportMap;
 		Ar << ExportMap;
 		Ar << DependsMap;
-		Ar << StringAssetReferencesMap;
+		Ar << SoftPackageReferenceList;
 		Ar << GatherableTextDataMap;
 		Ar << SearchableNamesMap;
 	}
diff --git a/Engine/Source/Runtime/CoreUObject/Private/UObject/LinkerLoad.cpp b/Engine/Source/Runtime/CoreUObject/Private/UObject/LinkerLoad.cpp
--- a/Engine/Source/Runtime/CoreUObject/Private/UObject/LinkerLoad.cpp
+++ b/Engine/Source/Runtime/CoreUObject/Private/UObject/LinkerLoad.cpp
@@ -3762,48 +3762,40 @@ UObject* FLinkerLoad::CreateExport( int32 Index )
 		// Try to find existing object first in case we're a forced export to be able to reconcile. Also do it for the
 		// case of async loading as we cannot in-place replace objects.
 
-		UObject* ActualObjectWithTheName = StaticFindObjectFastInternal(NULL, ThisParent, Export.ObjectName, true);
+		UObject* ActualObjectWithTheName = StaticFindObjectFastInternal(nullptr, ThisParent, Export.ObjectName, true);
 		
-		// if we require cooked data, attempt to find exports in memory first
-		if(	FPlatformProperties::RequiresCookedData() 
-			|| IsAsyncLoading()
-			|| Export.bForcedExport
-			|| LinkerRoot->ShouldFindExportsInMemoryFirst()
-			)
+		// Find object after making sure it isn't already set. This would be bad as the code below NULLs it in a certain
+		// case, which if it had been set would cause a linker detach mismatch.
+		check(Export.Object == nullptr);
+		if (ActualObjectWithTheName && (ActualObjectWithTheName->GetClass() == LoadClass))
 		{
-			// Find object after making sure it isn't already set. This would be bad as the code below NULLs it in a certain
-			// case, which if it had been set would cause a linker detach mismatch.
-			check( Export.Object == NULL );
-			if (ActualObjectWithTheName && (ActualObjectWithTheName->GetClass() == LoadClass))
+			Export.Object = ActualObjectWithTheName;
+		}
+
+		// Object is found in memory.
+		if (Export.Object)
+		{
+			// Mark that we need to dissociate forced exports later on if we are a forced export.
+			if (Export.bForcedExport)
 			{
-				Export.Object = ActualObjectWithTheName;
+				FUObjectThreadContext::Get().ForcedExportCount++;
 			}
-
-			// Object is found in memory.
-			if( Export.Object )
+			// Associate linker with object to avoid detachment mismatches.
+			else
 			{
-				// Mark that we need to dissociate forced exports later on if we are a forced export.
-				if( Export.bForcedExport )
-				{
-					FUObjectThreadContext::Get().ForcedExportCount++;
-				}
-				// Associate linker with object to avoid detachment mismatches.
-				else
-				{
-					Export.Object->SetLinker( this, Index );
+				Export.Object->SetLinker(this, Index);
 
-					// If this object was allocated but never loaded (components created by a constructor) make sure it gets loaded
-					// Do this for all subobjects created in the native constructor.
-					FUObjectThreadContext::Get().ObjLoaded.AddUnique(Export.Object);
-					if (!Export.Object->HasAnyFlags(RF_LoadCompleted) &&
-						(Export.Object->HasAnyFlags(RF_DefaultSubObject) || (ThisParent && ThisParent->IsTemplate(RF_ClassDefaultObject))))
-					{
-						check(!GEventDrivenLoaderEnabled || !EVENT_DRIVEN_ASYNC_LOAD_ACTIVE_AT_RUNTIME);
-						Export.Object->SetFlags(RF_NeedLoad | RF_NeedPostLoad | RF_NeedPostLoadSubobjects | RF_WasLoaded);
-					}
+				// If this object was allocated but never loaded (components created by a constructor) make sure it gets loaded
+				// Do this for all subobjects created in the native constructor.
+				FUObjectThreadContext::Get().ObjLoaded.AddUnique(Export.Object);
+				if (!Export.Object->HasAnyFlags(RF_LoadCompleted) &&
+					(Export.Object->HasAnyFlags(RF_DefaultSubObject) || (ThisParent && ThisParent->IsTemplate(RF_ClassDefaultObject))))
+				{
+					check(!GEventDrivenLoaderEnabled || !EVENT_DRIVEN_ASYNC_LOAD_ACTIVE_AT_RUNTIME);
+					Export.Object->SetFlags(RF_NeedLoad | RF_NeedPostLoad | RF_NeedPostLoadSubobjects | RF_WasLoaded);
 				}
-				return Export.Object;
 			}
+			return Export.Object;
 		}
 
 		// In cases when an object has been consolidated but its package hasn't been saved, look for UObjectRedirector before
diff --git a/Engine/Source/Runtime/CoreUObject/Private/UObject/LinkerPlaceholderBase.cpp b/Engine/Source/Runtime/CoreUObject/Private/UObject/LinkerPlaceholderBase.cpp
--- a/Engine/Source/Runtime/CoreUObject/Private/UObject/LinkerPlaceholderBase.cpp
+++ b/Engine/Source/Runtime/CoreUObject/Private/UObject/LinkerPlaceholderBase.cpp
@@ -492,11 +492,11 @@ int32 TLinkerImportPlaceholder<UClass>::ResolvePropertyReferences(UClass* Replac
 					++ReplacementCount;
 				}
 			}
-			else if (UAssetClassProperty* AssetClassProperty = Cast<UAssetClassProperty>(BaseObjProperty))
+			else if (USoftClassProperty* SoftClassProperty = Cast<USoftClassProperty>(BaseObjProperty))
 			{
-				if (AssetClassProperty->MetaClass == PlaceholderClass)
+				if (SoftClassProperty->MetaClass == PlaceholderClass)
 				{
-					AssetClassProperty->MetaClass = ReplacementClass;
+					SoftClassProperty->MetaClass = ReplacementClass;
 					++ReplacementCount;
 				}	
 			}
diff --git a/Engine/Source/Runtime/CoreUObject/Private/UObject/LinkerSave.cpp b/Engine/Source/Runtime/CoreUObject/Private/UObject/LinkerSave.cpp
--- a/Engine/Source/Runtime/CoreUObject/Private/UObject/LinkerSave.cpp
+++ b/Engine/Source/Runtime/CoreUObject/Private/UObject/LinkerSave.cpp
@@ -7,8 +7,7 @@
 #include "UObject/Class.h"
 #include "Templates/Casts.h"
 #include "UObject/LazyObjectPtr.h"
-#include "Misc/StringAssetReference.h"
-#include "UObject/AssetPtr.h"
+#include "UObject/SoftObjectPtr.h"
 #include "Internationalization/TextPackageNamespaceUtil.h"
 
 /*----------------------------------------------------------------------------
@@ -262,26 +261,26 @@ FArchive& FLinkerSave::operator<<( UObject*& Obj )
 	return *this << Save;
 }
 
-FArchive& FLinkerSave::operator<<( FLazyObjectPtr& LazyObjectPtr)
+FArchive& FLinkerSave::operator<<(FLazyObjectPtr& LazyObjectPtr)
 {
 	FUniqueObjectGuid ID;
 	ID = LazyObjectPtr.GetUniqueID();
 	return *this << ID;
 }
 
-FArchive& FLinkerSave::operator<<( FAssetPtr& AssetPtr)
+FArchive& FLinkerSave::operator<<(FSoftObjectPtr& SoftObjectPtr)
 {
-	FStringAssetReference ID;
-	UObject *Object = AssetPtr.Get();
+	FSoftObjectPath ID;
+	UObject *Object = SoftObjectPtr.Get();
 
 	if (Object)
 	{
 		// Use object in case name has changed. 
-		ID = FStringAssetReference(Object);
+		ID = FSoftObjectPath(Object);
 	}
 	else
 	{
-		ID = AssetPtr.GetUniqueID();
+		ID = SoftObjectPtr.GetUniqueID();
 	}
 
 	ID.Serialize(*this);
diff --git a/Engine/Source/Runtime/CoreUObject/Private/UObject/Obj.cpp b/Engine/Source/Runtime/CoreUObject/Private/UObject/Obj.cpp
--- a/Engine/Source/Runtime/CoreUObject/Private/UObject/Obj.cpp
+++ b/Engine/Source/Runtime/CoreUObject/Private/UObject/Obj.cpp
@@ -28,7 +28,7 @@
 #include "UObject/MetaData.h"
 #include "Templates/Casts.h"
 #include "UObject/LazyObjectPtr.h"
-#include "Misc/StringAssetReference.h"
+#include "UObject/SoftObjectPtr.h"
 #include "UObject/PropertyPortFlags.h"
 #include "UObject/UnrealType.h"
 #include "UObject/ObjectRedirector.h"
@@ -1241,7 +1241,7 @@ void UObject::Serialize( FArchive& Ar )
 	// Invalidate asset pointer caches when loading a new object
 	if (Ar.IsLoading() )
 	{
-		FStringAssetReference::InvalidateTag();
+		FSoftObjectPath::InvalidateTag();
 	}
 
 	// Memory counting (with proper alignment to match C++)
@@ -3135,9 +3135,9 @@ bool StaticExec( UWorld* InWorld, const TCHAR* Cmd, FOutputDevice& Ar )
 						{
 							ExtraInfo = FString::Printf(TEXT("class<%s>"), *ClassProperty->MetaClass->GetName());
 						}
-						else if (UAssetClassProperty* AssetClassProperty = dynamic_cast<UAssetClassProperty*>(It->GetClass()))
+						else if (USoftClassProperty* SoftClassProperty = dynamic_cast<USoftClassProperty*>(It->GetClass()))
 						{
-							ExtraInfo = FString::Printf(TEXT("AssetSubclassOf<%s>"), *AssetClassProperty->MetaClass->GetName());
+							ExtraInfo = FString::Printf(TEXT("AssetSubclassOf<%s>"), *SoftClassProperty->MetaClass->GetName());
 						}
 						else if (UObjectPropertyBase* ObjectPropertyBase = dynamic_cast<UObjectPropertyBase*>(It->GetClass()))
 						{
diff --git a/Engine/Source/Runtime/CoreUObject/Private/UObject/Package.cpp b/Engine/Source/Runtime/CoreUObject/Private/UObject/Package.cpp
--- a/Engine/Source/Runtime/CoreUObject/Private/UObject/Package.cpp
+++ b/Engine/Source/Runtime/CoreUObject/Private/UObject/Package.cpp
@@ -154,20 +154,13 @@ UMetaData* UPackage::GetMetaData()
 void UPackage::FullyLoad()
 {
 	// Make sure we're a topmost package.
-	checkf(GetOuter()==NULL, TEXT("Package is not topmost. Name:%s Path: %s"), *GetName(), *GetPathName());
+	checkf(GetOuter()==nullptr, TEXT("Package is not topmost. Name:%s Path: %s"), *GetName(), *GetPathName());
 
 	// Only perform work if we're not already fully loaded.
-	if( !IsFullyLoaded() )
+	if(!IsFullyLoaded())
 	{
-		// Mark package so exports are found in memory first instead of being clobbered.
-		bool bSavedState = ShouldFindExportsInMemoryFirst();
-		FindExportsInMemoryFirst( true );
-
 		// Re-load this package.
-		LoadPackage( NULL, *GetName(), LOAD_None );
-
-		// Restore original state.
-		FindExportsInMemoryFirst( bSavedState );
+		LoadPackage(nullptr, *GetName(), LOAD_None);
 	}
 }
 
diff --git a/Engine/Source/Runtime/CoreUObject/Private/UObject/PackageFileSummary.cpp b/Engine/Source/Runtime/CoreUObject/Private/UObject/PackageFileSummary.cpp
--- a/Engine/Source/Runtime/CoreUObject/Private/UObject/PackageFileSummary.cpp
+++ b/Engine/Source/Runtime/CoreUObject/Private/UObject/PackageFileSummary.cpp
@@ -181,7 +181,7 @@ FArchive& operator<<( FArchive& Ar, FPackageFileSummary& Sum )
 
 		if (Ar.IsSaving() || Sum.FileVersionUE4 >= VER_UE4_ADD_STRING_ASSET_REFERENCES_MAP)
 		{
-			Ar << Sum.StringAssetReferencesCount << Sum.StringAssetReferencesOffset;
+			Ar << Sum.SoftPackageReferencesCount << Sum.SoftPackageReferencesOffset;
 		}
 
 		if (Ar.IsSaving() || Sum.FileVersionUE4 >= VER_UE4_ADDED_SEARCHABLE_NAMES)
diff --git a/Engine/Source/Runtime/CoreUObject/Private/UObject/Property.cpp b/Engine/Source/Runtime/CoreUObject/Private/UObject/Property.cpp
--- a/Engine/Source/Runtime/CoreUObject/Private/UObject/Property.cpp
+++ b/Engine/Source/Runtime/CoreUObject/Private/UObject/Property.cpp
@@ -15,7 +15,7 @@
 #include "UObject/UnrealType.h"
 #include "UObject/PropertyHelper.h"
 #include "UObject/CoreRedirects.h"
-#include "Misc/StringClassReference.h"
+#include "UObject/SoftObjectPath.h"
 #include "Math/Box2D.h"
 
 DEFINE_LOG_CATEGORY(LogProperty);
@@ -273,7 +273,7 @@ IMPLEMENT_STRUCT(Timespan);
 
 
 template<>
-struct TStructOpsTypeTraits<FStringAssetReference> : public TStructOpsTypeTraitsBase2<FStringAssetReference>
+struct TStructOpsTypeTraits<FSoftObjectPath> : public TStructOpsTypeTraitsBase2<FSoftObjectPath>
 {
 	enum
 	{
@@ -286,10 +286,10 @@ struct TStructOpsTypeTraits<FStringAssetReference> : public TStructOpsTypeTraits
 		WithSerializeFromMismatchedTag = true,
 	};
 };
-IMPLEMENT_STRUCT(StringAssetReference);
+IMPLEMENT_STRUCT(SoftObjectPath);
 
 template<>
-struct TStructOpsTypeTraits<FStringClassReference> : public TStructOpsTypeTraitsBase2<FStringClassReference>
+struct TStructOpsTypeTraits<FSoftClassPath> : public TStructOpsTypeTraitsBase2<FSoftClassPath>
 {
 	enum
 	{
@@ -302,7 +302,7 @@ struct TStructOpsTypeTraits<FStringClassReference> : public TStructOpsTypeTraits
 		WithSerializeFromMismatchedTag = true,
 	};
 };
-IMPLEMENT_STRUCT(StringClassReference);
+IMPLEMENT_STRUCT(SoftClassPath);
 
 template<>
 struct TStructOpsTypeTraits<FPrimaryAssetType> : public TStructOpsTypeTraitsBase2<FPrimaryAssetType>
diff --git a/Engine/Source/Runtime/CoreUObject/Private/UObject/PropertyBaseObject.cpp b/Engine/Source/Runtime/CoreUObject/Private/UObject/PropertyBaseObject.cpp
--- a/Engine/Source/Runtime/CoreUObject/Private/UObject/PropertyBaseObject.cpp
+++ b/Engine/Source/Runtime/CoreUObject/Private/UObject/PropertyBaseObject.cpp
@@ -382,7 +382,7 @@ UObject* UObjectPropertyBase::FindImportedObject( const UProperty* Property, UOb
 		if (Result == nullptr && (PortFlags & PPF_SerializedAsImportText))
 		{
 			// Check string asset redirectors
-			FStringAssetReference Path = FString(Text);
+			FSoftObjectPath Path = FString(Text);
 			if (Path.PreSavePath())
 			{
 				Result = StaticFindObjectSafe(ObjectClass, nullptr, *Path.ToString());
diff --git a/Engine/Source/Runtime/CoreUObject/Private/UObject/PropertyMap.cpp b/Engine/Source/Runtime/CoreUObject/Private/UObject/PropertyMap.cpp
--- a/Engine/Source/Runtime/CoreUObject/Private/UObject/PropertyMap.cpp
+++ b/Engine/Source/Runtime/CoreUObject/Private/UObject/PropertyMap.cpp
@@ -848,7 +848,7 @@ bool UMapProperty::ConvertFromType(const FPropertyTag& Tag, FArchive& Ar, uint8*
 		if(CurrentType->GetID() == InTag.Type)
 		{
 			uint8* DestAddress = CurrentType->ContainerPtrToValuePtr<uint8>(InData, InTag.ArrayIndex);
-			CurrentType->SerializeItem(InAr, DestAddress, InDefaultsStruct);
+			CurrentType->SerializeItem(InAr, DestAddress, nullptr);
 			return true;
 		}
 		else if( CurrentType->ConvertFromType(InTag, InAr, InData, InDefaultsStruct, bDummyAdvance) )
diff --git a/Engine/Source/Runtime/CoreUObject/Private/UObject/PropertyObject.cpp b/Engine/Source/Runtime/CoreUObject/Private/UObject/PropertyObject.cpp
--- a/Engine/Source/Runtime/CoreUObject/Private/UObject/PropertyObject.cpp
+++ b/Engine/Source/Runtime/CoreUObject/Private/UObject/PropertyObject.cpp
@@ -2,7 +2,7 @@
 
 #include "CoreMinimal.h"
 #include "UObject/ObjectMacros.h"
-#include "UObject/AssetPtr.h"
+#include "UObject/SoftObjectPtr.h"
 #include "UObject/UnrealType.h"
 #include "Blueprint/BlueprintSupport.h"
 #include "UObject/LinkerPlaceholderBase.h"
@@ -32,14 +32,16 @@ FString UObjectProperty::GetCPPMacroType( FString& ExtendedTypeText ) const
 
 bool UObjectProperty::ConvertFromType(const FPropertyTag& Tag, FArchive& Ar, uint8* Data, UStruct* DefaultsStruct, bool& bOutAdvanceProperty)
 {
-	if (Tag.Type == NAME_AssetObjectProperty || Tag.Type == NAME_AssetSubclassOfProperty)
+	static FName NAME_AssetObjectProperty = "AssetObjectProperty"; // old name of soft object property
+
+	if (Tag.Type == NAME_SoftObjectProperty || Tag.Type == NAME_AssetObjectProperty)
 	{
-		// This property used to be a TAssetPtr<Foo> but is now a raw UObjectProperty Foo*, we can convert without loss of data
-		FAssetPtr PreviousValue;
+		// This property used to be a TSoftObjectPtr<Foo> but is now a raw UObjectProperty Foo*, we can convert without loss of data
+		FSoftObjectPtr PreviousValue;
 		Ar << PreviousValue;
 
 		// now copy the value into the object's address space
-		UObject* PreviousValueObj = PreviousValue.ToStringReference().TryLoad();
+		UObject* PreviousValueObj = PreviousValue.LoadSynchronous();
 		SetPropertyValue_InContainer(Data, PreviousValueObj, Tag.ArrayIndex);
 
 		return true;
