- Mirrored fix from //UE4/Dev-Framework (CL# 3611667).

Note: This was a minor regression introduced by the last Dev-Core integration to Main.

diff --git a/Engine/Source/Programs/UnrealHeaderTool/Private/CodeGenerator.cpp b/Engine/Source/Programs/UnrealHeaderTool/Private/CodeGenerator.cpp
--- a/Engine/Source/Programs/UnrealHeaderTool/Private/CodeGenerator.cpp
+++ b/Engine/Source/Programs/UnrealHeaderTool/Private/CodeGenerator.cpp
@@ -3196,7 +3196,7 @@ void FNativeClassHeaderGenerator::ExportGeneratedStructBodyMacros(FOutputDevice&
 	}
 	else
 	{
-		OuterFunc = FString::Printf(TEXT("[](){ return FindOrConstructDynamicTypePackage(TEXT(\"%s\")); }"), *FClass::GetTypePackageName(Struct));
+		OuterFunc = FString::Printf(TEXT("[](){ return (UObject*)FindOrConstructDynamicTypePackage(TEXT(\"%s\")); }"), *FClass::GetTypePackageName(Struct));
 	}
 
 	FString MetaDataParams = OutputMetaDataCodeForObject(GeneratedStructRegisterFunctionText, Struct, TEXT("Struct_MetaDataParams"), TEXT("\t\t\t"));
