Explicitly initialize const FString to workaround Linux compiler bug
#rnx

diff --git a/Engine/Source/Programs/UnrealHeaderTool/Private/CodeGenerator.cpp b/Engine/Source/Programs/UnrealHeaderTool/Private/CodeGenerator.cpp
--- a/Engine/Source/Programs/UnrealHeaderTool/Private/CodeGenerator.cpp
+++ b/Engine/Source/Programs/UnrealHeaderTool/Private/CodeGenerator.cpp
@@ -2060,7 +2060,7 @@ void FNativeClassHeaderGenerator::ExportNativeGeneratedInitCode(FOutputDevice& O
 
 	FString OverriddenClassName = *FNativeClassHeaderGenerator::GetOverriddenName(Class);
 
-	const FString EmptyString;
+	const FString EmptyString = FString();
 	const FString& InitSearchableValuesFunctionName = bIsDynamic ? Class->GetMetaData(TEXT("InitializeStaticSearchableValues")) : EmptyString;
 	const FString InitSearchableValuesFunctionParam = InitSearchableValuesFunctionName.IsEmpty() ? FString(TEXT("nullptr")) :
 		FString::Printf(TEXT("&%s::%s"), ClassNameCPP, *InitSearchableValuesFunctionName);
