Blueprint classes derived from a native parent class that inherit from a native interface class containing a BlueprintNativeEvent declaration will no longer crash if the associated _Implementation() method is not also implemented by the native parent class.

- Mirrored from 4081318.

#rnx

diff --git a/Engine/Source/Programs/UnrealHeaderTool/Private/CodeGenerator.cpp b/Engine/Source/Programs/UnrealHeaderTool/Private/CodeGenerator.cpp
--- a/Engine/Source/Programs/UnrealHeaderTool/Private/CodeGenerator.cpp
+++ b/Engine/Source/Programs/UnrealHeaderTool/Private/CodeGenerator.cpp
@@ -4112,22 +4112,22 @@ void FNativeClassHeaderGenerator::ExportNativeFunctionHeader(
 			// all methods in interface classes are pure virtuals
 			if (bIsK2Override)
 			{
-				// For BlueprintNativeEvent methods we use the PURE_VIRTUAL() expansion. This allows Blueprints that implement the interface class to be nativized.
+				// For BlueprintNativeEvent methods we emit a stub implementation. This allows Blueprints that implement the interface class to be nativized.
 				FString ReturnValue;
 				if (ReturnProperty != nullptr)
 				{
 					UByteProperty* ByteProperty = Cast<UByteProperty>(ReturnProperty);
 					if (ByteProperty != nullptr && ByteProperty->Enum != nullptr && ByteProperty->Enum->GetCppForm() != UEnum::ECppForm::EnumClass)
 					{
-						ReturnValue = FString::Printf(TEXT(" return TEnumAsByte<%s>(%s);"), *ByteProperty->Enum->CppType, *GetNullParameterValue(ReturnProperty, false));
+						ReturnValue = FString::Printf(TEXT(" return TEnumAsByte<%s>(%s); "), *ByteProperty->Enum->CppType, *GetNullParameterValue(ReturnProperty, false));
 					}
 					else
 					{
-						ReturnValue = FString::Printf(TEXT(" return %s;"), *GetNullParameterValue(ReturnProperty, false));
+						ReturnValue = FString::Printf(TEXT(" return %s; "), *GetNullParameterValue(ReturnProperty, false));
 					}
 				}
 
-				Out.Logf(TEXT(" PURE_VIRTUAL(%s::%s,%s)"), NameLookupCPP.GetNameCPP(CastChecked<UClass>(Function->GetOuter()), true), *FunctionName, *ReturnValue);
+				Out.Logf(TEXT(" {%s}"), *ReturnValue);
 			}
 			else
 			{
