PR #5024: Require BlueprintAssignable delegate parameters to be supported by blueprints (Contributed by [REDACTED_PERSON])

diff --git a/Engine/Source/Programs/UnrealHeaderTool/Private/HeaderParser.cpp b/Engine/Source/Programs/UnrealHeaderTool/Private/HeaderParser.cpp
--- a/Engine/Source/Programs/UnrealHeaderTool/Private/HeaderParser.cpp
+++ b/Engine/Source/Programs/UnrealHeaderTool/Private/HeaderParser.cpp
@@ -2698,6 +2698,14 @@ void FHeaderParser::FixupDelegateProperties( FClasses& AllClasses, UStruct* Stru
 							for (TFieldIterator<UProperty> PropIt(SourceDelegateFunction); PropIt && (PropIt->PropertyFlags & CPF_Parm); ++PropIt)
 							{
 								UProperty* FuncParam = *PropIt;
+
+								if (!IsPropertySupportedByBlueprint(FuncParam, false))
+								{
+									FString ExtendedCPPType;
+									FString CPPType = FuncParam->GetCPPType(&ExtendedCPPType);
+									UE_LOG_ERROR_UHT(TEXT("Type '%s%s' is not supported by blueprint. %s.%s"), *CPPType, *ExtendedCPPType, *SourceDelegateFunction->GetName(), *FuncParam->GetName());
+								}
+
 								if(FuncParam->HasAllPropertyFlags(CPF_OutParm) && !FuncParam->HasAllPropertyFlags(CPF_ConstParm)  )
 								{
 									const bool bClassGeneratedFromBP = FClass::IsDynamic(Struct);
diff --git a/Engine/Source/Runtime/AIModule/Classes/AITypes.h b/Engine/Source/Runtime/AIModule/Classes/AITypes.h
--- a/Engine/Source/Runtime/AIModule/Classes/AITypes.h
+++ b/Engine/Source/Runtime/AIModule/Classes/AITypes.h
@@ -426,7 +426,7 @@ namespace FAIResources
 }
 
 
-USTRUCT()
+USTRUCT(BlueprintType)
 struct AIMODULE_API FAIRequestID
 {
 	GENERATED_USTRUCT_BODY()
