Use define instead of duplicating text
#rnx

diff --git a/Engine/Source/Programs/UnrealHeaderTool/Private/ParserHelper.h b/Engine/Source/Programs/UnrealHeaderTool/Private/ParserHelper.h
--- a/Engine/Source/Programs/UnrealHeaderTool/Private/ParserHelper.h
+++ b/Engine/Source/Programs/UnrealHeaderTool/Private/ParserHelper.h
@@ -1069,7 +1069,7 @@ struct FFuncInfo
 		FString FunctionName = FunctionReference->GetName();
 		if( FunctionReference->HasAnyFunctionFlags( FUNC_Delegate ) )
 		{
-			FunctionName = FunctionName.LeftChop( FString( TEXT( "__DelegateSignature" ) ).Len() );
+			FunctionName = FunctionName.LeftChop( FString(HEADER_GENERATED_DELEGATE_SIGNATURE_SUFFIX).Len() );
 		}
 		UnMarshallAndCallName = FString(TEXT("exec")) + FunctionName;
 		
