Fix error when reading UHT manifest.

#rnx

diff --git a/Engine/Source/Programs/UnrealHeaderTool/Private/Manifest.cpp b/Engine/Source/Programs/UnrealHeaderTool/Private/Manifest.cpp
--- a/Engine/Source/Programs/UnrealHeaderTool/Private/Manifest.cpp
+++ b/Engine/Source/Programs/UnrealHeaderTool/Private/Manifest.cpp
@@ -84,7 +84,7 @@ FManifest FManifest::LoadFromFile(const FString& Filename)
 
 	GetJsonFieldValue(Result.IsGameTarget,             RootObject, TEXT("IsGameTarget"),             TEXT("{manifest root}"));
 	GetJsonFieldValue(Result.RootLocalPath,            RootObject, TEXT("RootLocalPath"),            TEXT("{manifest root}"));
-	GetJsonFieldValue(Result.RootBuildPath,            RootObject, TEXT("RootBuildPath"),            TEXT("{manifest root}"));
+	Result.RootBuildPath = Result.RootLocalPath + FPlatformMisc::GetDefaultPathSeparator();
 	GetJsonFieldValue(Result.TargetName,               RootObject, TEXT("TargetName"),               TEXT("{manifest root}"));
 	GetJsonFieldValue(Result.ExternalDependenciesFile, RootObject, TEXT("ExternalDependenciesFile"), TEXT("{manifest root}"));
 	GetJsonFieldValue(ModulesArray,                    RootObject, TEXT("Modules"),                  TEXT("{manifest root}"));
