Copy of CL 4385549 from Dev-Core

Try and detect a serialize function inside arbitrary preprocessor blocks and throw an error
Make a couple of other "can't have things inside preprocessor blocks" messages consistent with this one

diff --git a/Engine/Source/Programs/UnrealHeaderTool/Private/HeaderParser.cpp b/Engine/Source/Programs/UnrealHeaderTool/Private/HeaderParser.cpp
--- a/Engine/Source/Programs/UnrealHeaderTool/Private/HeaderParser.cpp
+++ b/Engine/Source/Programs/UnrealHeaderTool/Private/HeaderParser.cpp
@@ -5325,7 +5325,7 @@ bool FHeaderParser::CompileDeclaration(FClasses& AllClasses, TArray<UDelegateFun
 							}
 							else
 							{
-								FError::Throwf(TEXT("Serialize functions must be defined outside of all compiler define blocks, except for WITH_EDITORONLY_DATA"));
+								FError::Throwf(TEXT("Serialize functions must not be inside preprocessor blocks, except for WITH_EDITORONLY_DATA"));
 							}
 						}
 					}
@@ -8375,7 +8375,32 @@ void FHeaderParser::SimplifiedClassParse(const TCHAR* Filename, const TCHAR* InB
 							|| FindInitialStr(FoundSubstr, TrimmedStrLine, TEXT("UDELEGATE"))
 							|| FindInitialStr(FoundSubstr, TrimmedStrLine, TEXT("UFUNCTION")))
 						{
-							FFileLineException::Throwf(Filename, CurrentLine, TEXT("%s inside this preprocessor block will be skipped"), FoundSubstr);
+							FFileLineException::Throwf(Filename, CurrentLine, TEXT("%s must not be inside preprocessor blocks, except for WITH_EDITORONLY_DATA"), FoundSubstr);
+						}
+
+						// Try and determine if this line contains something like a serialize function
+						if (TrimmedStrLine.Len() > 0)
+						{
+							static const FString Str_Void = TEXT("void");
+							static const FString Str_Serialize = TEXT("Serialize(");
+							static const FString Str_FArchive = TEXT("FArchive");
+							static const FString Str_FStructuredArchive = TEXT("FStructuredArchive::FSlot");
+
+							int32 Pos = 0;
+							if ((Pos = TrimmedStrLine.Find(Str_Void, ESearchCase::CaseSensitive, ESearchDir::FromStart, Pos)) != -1)
+							{
+								Pos += Str_Void.Len();
+								if ((Pos = TrimmedStrLine.Find(Str_Serialize, ESearchCase::CaseSensitive, ESearchDir::FromStart, Pos)) != -1)
+								{
+									Pos += Str_Serialize.Len();
+
+									if (((TrimmedStrLine.Find(Str_FArchive, ESearchCase::CaseSensitive, ESearchDir::FromStart, Pos)) != -1) ||
+										((TrimmedStrLine.Find(Str_FStructuredArchive, ESearchCase::CaseSensitive, ESearchDir::FromStart, Pos)) != -1))
+									{
+										FFileLineException::Throwf(Filename, CurrentLine, TEXT("'%s' must not be inside preprocessor blocks, except for WITH_EDITORONLY_DATA"), *TrimmedStrLine);
+									}
+								}
+							}
 						}
 					}
 				}
