#UE4 Add a little more information to a UHT error

diff --git a/Engine/Source/Programs/UnrealHeaderTool/Private/CodeGenerator.cpp b/Engine/Source/Programs/UnrealHeaderTool/Private/CodeGenerator.cpp
--- a/Engine/Source/Programs/UnrealHeaderTool/Private/CodeGenerator.cpp
+++ b/Engine/Source/Programs/UnrealHeaderTool/Private/CodeGenerator.cpp
@@ -5682,7 +5682,7 @@ void ResolveSuperClasses(UPackage* Package)
 			if (FoundBaseClass == nullptr)
 			{
 				// Don't know its parent class. Raise error.
-				FError::Throwf(TEXT("Couldn't find parent type for '%s' named '%s' in current module or any other module parsed so far."), *DefinedClass->GetName(), *BaseClassName);
+				FError::Throwf(TEXT("Couldn't find parent type for '%s' named '%s' in current module (Package: %s) or any other module parsed so far."), *DefinedClass->GetName(), *BaseClassName, *GetNameSafe(Package));
 			}
 
 			DefinedClass->SetSuperStruct(FoundBaseClass);
