Copying //UE4/Dev-Anim to Dev-Main (//UE4/Dev-Main) [at] 6861656

#rnx

diff --git a/Engine/Plugins/Developer/AnimationSharing/Source/AnimationSharing/Private/AnimationSharingManager.cpp b/Engine/Plugins/Developer/AnimationSharing/Source/AnimationSharing/Private/AnimationSharingManager.cpp
--- a/Engine/Plugins/Developer/AnimationSharing/Source/AnimationSharing/Private/AnimationSharingManager.cpp
+++ b/Engine/Plugins/Developer/AnimationSharing/Source/AnimationSharing/Private/AnimationSharingManager.cpp
@@ -980,7 +980,7 @@ void UAnimSharingInstance::SetupState(FPerStateData& StateData, const FAnimation
 					Components.Add(Component);
 				}
 				else
-				{					
+				{	
 					const FVector SpawnLocation(FVector(NumSetups * SkeletalMeshBounds.X, 0.f, ZOffset));
 					const FName AdditiveComponentName(*(SkeletalMesh->GetName() + TEXT("_") + StateEnum->GetNameStringByIndex(StateEntry.State) + FString::FromInt(InstanceIndex)));
 					USkeletalMeshComponent* AdditiveComponent = NewObject<USkeletalMeshComponent>(SharingActor, AdditiveComponentName);
@@ -994,6 +994,10 @@ void UAnimSharingInstance::SetupState(FPerStateData& StateData, const FAnimation
 					FAdditiveAnimationInstance* AdditiveInstance = new FAdditiveAnimationInstance();
 					AdditiveInstance->Initialise(AdditiveComponent, SkeletonSetup.AdditiveAnimBlueprint.Get());
 					AdditiveInstanceStack.AddInstance(AdditiveInstance);
+					
+					/** Set the current animation length length */
+					StateData.AnimationLengths.Add(AnimSequence->SequenceLength);
+					Components.Add(AdditiveComponent);
 				}
 				
 				++NumSetups;
@@ -1309,6 +1313,9 @@ void UAnimSharingInstance::TickAdditiveInstances()
 	for (int32 InstanceIndex = 0; InstanceIndex < AdditiveInstances.Num(); ++InstanceIndex)
 	{
 		FAdditiveInstance& Instance = AdditiveInstances[InstanceIndex];
+		SetComponentUsage(true, Instance.State, Instance.UsedPerStateComponentIndex);		
+		SetComponentTick(Instance.State, Instance.UsedPerStateComponentIndex);
+
 		if (Instance.bActive)
 		{
 			const float WorldTimeSeconds = GetWorld()->GetTimeSeconds();
@@ -1326,6 +1333,8 @@ void UAnimSharingInstance::TickAdditiveInstances()
 				FreeAdditiveInstance(Instance.AdditiveAnimationInstance);
 				RemoveAdditiveInstance(InstanceIndex);
 				--InstanceIndex;
+
+				SetComponentUsage(false, Instance.State, Instance.UsedPerStateComponentIndex);
 			}
 		}
 		else
@@ -2025,6 +2034,7 @@ uint32 UAnimSharingInstance::SetupAdditiveInstance(uint8 StateIndex, uint8 FromS
 		const float WorldTimeSeconds = GetWorld()->GetTimeSeconds();
 		Instance.EndTime = WorldTimeSeconds + StateData.AdditiveAnimationSequence->SequenceLength;
 		Instance.State = StateIndex;
+		Instance.UsedPerStateComponentIndex = PerStateData[StateIndex].Components.IndexOfByKey(AnimationInstance->GetComponent());
 
 		InstanceIndex = AdditiveInstances.Num() - 1;
 		AnimationInstance->Setup(Instance.BaseComponent, StateData.AdditiveAnimationSequence);
diff --git a/Engine/Plugins/Developer/AnimationSharing/Source/AnimationSharing/Public/AnimationSharingManager.h b/Engine/Plugins/Developer/AnimationSharing/Source/AnimationSharing/Public/AnimationSharingManager.h
--- a/Engine/Plugins/Developer/AnimationSharing/Source/AnimationSharing/Public/AnimationSharingManager.h
+++ b/Engine/Plugins/Developer/AnimationSharing/Source/AnimationSharing/Public/AnimationSharingManager.h
@@ -112,6 +112,9 @@ struct FAdditiveInstance
 
 	/** Actor used for playing the additive animation */
 	FAdditiveAnimationInstance* AdditiveAnimationInstance;
+
+	/** Index into Components array for the current state data which is used for playing the animation*/
+	uint32 UsedPerStateComponentIndex;
 };
 
 /** Structure which holds data about a unique state which is linked to an enumeration value defined by the user. The data is populated from the user exposed FAnimationStateEntry */
diff --git a/Engine/Plugins/Experimental/AlembicImporter/Source/AlembicLibrary/Private/AbcFile.cpp b/Engine/Plugins/Experimental/AlembicImporter/Source/AlembicLibrary/Private/AbcFile.cpp
--- a/Engine/Plugins/Experimental/AlembicImporter/Source/AlembicLibrary/Private/AbcFile.cpp
+++ b/Engine/Plugins/Experimental/AlembicImporter/Source/AlembicLibrary/Private/AbcFile.cpp
@@ -126,12 +126,12 @@ EAbcImportError FAbcFile::Import(UAbcImportSettings* InImportSettings)
 
 	if (ImportSettings->ImportType == EAlembicImportType::StaticMesh)
 	{
-		EndFrameIndex = StartFrameIndex + 1;
+		EndFrameIndex = StartFrameIndex;
 	}
 
 	int32 FrameSpan = EndFrameIndex - StartFrameIndex;
 	// If Start==End or Start > End output error message due to invalid frame span
-	if (FrameSpan <= 0)
+	if (FrameSpan <= 0 && ImportSettings->ImportType != EAlembicImportType::StaticMesh)
 	{
 		TSharedRef<FTokenizedMessage> Message = FTokenizedMessage::Create(EMessageSeverity::Error, FText::Format(LOCTEXT("NoFramesForMeshObject", "Invalid frame range specified {0} - {1}."), FText::FromString(FString::FromInt(StartFrameIndex)), FText::FromString(FString::FromInt(EndFrameIndex))));
 		FAbcImportLogger::AddImportMessage(Message);
diff --git a/Engine/Plugins/Experimental/ControlRig/Source/ControlRig/Private/Units/Deprecated/RigUnit_Example.cpp b/Engine/Plugins/Experimental/ControlRig/Source/ControlRig/Private/Units/Deprecated/RigUnit_Example.cpp
deleted file mode 100644
diff --git a/Engine/Plugins/Experimental/ControlRig/Source/ControlRig/Private/Units/Deprecated/RigUnit_Example.h b/Engine/Plugins/Experimental/ControlRig/Source/ControlRig/Private/Units/Deprecated/RigUnit_Example.h
deleted file mode 100644
diff --git a/Engine/Plugins/Experimental/ControlRig/Source/ControlRig/Private/Units/Deprecated/RigUnit_TwoBoneIKFK.cpp b/Engine/Plugins/Experimental/ControlRig/Source/ControlRig/Private/Units/Deprecated/RigUnit_TwoBoneIKFK.cpp
--- a/Engine/Plugins/Experimental/ControlRig/Source/ControlRig/Private/Units/Deprecated/RigUnit_TwoBoneIKFK.cpp
+++ b/Engine/Plugins/Experimental/ControlRig/Source/ControlRig/Private/Units/Deprecated/RigUnit_TwoBoneIKFK.cpp
@@ -118,18 +118,7 @@ void FRigUnit_TwoBoneIKFK::Execute(const FRigUnitContext& Context)
 
 void FRigUnit_TwoBoneIKFK::SolveIK()
 {
-	if (bUsePoleTarget)
-	{
-		AnimationCore::SolveTwoBoneIK(StartJointIKTransform, MidJointIKTransform, EndJointIKTransform, PoleTarget, EndEffector.GetLocation(), UpperLimbLength, LowerLimbLength, false, 1.0f, 1.05f);
-	}
-	else
-	{
-		// default IK
-		// @todo: we don't have that yet, so default to using pole vector
-		ensure(false);
-		AnimationCore::SolveTwoBoneIK(StartJointIKTransform, MidJointIKTransform, EndJointIKTransform, PoleTarget, EndEffector.GetLocation(), UpperLimbLength, LowerLimbLength, false, 1.0f, 1.05f);
-	}
-
+	AnimationCore::SolveTwoBoneIK(StartJointIKTransform, MidJointIKTransform, EndJointIKTransform, PoleTarget, EndEffector.GetLocation(), UpperLimbLength, LowerLimbLength, false, 1.0f, 1.05f);
 	// set end effector rotation to current rotation
 	EndJointIKTransform.SetRotation(EndEffector.GetRotation());
 }
diff --git a/Engine/Plugins/Experimental/ControlRig/Source/ControlRig/Private/Units/Highlevel/Hierarchy/RigUnit_ModifyBoneTransforms.h b/Engine/Plugins/Experimental/ControlRig/Source/ControlRig/Private/Units/Highlevel/Hierarchy/RigUnit_ModifyBoneTransforms.h
--- a/Engine/Plugins/Experimental/ControlRig/Source/ControlRig/Private/Units/Highlevel/Hierarchy/RigUnit_ModifyBoneTransforms.h
+++ b/Engine/Plugins/Experimental/ControlRig/Source/ControlRig/Private/Units/Highlevel/Hierarchy/RigUnit_ModifyBoneTransforms.h
@@ -81,19 +81,19 @@ struct FRigUnit_ModifyBoneTransforms : public FRigUnit_HighlevelBaseMutable
 	/**
 	 * At 1 this sets the transform, between 0 and 1 the transform is blended with previous results.
 	 */
-	UPROPERTY(meta = (Input))
+	UPROPERTY(meta = (Input, ClampMin=0.f, ClampMax=1.f, UIMin = 0.f, UIMax = 1.f))
 	float Weight;
 
 	/**
 	 * The minimum of the weight - defaults to 0.0
 	 */
-	UPROPERTY(meta = (Input, Constant))
+	UPROPERTY(meta = (Input, Constant, ClampMin = 0.f, ClampMax = 1.f, UIMin = 0.f, UIMax = 1.f))
 	float WeightMinimum;
 
 	/**
 	 * The maximum of the weight - defaults to 1.0
 	 */
-	UPROPERTY(meta = (Input, Constant))
+	UPROPERTY(meta = (Input, Constant, ClampMin = 0.f, ClampMax = 1.f, UIMin = 0.f, UIMax = 1.f))
 	float WeightMaximum;
 
 	/**
diff --git a/Engine/Plugins/Experimental/ControlRig/Source/ControlRig/Public/Units/Deprecated/RigUnit_TwoBoneIKFK.h b/Engine/Plugins/Experimental/ControlRig/Source/ControlRig/Public/Units/Deprecated/RigUnit_TwoBoneIKFK.h
--- a/Engine/Plugins/Experimental/ControlRig/Source/ControlRig/Public/Units/Deprecated/RigUnit_TwoBoneIKFK.h
+++ b/Engine/Plugins/Experimental/ControlRig/Source/ControlRig/Public/Units/Deprecated/RigUnit_TwoBoneIKFK.h
@@ -19,8 +19,7 @@ struct CONTROLRIG_API FRigUnit_TwoBoneIKFK : public FRigUnitMutable
 	virtual void Execute(const FRigUnitContext& Context) override;
 
 	FRigUnit_TwoBoneIKFK()
-		: bUsePoleTarget(true)
-		, PoleTarget(ForceInitToZero)
+		: PoleTarget(ForceInitToZero)
 		, Spin(0.f)
 		, IKBlend(1.f)
 		, PreviousFKIKBlend(1.f)
@@ -37,10 +36,6 @@ struct CONTROLRIG_API FRigUnit_TwoBoneIKFK : public FRigUnitMutable
 	UPROPERTY(EditAnywhere, Category = "TwoBoneIKFK", meta = (Input))
 	FName EndJoint;
 
-	// # Whether or not to use the pole target matrix.If false, use a minimum energy solution
-	UPROPERTY(EditAnywhere, Category = "TwoBoneIKFK", meta = (Input))
-	bool bUsePoleTarget;
-
 	// # Transform to use as the pole target(specifies the plane of solution)
 	UPROPERTY(EditAnywhere, Category = "TwoBoneIKFK", meta = (Input))
 	FVector PoleTarget;
diff --git a/Engine/Plugins/Experimental/GeometryCache/Source/GeometryCache/Classes/GeometryCacheCodecV1.h b/Engine/Plugins/Experimental/GeometryCache/Source/GeometryCache/Classes/GeometryCacheCodecV1.h
--- a/Engine/Plugins/Experimental/GeometryCache/Source/GeometryCache/Classes/GeometryCacheCodecV1.h
+++ b/Engine/Plugins/Experimental/GeometryCache/Source/GeometryCache/Classes/GeometryCacheCodecV1.h
@@ -18,7 +18,6 @@ struct FGeometryCacheCodecRenderStateV1 : FGeometryCacheCodecRenderStateBase
 {
 	FGeometryCacheCodecRenderStateV1(const TArray<int32> &SetTopologyRanges); 
 	~FGeometryCacheCodecRenderStateV1();	
-	ICodecDecoder* Decoder;
 	
 	virtual bool DecodeSingleFrame(FGeometryCacheCodecDecodeArguments &Args) override;
 };
diff --git a/Engine/Plugins/Experimental/GeometryCache/Source/GeometryCache/Classes/GeometryCacheComponent.h b/Engine/Plugins/Experimental/GeometryCache/Source/GeometryCache/Classes/GeometryCacheComponent.h
--- a/Engine/Plugins/Experimental/GeometryCache/Source/GeometryCache/Classes/GeometryCacheComponent.h
+++ b/Engine/Plugins/Experimental/GeometryCache/Source/GeometryCache/Classes/GeometryCacheComponent.h
@@ -179,6 +179,8 @@ public:
 	void SetManualTick(bool bInManualTick);
 	bool GetManualTick() const;
 
+	void ResetAnimationTime();
+
 	UFUNCTION(BlueprintCallable, Category = "Components|GeometryCache")
 	void TickAtThisTime(const float Time, bool bInIsRunning, bool bInBackwards, bool bInIsLooping);
 
@@ -234,7 +236,7 @@ protected:
 	UPROPERTY(VisibleAnywhere, Category = GeometryCache)
 	int32 NumTracks;
 
-	UPROPERTY(VisibleAnywhere, Category = GeometryCache)
+	UPROPERTY(VisibleAnywhere, transient, Category = GeometryCache)
 	float ElapsedTime;
 
 	/** Component local space bounds of geometry cache mesh */
diff --git a/Engine/Plugins/Experimental/GeometryCache/Source/GeometryCache/Private/CodecV1.cpp b/Engine/Plugins/Experimental/GeometryCache/Source/GeometryCache/Private/CodecV1.cpp
--- a/Engine/Plugins/Experimental/GeometryCache/Source/GeometryCache/Private/CodecV1.cpp
+++ b/Engine/Plugins/Experimental/GeometryCache/Source/GeometryCache/Private/CodecV1.cpp
@@ -832,18 +832,6 @@ void FCodecV1Encoder::WriteTables()
 
 #endif // WITH_EDITOR
 
-FCodecV1Decoder::FCodecV1Decoder()
-{
-	// Precalculate table mapping symbol index to non-raw bits. ((Sign ? -2 : 1) << NumRawBits)
-	for (int32 NumRawBits = 1; NumRawBits <= 30; NumRawBits++)
-	{
-		for (int32 Sign = 0; Sign <= 1; Sign++)
-		{
-			HighBitsLUT[2 + Sign + NumRawBits * 2] = (Sign ? -2 : 1) << NumRawBits;
-		}
-	}	
-}
-
 void FCodecV1Decoder::DecodeIndexStream(FHuffmanBitStreamReader& Reader, uint32* Stream, uint64 ElementOffset, uint32 ElementCount)
 {
 	const uint8* RawElementData = (const uint8*)Stream;
@@ -1045,6 +1033,19 @@ void FCodecV1Decoder::ReadCodedStreamDescription(FHuffmanBitStreamReader& Reader
 	VertexInfo.bConstantIndices = (ReadBits(Reader, 1) == 1);
 }
 
+int32 FCodecV1Decoder::CachedHighBitsLUT[64];
+void FCodecV1Decoder::InitLUT()
+{
+	// Precalculate table mapping symbol index to non-raw bits. ((Sign ? -2 : 1) << NumRawBits)
+	for (int32 NumRawBits = 1; NumRawBits <= 30; NumRawBits++)
+	{
+		for (int32 Sign = 0; Sign <= 1; Sign++)
+		{
+			CachedHighBitsLUT[2 + Sign + NumRawBits * 2] = (Sign ? -2 : 1) << NumRawBits;
+		}
+	}
+}
+
 DECLARE_CYCLE_STAT(TEXT("FCodecV1Decoder"), STAT_CodecV1Decoder, STATGROUP_GeometryCache);
 DECLARE_CYCLE_STAT(TEXT("SetupAndReadTables"), STAT_SetupAndReadTables, STATGROUP_GeometryCache);
 DECLARE_CYCLE_STAT(TEXT("DecodeIndexStream"), STAT_DecodeIndexStream, STATGROUP_GeometryCache);
@@ -1278,6 +1279,6 @@ int32 FCodecV1Decoder::ReadInt32(FHuffmanBitStreamReader& Reader, FHuffmanDecode
 	{
 		// At least one raw bit.
 		int32 NumRawBits = (Packed - 2) >> 1;
-		return ReadBitsNoRefill(Reader, NumRawBits) + HighBitsLUT[Packed];
+		return ReadBitsNoRefill(Reader, NumRawBits) + CachedHighBitsLUT[Packed];
 	}
 }
\ No newline at end of file
diff --git a/Engine/Plugins/Experimental/GeometryCache/Source/GeometryCache/Private/CodecV1.h b/Engine/Plugins/Experimental/GeometryCache/Source/GeometryCache/Private/CodecV1.h
--- a/Engine/Plugins/Experimental/GeometryCache/Source/GeometryCache/Private/CodecV1.h
+++ b/Engine/Plugins/Experimental/GeometryCache/Source/GeometryCache/Private/CodecV1.h
@@ -287,7 +287,9 @@ private:
 class FCodecV1Decoder : public ICodecDecoder
 {
 public:
-	FCodecV1Decoder();
+	FCodecV1Decoder() {}
+
+	static void InitLUT();
 
 	/**
 	* Read a frame's bit stream and decode the frame
@@ -344,6 +346,6 @@ private:
 	/** Any context information to decode a frame in a sequence of frames, such as any Huffman tables used. */
 	FCodecV1DecodingContext DecodingContext;	
 
-	int32 HighBitsLUT[64];
+	static int32 CachedHighBitsLUT[64];
 };
 
diff --git a/Engine/Plugins/Experimental/GeometryCache/Source/GeometryCache/Private/GeometryCacheCodecV1.cpp b/Engine/Plugins/Experimental/GeometryCache/Source/GeometryCache/Private/GeometryCacheCodecV1.cpp
--- a/Engine/Plugins/Experimental/GeometryCache/Source/GeometryCache/Private/GeometryCacheCodecV1.cpp
+++ b/Engine/Plugins/Experimental/GeometryCache/Source/GeometryCache/Private/GeometryCacheCodecV1.cpp
@@ -123,13 +123,10 @@ void UGeometryCacheCodecV1::CodeFrame(const FGeometryCacheCodecEncodeArguments &
 
 FGeometryCacheCodecRenderStateV1::FGeometryCacheCodecRenderStateV1(const TArray<int32> &SetTopologyRanges) : FGeometryCacheCodecRenderStateBase(SetTopologyRanges)
 {
-	Decoder = new FCodecV1Decoder();
 }
 
 FGeometryCacheCodecRenderStateV1::~FGeometryCacheCodecRenderStateV1()
 {	
-	delete Decoder;
-	Decoder = nullptr;
 }
 
 DECLARE_CYCLE_STAT(TEXT("Deserialize Mesh"), STAT_DeserializeMeshV1, STATGROUP_GeometryCache);
@@ -153,7 +150,8 @@ bool FGeometryCacheCodecRenderStateV1::DecodeSingleFrame(FGeometryCacheCodecDeco
 
 	{
 		SCOPE_CYCLE_COUNTER(STAT_DeserializeMeshV1);
-		Decoder->DecodeFrameData(Ar, Args.OutMeshData);
+		FCodecV1Decoder Decoder;
+		Decoder.DecodeFrameData(Ar, Args.OutMeshData);
 	}
 	IGeometryCacheStreamingManager::Get().UnmapChunk(Args.Track.GetTrack(), Args.FrameIdentifier);
 	return true;
diff --git a/Engine/Plugins/Experimental/GeometryCache/Source/GeometryCache/Private/GeometryCacheComponent.cpp b/Engine/Plugins/Experimental/GeometryCache/Source/GeometryCache/Private/GeometryCacheComponent.cpp
--- a/Engine/Plugins/Experimental/GeometryCache/Source/GeometryCache/Private/GeometryCacheComponent.cpp
+++ b/Engine/Plugins/Experimental/GeometryCache/Source/GeometryCache/Private/GeometryCacheComponent.cpp
@@ -148,6 +148,11 @@ bool UGeometryCacheComponent::GetManualTick() const
 	return bManualTick;
 }
 
+void UGeometryCacheComponent::ResetAnimationTime()
+{
+	ElapsedTime = 0.0f;
+}
+
 void UGeometryCacheComponent::TickAtThisTime(const float Time, bool bInIsRunning, bool bInBackwards, bool bInIsLooping)
 {
 	if (bManualTick && GeometryCache && bRunning)
diff --git a/Engine/Plugins/Experimental/GeometryCache/Source/GeometryCache/Private/GeometryCacheModule.cpp b/Engine/Plugins/Experimental/GeometryCache/Source/GeometryCache/Private/GeometryCacheModule.cpp
--- a/Engine/Plugins/Experimental/GeometryCache/Source/GeometryCache/Private/GeometryCacheModule.cpp
+++ b/Engine/Plugins/Experimental/GeometryCache/Source/GeometryCache/Private/GeometryCacheModule.cpp
@@ -4,6 +4,7 @@
 #if WITH_EDITOR
 #include "GeometryCacheEdModule.h"
 #endif // WITH_EDITOR
+#include "CodecV1.h"
 
 IMPLEMENT_MODULE(FGeometryCacheModule, GeometryCache)
 
@@ -12,6 +13,8 @@ void FGeometryCacheModule::StartupModule()
 #if WITH_EDITOR
 	FGeometryCacheEdModule& Module = FModuleManager::LoadModuleChecked<FGeometryCacheEdModule>(TEXT("GeometryCacheEd"));
 #endif
+
+	FCodecV1Decoder::InitLUT();
 }
 
 void FGeometryCacheModule::ShutdownModule()
diff --git a/Engine/Plugins/Experimental/GeometryCache/Source/GeometryCache/Private/GeometryCacheSceneProxy.cpp b/Engine/Plugins/Experimental/GeometryCache/Source/GeometryCache/Private/GeometryCacheSceneProxy.cpp
--- a/Engine/Plugins/Experimental/GeometryCache/Source/GeometryCache/Private/GeometryCacheSceneProxy.cpp
+++ b/Engine/Plugins/Experimental/GeometryCache/Source/GeometryCache/Private/GeometryCacheSceneProxy.cpp
@@ -54,6 +54,7 @@ FGeometryCacheSceneProxy::FGeometryCacheSceneProxy(UGeometryCacheComponent* Comp
 	bLooping = Component->IsLooping();
 	bAlwaysHasVelocity = true;
 	PlaybackSpeed = (Component->IsPlaying()) ? Component->GetPlaybackSpeed() : 0.0f;
+	UpdatedFrameNum = 0;
 
 	// Copy each section
 	const int32 NumTracks = Component->TrackSections.Num();
@@ -324,17 +325,18 @@ void FGeometryCacheSceneProxy::CreateMeshBatch(
 				&& (TrackProxy->NextFrameMeshData->Positions.Num() == TrackProxy->NextFrameMeshData->MotionVectors.Num());
 
 				if (!bHasMotionVectors)
-				{
-					UserData.MotionBlurDataExtension = FVector::OneVector;
-					UserData.MotionBlurDataOrigin = FVector::ZeroVector;
-					UserData.MotionBlurPositionScale = 0.0f;
-				}
-				else
-				{
-					UserData.MotionBlurDataExtension = FVector::OneVector * PlaybackSpeed;
+	            {
+		            const float PreviousPositionScale = (GFrameNumber <= UpdatedFrameNum) ? 1.f : 0.f;
+					UserData.MotionBlurDataExtension = FVector::OneVector * PreviousPositionScale;
 					UserData.MotionBlurDataOrigin = FVector::ZeroVector;
-					UserData.MotionBlurPositionScale = 1.0f;
-				}
+					UserData.MotionBlurPositionScale = 1.f - PreviousPositionScale;
+	            }
+	            else
+	            {
+		            UserData.MotionBlurDataExtension = FVector::OneVector * PlaybackSpeed;
+		            UserData.MotionBlurDataOrigin = FVector::ZeroVector;
+		            UserData.MotionBlurPositionScale = 1.0f;
+	            }
 
 				if (IsRayTracingEnabled())
 				{
@@ -539,6 +541,7 @@ void FGeometryCacheSceneProxy::UpdateAnimation(float NewTime, bool bNewLooping,
 	bLooping = bNewLooping;
 	bIsPlayingBackwards = bNewIsPlayingBackwards;
 	PlaybackSpeed = NewPlaybackSpeed;
+	UpdatedFrameNum = GFrameNumber + 1;
 
 	if (IsRayTracingEnabled())
 	{
@@ -1017,7 +1020,6 @@ void FGeomCacheVertexFactory::Init(const FVertexBuffer* PositionBuffer, const FV
 			{
 			Init_RenderThread(PositionBuffer, MotionBlurDataBuffer, TangentXBuffer, TangentZBuffer, TextureCoordinateBuffer, ColorBuffer);
 		});
-		FlushRenderingCommands();
 	}
 }
 
diff --git a/Engine/Plugins/Experimental/GeometryCache/Source/GeometryCache/Private/GeometryCacheSceneProxy.h b/Engine/Plugins/Experimental/GeometryCache/Source/GeometryCache/Private/GeometryCacheSceneProxy.h
--- a/Engine/Plugins/Experimental/GeometryCache/Source/GeometryCache/Private/GeometryCacheSceneProxy.h
+++ b/Engine/Plugins/Experimental/GeometryCache/Source/GeometryCache/Private/GeometryCacheSceneProxy.h
@@ -217,6 +217,7 @@ private:
 	bool bLooping;
 	bool bIsPlayingBackwards;
 	float PlaybackSpeed;
+	uint32 UpdatedFrameNum;
 
 	/** Array of Track Proxies */
 	TArray<FGeomCacheTrackProxy*> Tracks;
diff --git a/Engine/Plugins/Experimental/GeometryCache/Source/GeometryCache/Private/StreamingGeometryCacheData.cpp b/Engine/Plugins/Experimental/GeometryCache/Source/GeometryCache/Private/StreamingGeometryCacheData.cpp
--- a/Engine/Plugins/Experimental/GeometryCache/Source/GeometryCache/Private/StreamingGeometryCacheData.cpp
+++ b/Engine/Plugins/Experimental/GeometryCache/Source/GeometryCache/Private/StreamingGeometryCacheData.cpp
@@ -329,42 +329,45 @@ void FStreamingGeometryCacheData::ProcessCompletedChunks()
 	//Note: This function should only be called from code which owns the CriticalSection
 	check(IsInGameThread() || IsInRenderingThread());
 
-	FCompletedChunk CompletedCunk;
-	while (CompletedChunks.Dequeue(CompletedCunk))
+	FCompletedChunk CompletedChunk;
+	while (CompletedChunks.Dequeue(CompletedChunk))
 	{
-		FResidentChunk* Chunk = Chunks.Find(CompletedCunk.LoadedChunkIndex);
+		FResidentChunk* Chunk = Chunks.Find(CompletedChunk.LoadedChunkIndex);
 		if (!Chunk)
 		{
 			UE_LOG(LogGeoCaStreaming, Error, TEXT("Got a stray async read request"));
 			return;
 		}
 
-		check(CompletedCunk.ReadRequest == Chunk->IORequest);
-
-		// Check to see if we successfully managed to load anything
-		uint8* Mem = CompletedCunk.ReadRequest->GetReadResults();
-		if (Mem)
+		// Chunks can be queued up multiple times when scrubbing, either the request has already been fulfilled (nullptr) or it could be waiting for a different IO request
+		if (CompletedChunk.ReadRequest == Chunk->IORequest || Chunk->IORequest != nullptr)
 		{
-			Chunk->Memory = Mem;
-			ChunksAvailable.Add(CompletedCunk.LoadedChunkIndex);
-			ChunksRequested.Remove(CompletedCunk.LoadedChunkIndex);
-			DEC_DWORD_STAT(STAT_OutstandingRequests);
-			INC_MEMORY_STAT_BY(STAT_ChunkDataResident, Chunk->DataSize);
-			INC_MEMORY_STAT_BY(STAT_ChunkDataStreamed, Chunk->DataSize);
-			IGeometryCacheStreamingManager::Get().IoBandwidth.Add(Chunk->DataSize);
-		}
-		else
-		{
-			UE_LOG(LogGeoCaStreaming, Error, TEXT("Async loading request failed!"));
-			ChunksRequested.Remove(CompletedCunk.LoadedChunkIndex);
-			// Fix me do we want to recover from this? Granite simply reschedules requests
-			// as it may have failed for transient reasons (buffer contention, ...)
+			// Check to see if we successfully managed to load anything
+			uint8* Mem = CompletedChunk.ReadRequest->GetReadResults();
+			if (Mem)
+			{
+				Chunk->Memory = Mem;
+				ChunksAvailable.Add(CompletedChunk.LoadedChunkIndex);
+				ChunksRequested.Remove(CompletedChunk.LoadedChunkIndex);
+				DEC_DWORD_STAT(STAT_OutstandingRequests);
+				INC_MEMORY_STAT_BY(STAT_ChunkDataResident, Chunk->DataSize);
+				INC_MEMORY_STAT_BY(STAT_ChunkDataStreamed, Chunk->DataSize);
+				IGeometryCacheStreamingManager::Get().IoBandwidth.Add(Chunk->DataSize);
+			}
+			else
+			{
+				UE_LOG(LogGeoCaStreaming, Error, TEXT("Async loading request failed!"));
+				ChunksRequested.Remove(CompletedChunk.LoadedChunkIndex);
+				// Fix me do we want to recover from this? Granite simply reschedules requests
+				// as it may have failed for transient reasons (buffer contention, ...)
+			}
+
+			Chunk->IORequest = nullptr;
 		}
 
 		// Clean up the now fully processed IO request
-		check(Chunk->IORequest->PollCompletion());
-		delete Chunk->IORequest;
-		Chunk->IORequest = nullptr;
+		check(CompletedChunk.ReadRequest->PollCompletion());
+		delete CompletedChunk.ReadRequest;
 	}
 }
 
@@ -403,7 +406,6 @@ const uint8* FStreamingGeometryCacheData::MapChunk(uint32 ChunkIndex, uint32* Ou
 			*OutChunkSize = ResidentChunk->DataSize;
 		}
 		ResidentChunk->Refcount++;
-		check(ResidentChunk->Refcount == 1);
 		return (uint8*)ResidentChunk->Memory;
 	}
 }
@@ -419,7 +421,6 @@ void FStreamingGeometryCacheData::UnmapChunk(uint32 ChunkIndex)
 		checkf(ResidentChunk->Refcount > 0, TEXT("Map/Unmap out of balance. Make sure you unmap once fore every map."));
 		checkf(ChunksAvailable.Contains(ChunkIndex) || ChunksEvicted.Contains(ChunkIndex), TEXT("Tried to unmap a chunk in an invalid state."));
 		ResidentChunk->Refcount--;
-		check(ResidentChunk->Refcount == 0);
 	}
 	else
 	{
diff --git a/Engine/Plugins/Experimental/GeometryCache/Source/GeometryCacheTracks/Private/MovieSceneGeometryCacheTemplate.cpp b/Engine/Plugins/Experimental/GeometryCache/Source/GeometryCacheTracks/Private/MovieSceneGeometryCacheTemplate.cpp
--- a/Engine/Plugins/Experimental/GeometryCache/Source/GeometryCacheTracks/Private/MovieSceneGeometryCacheTemplate.cpp
+++ b/Engine/Plugins/Experimental/GeometryCache/Source/GeometryCacheTracks/Private/MovieSceneGeometryCacheTemplate.cpp
@@ -29,6 +29,7 @@ struct FPreAnimatedGeometryCacheTokenProducer : IMovieScenePreAnimatedTokenProdu
 			{
 				UGeometryCacheComponent* Component = CastChecked<UGeometryCacheComponent>(&ObjectToRestore);
 				Component->SetManualTick(bInManualTick);
+				Component->ResetAnimationTime();
 			}
 			bool bInManualTick;
 		};
diff --git a/Engine/Source/Developer/AssetTools/Private/AssetTypeActions/AssetTypeActions_SkeletalMesh.cpp b/Engine/Source/Developer/AssetTools/Private/AssetTypeActions/AssetTypeActions_SkeletalMesh.cpp
--- a/Engine/Source/Developer/AssetTools/Private/AssetTypeActions/AssetTypeActions_SkeletalMesh.cpp
+++ b/Engine/Source/Developer/AssetTools/Private/AssetTypeActions/AssetTypeActions_SkeletalMesh.cpp
@@ -588,15 +588,16 @@ EVisibility FAssetTypeActions_SkeletalMesh::GetThumbnailSkinningOverlayVisibilit
 
 	//The object is loaded we can use the memory value of the object to set the overlay
 	UObject* Obj = AssetData.GetAsset();
-	USkeletalMesh* SkeletalMesh = CastChecked<USkeletalMesh>(Obj);
-
-	UAssetImportData* GenericImportData = SkeletalMesh->AssetImportData;
-	if (GenericImportData != nullptr)
+	if(USkeletalMesh* SkeletalMesh = Cast<USkeletalMesh>(Obj))
 	{
-		UFbxSkeletalMeshImportData* ImportData = Cast<UFbxSkeletalMeshImportData>(GenericImportData);
-		if (ImportData != nullptr && ImportData->LastImportContentType == EFBXImportContentType::FBXICT_Geometry)
+		UAssetImportData* GenericImportData = SkeletalMesh->AssetImportData;
+		if (GenericImportData != nullptr)
 		{
-			return EVisibility::HitTestInvisible;
+			UFbxSkeletalMeshImportData* ImportData = Cast<UFbxSkeletalMeshImportData>(GenericImportData);
+			if (ImportData != nullptr && ImportData->LastImportContentType == EFBXImportContentType::FBXICT_Geometry)
+			{
+				return EVisibility::HitTestInvisible;
+			}
 		}
 	}
 	return EVisibility::Collapsed;
diff --git a/Engine/Source/Developer/MaterialBaking/Private/MaterialBakingHelpers.cpp b/Engine/Source/Developer/MaterialBaking/Private/MaterialBakingHelpers.cpp
--- a/Engine/Source/Developer/MaterialBaking/Private/MaterialBakingHelpers.cpp
+++ b/Engine/Source/Developer/MaterialBaking/Private/MaterialBakingHelpers.cpp
@@ -164,7 +164,7 @@ void FMaterialBakingHelpers::PerformUVBorderSmear(TArray<FColor>& InOutPixels, i
 	TArray<uint32> RowRemainingPixels;
 	RowRemainingPixels.SetNumZeroed(PaddedImageHeight);
 
-	const int32 MaxThreads = FPlatformProcess::SupportsMultithreading() ? FPlatformMisc::NumberOfCores() : 1;
+	const int32 MaxThreads = FPlatformProcess::SupportsMultithreading() ? FPlatformMisc::NumberOfCoresIncludingHyperthreads() : 1;
 
 	//
 	// Iteratively smear until all rows are filled.
diff --git a/Engine/Source/Developer/MeshUtilities/Private/MeshUtilities.cpp b/Engine/Source/Developer/MeshUtilities/Private/MeshUtilities.cpp
--- a/Engine/Source/Developer/MeshUtilities/Private/MeshUtilities.cpp
+++ b/Engine/Source/Developer/MeshUtilities/Private/MeshUtilities.cpp
@@ -5711,6 +5711,8 @@ void FMeshUtilities::ShutdownModule()
 		FPropertyEditorModule& PropertyEditorModule = FModuleManager::Get().GetModuleChecked<FPropertyEditorModule>(PropertyEditorModuleName);
 
 		PropertyEditorModule.UnregisterCustomClassLayout("MeshSimplificationSettings");
+		PropertyEditorModule.UnregisterCustomClassLayout("SkeletalMeshSimplificationSettings");
+		PropertyEditorModule.UnregisterCustomClassLayout("ProxyLODMeshSimplificationSettings");
 	}
 
 	RemoveLevelViewportMenuExtender();
diff --git a/Engine/Source/Editor/AnimGraph/Classes/AnimGraphNode_Base.h b/Engine/Source/Editor/AnimGraph/Classes/AnimGraphNode_Base.h
--- a/Engine/Source/Editor/AnimGraph/Classes/AnimGraphNode_Base.h
+++ b/Engine/Source/Editor/AnimGraph/Classes/AnimGraphNode_Base.h
@@ -143,6 +143,7 @@ class ANIMGRAPH_API UAnimGraphNode_Base : public UK2Node
 	virtual void GetPinHoverText(const UEdGraphPin& Pin, FString& HoverTextOut) const override;
 	virtual bool ShowPaletteIconOnNode() const override{ return false; }
 	virtual void PinDefaultValueChanged(UEdGraphPin* Pin) override;
+	virtual FString GetPinMetaData(FName InPinName, FName InKey) override;
 	// End of UEdGraphNode interface
 
 	// UK2Node interface
diff --git a/Engine/Source/Editor/AnimGraph/Private/AnimGraphNode_Base.cpp b/Engine/Source/Editor/AnimGraph/Private/AnimGraphNode_Base.cpp
--- a/Engine/Source/Editor/AnimGraph/Private/AnimGraphNode_Base.cpp
+++ b/Engine/Source/Editor/AnimGraph/Private/AnimGraphNode_Base.cpp
@@ -410,6 +410,27 @@ void UAnimGraphNode_Base::PinDefaultValueChanged(UEdGraphPin* Pin)
 	}
 }
 
+FString UAnimGraphNode_Base::GetPinMetaData(FName InPinName, FName InKey)
+{
+	FString MetaData = Super::GetPinMetaData(InPinName, InKey);
+	if(MetaData.IsEmpty())
+	{
+		// Check properties of our anim node
+		if(UStructProperty* NodeStructProperty = GetFNodeProperty())
+		{
+			for (TFieldIterator<UProperty> It(NodeStructProperty->Struct); It; ++It)
+			{
+				const UProperty* Property = *It;
+				if (Property && Property->GetFName() == InPinName)
+				{
+					return Property->GetMetaData(InKey);
+				}
+			}
+		}
+	}
+	return MetaData;
+}
+
 bool UAnimGraphNode_Base::IsPinExposedAndLinked(const FString& InPinName, const EEdGraphPinDirection InDirection) const
 {
 	UEdGraphPin* Pin = FindPin(InPinName, InDirection);
diff --git a/Engine/Source/Editor/AnimGraph/Private/AnimGraphNode_SequencePlayer.cpp b/Engine/Source/Editor/AnimGraph/Private/AnimGraphNode_SequencePlayer.cpp
--- a/Engine/Source/Editor/AnimGraph/Private/AnimGraphNode_SequencePlayer.cpp
+++ b/Engine/Source/Editor/AnimGraph/Private/AnimGraphNode_SequencePlayer.cpp
@@ -197,10 +197,13 @@ void UAnimGraphNode_SequencePlayer::GetMenuActions(FBlueprintActionDatabaseRegis
 			UBlueprintNodeSpawner* NodeSpawner = UBlueprintNodeSpawner::Create(GetClass());
 			if (Asset.IsAssetLoaded())
 			{
-				TWeakObjectPtr<UAnimSequence> AnimSequence = Cast<UAnimSequence>(Asset.GetAsset());
-				NodeSpawner->CustomizeNodeDelegate = UBlueprintNodeSpawner::FCustomizeNodeDelegate::CreateStatic(LoadedAssetSetup, AnimSequence);
-				NodeSpawner->DefaultMenuSignature.MenuName = GetTitleGivenAssetInfo(FText::FromName(AnimSequence->GetFName()), AnimSequence->IsValidAdditive());
-				NodeSpawner->DefaultMenuSignature.Tooltip = GetTitleGivenAssetInfo(FText::FromString(AnimSequence->GetPathName()), AnimSequence->IsValidAdditive());
+				UAnimSequence* AnimSequence = Cast<UAnimSequence>(Asset.GetAsset());
+				if(AnimSequence)
+				{
+					NodeSpawner->CustomizeNodeDelegate = UBlueprintNodeSpawner::FCustomizeNodeDelegate::CreateStatic(LoadedAssetSetup, TWeakObjectPtr<UAnimSequence>(AnimSequence));
+					NodeSpawner->DefaultMenuSignature.MenuName = GetTitleGivenAssetInfo(FText::FromName(AnimSequence->GetFName()), AnimSequence->IsValidAdditive());
+					NodeSpawner->DefaultMenuSignature.Tooltip = GetTitleGivenAssetInfo(FText::FromString(AnimSequence->GetPathName()), AnimSequence->IsValidAdditive());
+				}
 			}
 			else
 			{
@@ -326,6 +329,10 @@ void UAnimGraphNode_SequencePlayer::ValidateAnimNodeDuringCompilation(class USke
 			MessageLog.Error(TEXT("@@ references an unknown sequence"), this);
 		}
 	}
+	else if(SupportsAssetClass(SequenceToCheck->GetClass()) == EAnimAssetHandlerType::NotSupported)
+	{
+		MessageLog.Error(*FText::Format(LOCTEXT("UnsupportedAssetError", "@@ is trying to play a {0} as a sequence, which is not allowed."), SequenceToCheck->GetClass()->GetDisplayNameText()).ToString(), this);
+	}
 	else
 	{
 		USkeleton* SeqSkeleton = SequenceToCheck->GetSkeleton();
diff --git a/Engine/Source/Editor/AnimationModifiers/Private/AnimationBlueprintLibrary.cpp b/Engine/Source/Editor/AnimationModifiers/Private/AnimationBlueprintLibrary.cpp
--- a/Engine/Source/Editor/AnimationModifiers/Private/AnimationBlueprintLibrary.cpp
+++ b/Engine/Source/Editor/AnimationModifiers/Private/AnimationBlueprintLibrary.cpp
@@ -489,7 +489,7 @@ void UAnimationBlueprintLibrary::GetAnimationNotifyEventNames(const UAnimSequenc
 	}	
 }
 
-UAnimNotify* UAnimationBlueprintLibrary::AddAnimationNotifyEvent(UAnimSequence* AnimationSequence, FName NotifyTrackName, float StartTime, float Duration, TSubclassOf<UObject> NotifyClass)
+UAnimNotify* UAnimationBlueprintLibrary::AddAnimationNotifyEvent(UAnimSequence* AnimationSequence, FName NotifyTrackName, float StartTime, TSubclassOf<UAnimNotify> NotifyClass)
 {
 	UAnimNotify* Notify = nullptr;
 	if (AnimationSequence)
@@ -499,19 +499,76 @@ UAnimNotify* UAnimationBlueprintLibrary::AddAnimationNotifyEvent(UAnimSequence*
 
 		if (bIsValidTrackName && bIsValidTime)
 		{
-			AnimationSequence->Notifies.AddDefaulted();
-			FAnimNotifyEvent& NewEvent = AnimationSequence->Notifies.Last();
+			FAnimNotifyEvent& NewEvent = AnimationSequence->Notifies.AddDefaulted_GetRef();
 
 			NewEvent.NotifyName = NAME_None;
 			NewEvent.Link(AnimationSequence, StartTime);
 			NewEvent.TriggerTimeOffset = GetTriggerTimeOffsetForType(AnimationSequence->CalculateOffsetForNotify(StartTime));
 			NewEvent.TrackIndex = GetTrackIndexForAnimationNotifyTrackName(AnimationSequence, NotifyTrackName);
+			NewEvent.NotifyStateClass = nullptr;
 
 			if (NotifyClass)
 			{
-				class UObject* AnimNotifyClass = NewObject<UObject>(AnimationSequence, NotifyClass, NAME_None, RF_Transactional);
-				NewEvent.NotifyStateClass = Cast<UAnimNotifyState>(AnimNotifyClass);
-				NewEvent.Notify = Cast<UAnimNotify>(AnimNotifyClass);
+				Notify = NewObject<UAnimNotify>(AnimationSequence, NotifyClass, NAME_None, RF_Transactional);
+				NewEvent.Notify = Notify;
+
+				// Setup name for new event
+				if(NewEvent.Notify)
+				{
+					NewEvent.NotifyName = FName(*NewEvent.Notify->GetNotifyName());
+				}
+			}
+			else
+			{
+				NewEvent.Notify = nullptr;
+			}
+
+			// Refresh all cached data
+			AnimationSequence->RefreshCacheData();
+		}
+		else
+		{
+			if (!bIsValidTrackName)
+			{
+				UE_LOG(LogAnimationBlueprintLibrary, Warning, TEXT("Animation Notify Track %s does not exist on Animation Sequence %s"), *NotifyTrackName.ToString(), *AnimationSequence->GetName());
+			}
+
+			if (!bIsValidTime)
+			{
+				UE_LOG(LogAnimationBlueprintLibrary, Warning, TEXT("%f is side of Animation Sequence %s range"), StartTime, *AnimationSequence->GetName());
+			}
+		}
+	}
+	else
+	{
+		UE_LOG(LogAnimationBlueprintLibrary, Warning, TEXT("Invalid Animation Sequence for AddAnimationNotifyEvent"));
+	}
+
+	return Notify;
+}
+
+UAnimNotifyState* UAnimationBlueprintLibrary::AddAnimationNotifyStateEvent(UAnimSequence* AnimationSequence, FName NotifyTrackName, float StartTime, float Duration, TSubclassOf<UAnimNotifyState> NotifyStateClass)
+{
+	UAnimNotifyState* NotifyState = nullptr;
+	if (AnimationSequence)
+	{
+		const bool bIsValidTrackName = IsValidAnimNotifyTrackName(AnimationSequence, NotifyTrackName);
+		const bool bIsValidTime = IsValidTimeInternal(AnimationSequence, StartTime);
+
+		if (bIsValidTrackName && bIsValidTime)
+		{
+			FAnimNotifyEvent& NewEvent = AnimationSequence->Notifies.AddDefaulted_GetRef();
+
+			NewEvent.NotifyName = NAME_None;
+			NewEvent.Link(AnimationSequence, StartTime);
+			NewEvent.TriggerTimeOffset = GetTriggerTimeOffsetForType(AnimationSequence->CalculateOffsetForNotify(StartTime));
+			NewEvent.TrackIndex = GetTrackIndexForAnimationNotifyTrackName(AnimationSequence, NotifyTrackName);
+			NewEvent.Notify = nullptr;
+
+			if (NotifyStateClass)
+			{
+				NotifyState = NewObject<UAnimNotifyState>(AnimationSequence, NotifyStateClass, NAME_None, RF_Transactional);
+				NewEvent.NotifyStateClass = NotifyState;
 
 				// Setup name and duration for new event
 				if (NewEvent.NotifyStateClass)
@@ -520,21 +577,14 @@ UAnimNotify* UAnimationBlueprintLibrary::AddAnimationNotifyEvent(UAnimSequence*
 					NewEvent.SetDuration(Duration);
 					NewEvent.EndLink.Link(AnimationSequence, NewEvent.EndLink.GetTime());
 				}
-				else if(NewEvent.Notify)
-				{
-					NewEvent.NotifyName = FName(*NewEvent.Notify->GetNotifyName());
-				}
 			}
 			else
 			{
-				NewEvent.Notify = NULL;
-				NewEvent.NotifyStateClass = NULL;
+				NewEvent.NotifyStateClass = nullptr;
 			}
 
 			// Refresh all cached data
 			AnimationSequence->RefreshCacheData();			
-
-			Notify = NewEvent.Notify;
 		}
 		else
 		{
@@ -551,10 +601,10 @@ UAnimNotify* UAnimationBlueprintLibrary::AddAnimationNotifyEvent(UAnimSequence*
 	}
 	else
 	{
-		UE_LOG(LogAnimationBlueprintLibrary, Warning, TEXT("Invalid Animation Sequence for AddAnimationNotifyEvent"));
+		UE_LOG(LogAnimationBlueprintLibrary, Warning, TEXT("Invalid Animation Sequence for AddAnimationNotifyStateEvent"));
 	}
 
-	return Notify;
+	return NotifyState;
 }
 
 void UAnimationBlueprintLibrary::AddAnimationNotifyEventObject(UAnimSequence* AnimationSequence, float StartTime, UAnimNotify* Notify, FName NotifyTrackName)
@@ -568,14 +618,13 @@ void UAnimationBlueprintLibrary::AddAnimationNotifyEventObject(UAnimSequence* An
 
 		if (bValidNotify && bValidOuter && bIsValidTrackName && bIsValidTime)
 		{
-			AnimationSequence->Notifies.AddDefaulted();
-			FAnimNotifyEvent& NewEvent = AnimationSequence->Notifies.Last();
+			FAnimNotifyEvent& NewEvent = AnimationSequence->Notifies.AddDefaulted_GetRef();
 
 			NewEvent.NotifyName = NAME_None;
 			NewEvent.Link(AnimationSequence, StartTime);
 			NewEvent.TriggerTimeOffset = GetTriggerTimeOffsetForType(AnimationSequence->CalculateOffsetForNotify(StartTime));
 			NewEvent.TrackIndex = GetTrackIndexForAnimationNotifyTrackName(AnimationSequence, NotifyTrackName);
-			NewEvent.NotifyStateClass = Cast<UAnimNotifyState>(Notify);
+			NewEvent.NotifyStateClass = nullptr;
 			NewEvent.Notify = Notify;
 
 			// Refresh all cached data
@@ -609,7 +658,61 @@ void UAnimationBlueprintLibrary::AddAnimationNotifyEventObject(UAnimSequence* An
 	{
 		UE_LOG(LogAnimationBlueprintLibrary, Warning, TEXT("Invalid Animation Sequence for AddAnimationNotifyEventObject"));
 	}
-	
+}
+
+void UAnimationBlueprintLibrary::AddAnimationNotifyStateEventObject(UAnimSequence* AnimationSequence, float StartTime, float Duration, UAnimNotifyState* NotifyState, FName NotifyTrackName)
+{
+	if (AnimationSequence)
+	{
+		const bool bValidNotifyState = NotifyState != nullptr;
+		const bool bValidOuter = bValidNotifyState && NotifyState->GetOuter() == AnimationSequence;
+		const bool bIsValidTrackName = IsValidAnimNotifyTrackName(AnimationSequence, NotifyTrackName);
+		const bool bIsValidTime = IsValidTimeInternal(AnimationSequence, StartTime);
+
+		if (bValidNotifyState && bValidOuter && bIsValidTrackName && bIsValidTime)
+		{
+			FAnimNotifyEvent& NewEvent = AnimationSequence->Notifies.AddDefaulted_GetRef();
+
+			NewEvent.NotifyName = NAME_None;
+			NewEvent.Link(AnimationSequence, StartTime);
+			NewEvent.TriggerTimeOffset = GetTriggerTimeOffsetForType(AnimationSequence->CalculateOffsetForNotify(StartTime));
+			NewEvent.TrackIndex = GetTrackIndexForAnimationNotifyTrackName(AnimationSequence, NotifyTrackName);
+			NewEvent.NotifyStateClass = NotifyState;
+			NewEvent.Notify = nullptr;
+			NewEvent.SetDuration(Duration);
+			NewEvent.EndLink.Link(AnimationSequence, NewEvent.EndLink.GetTime());
+
+			// Refresh all cached data
+			AnimationSequence->RefreshCacheData();
+		}
+		else
+		{
+			if (!bValidNotifyState)
+			{
+				UE_LOG(LogAnimationBlueprintLibrary, Warning, TEXT("Invalid Animation Notify in AddAnimationNotifyEventObject"));
+			}
+
+			if (!bValidOuter)
+			{
+				const FString NotifyName = bValidNotifyState ? NotifyState->GetName() : "Invalid Notify";
+				UE_LOG(LogAnimationBlueprintLibrary, Warning, TEXT("Animation Notify %s Outer is not %s"), *NotifyName, *AnimationSequence->GetName());
+			}
+
+			if (!bIsValidTrackName)
+			{
+				UE_LOG(LogAnimationBlueprintLibrary, Warning, TEXT("Animation Notify Track %s does not exist on Animation Sequence %s"), *NotifyTrackName.ToString(), *AnimationSequence->GetName());
+			}
+
+			if (!bIsValidTime)
+			{
+				UE_LOG(LogAnimationBlueprintLibrary, Warning, TEXT("%f is side of Animation Sequence %s range"), StartTime, *AnimationSequence->GetName());
+			}
+		}
+	}
+	else
+	{
+		UE_LOG(LogAnimationBlueprintLibrary, Warning, TEXT("Invalid Animation Sequence for AddAnimationNotifyEventObject"));
+	}
 }
 
 static void ReplaceAnimNotifies_Helper(UAnimSequenceBase* AnimationSequence, UClass* OldNotifyClass, UClass* NewNotifyClass, FOnNotifyReplaced OnNotifyReplaced, FOnNotifyStateReplaced OnNotifyStateReplaced)
diff --git a/Engine/Source/Editor/AnimationModifiers/Private/AnimationModifier.cpp b/Engine/Source/Editor/AnimationModifiers/Private/AnimationModifier.cpp
--- a/Engine/Source/Editor/AnimationModifiers/Private/AnimationModifier.cpp
+++ b/Engine/Source/Editor/AnimationModifiers/Private/AnimationModifier.cpp
@@ -73,9 +73,10 @@ void UAnimationModifier::ApplyToAnimationSequence(class UAnimSequence* InAnimati
 
 
 		CurrentAnimSequence->PostEditChange();
-		CurrentSkeleton->PostEditChange();		
+		CurrentSkeleton->PostEditChange();
 		CurrentAnimSequence->RefreshCacheData();
 		CurrentAnimSequence->RefreshCurveData();
+		CurrentAnimSequence->MarkRawDataAsModified();
 
 		UpdateStoredRevisions();
 	}
@@ -126,6 +127,7 @@ void UAnimationModifier::RevertFromAnimationSequence(class UAnimSequence* InAnim
 	CurrentSkeleton->PostEditChange();
 	CurrentAnimSequence->RefreshCacheData();
 	CurrentAnimSequence->RefreshCurveData();
+	CurrentAnimSequence->MarkRawDataAsModified();
 
 	ResetStoredRevisions();
 
diff --git a/Engine/Source/Editor/AnimationModifiers/Public/AnimationBlueprintLibrary.h b/Engine/Source/Editor/AnimationModifiers/Public/AnimationBlueprintLibrary.h
--- a/Engine/Source/Editor/AnimationModifiers/Public/AnimationBlueprintLibrary.h
+++ b/Engine/Source/Editor/AnimationModifiers/Public/AnimationBlueprintLibrary.h
@@ -167,12 +167,20 @@ public:
 
 	/** Adds an Animation Notify Event to Notify track in the given Animation with the given Notify creation data */
 	UFUNCTION(BlueprintCallable, Category = "AnimationBlueprintLibrary|NotifyEvents")
-	static UAnimNotify* AddAnimationNotifyEvent(UAnimSequence* AnimationSequence, FName NotifyTrackName, float StartTime, float Duration, TSubclassOf<UObject> NotifyClass);
+	static UAnimNotify* AddAnimationNotifyEvent(UAnimSequence* AnimationSequence, FName NotifyTrackName, float StartTime, TSubclassOf<UAnimNotify> NotifyClass);
+
+	/** Adds an Animation Notify State Event to Notify track in the given Animation with the given Notify State creation data */
+	UFUNCTION(BlueprintCallable, Category = "AnimationBlueprintLibrary|NotifyEvents")
+	static UAnimNotifyState* AddAnimationNotifyStateEvent(UAnimSequence* AnimationSequence, FName NotifyTrackName, float StartTime, float Duration, TSubclassOf<UAnimNotifyState> NotifyStateClass);
 
 	/** Adds an the specific Animation Notify to the Animation Sequence (requires Notify's outer to be the Animation Sequence) */
 	UFUNCTION(BlueprintCallable, Category = "AnimationBlueprintLibrary|NotifyEvents")
 	static void AddAnimationNotifyEventObject(UAnimSequence* AnimationSequence, float StartTime, UAnimNotify* Notify, FName NotifyTrackName);
 
+	/** Adds an the specific Animation Notify State to the Animation Sequence (requires Notify State's outer to be the Animation Sequence) */
+	UFUNCTION(BlueprintCallable, Category = "AnimationBlueprintLibrary|NotifyEvents")
+	static void AddAnimationNotifyStateEventObject(UAnimSequence* AnimationSequence, float StartTime, float Duration, UAnimNotifyState* NotifyState, FName NotifyTrackName);
+
 	/** Removes Animation Notify Events found by Name within the Animation Sequence, and returns the number of removed name instances */
 	UFUNCTION(BlueprintCallable, Category = "AnimationBlueprintLibrary|NotifyEvents")
 	static int32 RemoveAnimationNotifyEventsByName(UAnimSequence* AnimationSequence, FName NotifyName);
diff --git a/Engine/Source/Editor/BlueprintGraph/Private/BlueprintActionFilter.cpp b/Engine/Source/Editor/BlueprintGraph/Private/BlueprintActionFilter.cpp
--- a/Engine/Source/Editor/BlueprintGraph/Private/BlueprintActionFilter.cpp
+++ b/Engine/Source/Editor/BlueprintGraph/Private/BlueprintActionFilter.cpp
@@ -383,6 +383,15 @@ namespace BlueprintActionFilterImpl
 	 */
 	static bool IsIncompatibleAnimNotification( FBlueprintActionFilter const& Filter, FBlueprintActionInfo& BlueprintAction); 
 
+	/**
+	 * Rejection test that checks if an anim notify event is in an anim blueprint
+	 * 
+	 * @param  Filter			Holds the graph context for this test.
+	 * @param  BlueprintAction	The action you wish to query.
+	 * @return True if the action is an anim notify and the current blueprint is not an anim blueprint
+	 */
+	static bool IsAnimNotificationInAnimBlueprint( FBlueprintActionFilter const& Filter, FBlueprintActionInfo& BlueprintAction); 
+
 	/**
 	 * 
 	 * 
@@ -1563,6 +1572,22 @@ static bool BlueprintActionFilterImpl::IsNodeTemplateSelfFiltered(FBlueprintActi
 	return BlueprintAction.NodeSpawner->IsTemplateNodeFilteredOut(Filter);
 }
 
+//------------------------------------------------------------------------------
+static bool BlueprintActionFilterImpl::IsAnimNotificationInAnimBlueprint(FBlueprintActionFilter const& Filter, FBlueprintActionInfo& BlueprintAction)
+{
+	bool bIsFilteredOut = false;
+
+	if(!BlueprintAction.GetNodeClass()->IsChildOf<UK2Node_Event>())
+	{
+		if (const UAnimBlueprint* AnimBlueprintOwner = Cast<UAnimBlueprint>(BlueprintAction.GetActionOwner()))
+		{
+			bIsFilteredOut = !(Filter.Context.Blueprints.Num() == 1 && Filter.Context.Blueprints[0] == AnimBlueprintOwner);
+		}
+	}
+
+	return bIsFilteredOut;
+}
+
 //------------------------------------------------------------------------------
 static bool BlueprintActionFilterImpl::IsIncompatibleAnimNotification(FBlueprintActionFilter const& Filter, FBlueprintActionInfo& BlueprintAction)
 {
@@ -1904,6 +1929,7 @@ FBlueprintActionFilter::FBlueprintActionFilter(uint32 Flags/*= 0x00*/)
 	//
 	// this test in-particular spawns a template-node and then calls 
 	// AllocateDefaultPins() which is costly, so it should be very last!
+	AddRejectionTest(FRejectionTestDelegate::CreateStatic(IsAnimNotificationInAnimBlueprint));
 	AddRejectionTest(FRejectionTestDelegate::CreateStatic(IsIncompatibleAnimNotification));
 	AddRejectionTest(FRejectionTestDelegate::CreateStatic(IsNodeTemplateSelfFiltered));
 	AddRejectionTest(FRejectionTestDelegate::CreateStatic(IsMissingMatchingPinParam));
diff --git a/Engine/Source/Editor/DetailCustomizations/Private/PerPlatformPropertyCustomization.cpp b/Engine/Source/Editor/DetailCustomizations/Private/PerPlatformPropertyCustomization.cpp
--- a/Engine/Source/Editor/DetailCustomizations/Private/PerPlatformPropertyCustomization.cpp
+++ b/Engine/Source/Editor/DetailCustomizations/Private/PerPlatformPropertyCustomization.cpp
@@ -14,6 +14,7 @@
 #include "PropertyHandle.h"
 #include "DetailLayoutBuilder.h"
 #include "SPerPlatformPropertiesWidget.h"
+#include "PlatformInfo.h"
 #include "ScopedTransaction.h"
 #include "IPropertyUtilities.h"
 #include "UObject/MetaData.h"
diff --git a/Engine/Source/Editor/GraphEditor/Private/KismetPins/SGraphPinObject.cpp b/Engine/Source/Editor/GraphEditor/Private/KismetPins/SGraphPinObject.cpp
--- a/Engine/Source/Editor/GraphEditor/Private/KismetPins/SGraphPinObject.cpp
+++ b/Engine/Source/Editor/GraphEditor/Private/KismetPins/SGraphPinObject.cpp
@@ -210,19 +210,30 @@ TSharedRef<SWidget> SGraphPinObject::GenerateAssetPicker()
 	AssetPickerConfig.InitialAssetViewType = EAssetViewType::List;
 	AssetPickerConfig.bAllowDragging = false;
 
-	// Check with the node to see if there is any "AllowClasses" metadata for the pin
-	FString ClassFilterString = GraphPinObj->GetOwningNode()->GetPinMetaData(GraphPinObj->PinName, FName(TEXT("AllowedClasses")));
-	if( !ClassFilterString.IsEmpty() )
+	// Check with the node to see if there is any "AllowClasses" or "DisallowedClasses" metadata for the pin
+	FString AllowedClassesFilterString = GraphPinObj->GetOwningNode()->GetPinMetaData(GraphPinObj->PinName, FName(TEXT("AllowedClasses")));
+	if( !AllowedClassesFilterString.IsEmpty() )
 	{
 		// Clear out the allowed class names and have the pin's metadata override.
 		AssetPickerConfig.Filter.ClassNames.Empty();
 
 		// Parse and add the classes from the metadata
-		TArray<FString> CustomClassFilterNames;
-		ClassFilterString.ParseIntoArray(CustomClassFilterNames, TEXT(","), true);
-		for(auto It = CustomClassFilterNames.CreateConstIterator(); It; ++It)
+		TArray<FString> AllowedClassesFilterNames;
+		AllowedClassesFilterString.ParseIntoArray(AllowedClassesFilterNames, TEXT(","), true);
+		for(const FString& AllowedClassesFilterName : AllowedClassesFilterNames)
 		{
-			AssetPickerConfig.Filter.ClassNames.Add(FName(**It));
+			AssetPickerConfig.Filter.ClassNames.Add(FName(*AllowedClassesFilterName));
+		}
+	}
+
+	FString DisallowedClassesFilterString = GraphPinObj->GetOwningNode()->GetPinMetaData(GraphPinObj->PinName, FName(TEXT("DisallowedClasses")));
+	if(!DisallowedClassesFilterString.IsEmpty())
+	{
+		TArray<FString> DisallowedClassesFilterNames;
+		DisallowedClassesFilterString.ParseIntoArray(DisallowedClassesFilterNames, TEXT(","), true);
+		for(const FString& DisallowedClassesFilterName : DisallowedClassesFilterNames)
+		{
+			AssetPickerConfig.Filter.RecursiveClassesExclusionSet.Add(FName(*DisallowedClassesFilterName));
 		}
 	}
 
diff --git a/Engine/Source/Editor/HierarchicalLODOutliner/Private/HLODOutliner.cpp b/Engine/Source/Editor/HierarchicalLODOutliner/Private/HLODOutliner.cpp
--- a/Engine/Source/Editor/HierarchicalLODOutliner/Private/HLODOutliner.cpp
+++ b/Engine/Source/Editor/HierarchicalLODOutliner/Private/HLODOutliner.cpp
@@ -49,6 +49,7 @@
 #include "Widgets/Notifications/SNotificationList.h"
 #include "Toolkits/AssetEditorManager.h"
 #include "Engine/HLODProxy.h"
+#include "Misc/ScopedSlowTask.h"
 
 #define LOCTEXT_NAMESPACE "HLODOutliner"
 
@@ -1098,6 +1099,9 @@ namespace HLODOutliner
 			// This call came from a context menu
 			auto SelectedItems = TreeView->GetSelectedItems();
 			
+			FScopedSlowTask SlowTask(SelectedItems.Num(), FText::Format(LOCTEXT("SHLODOutliner_BuildProxy", "Building Proxy {0}|plural(one=Mesh,other=Meshes)"), SelectedItems.Num()));
+			SlowTask.MakeDialog();
+
 			// Loop over all selected items (context menu can't be called with multiple items selected that aren't of the same types)
 			for (auto SelectedItem : SelectedItems )
 			{
@@ -1116,6 +1120,8 @@ namespace HLODOutliner
 						}
 					}
 				}
+
+				SlowTask.EnterProgressFrame();
 			}
 			
 			SetForcedLODLevel(ForcedLODLevel);
@@ -1135,6 +1141,9 @@ namespace HLODOutliner
 			// This call came from a context menu
 			auto SelectedItems = TreeView->GetSelectedItems();
 
+			FScopedSlowTask SlowTask(SelectedItems.Num(), FText::Format(LOCTEXT("SHLODOutliner_RebuildProxy", "Rebuild Proxy {0}|plural(one=Mesh,other=Meshes)"), SelectedItems.Num()));
+			SlowTask.MakeDialog();
+
 			// Loop over all selected items (context menu can't be called with multiple items selected that aren't of the same types)
 			for (auto SelectedItem : SelectedItems)
 			{
@@ -1150,6 +1159,8 @@ namespace HLODOutliner
 						CurrentWorld->HierarchicalLODBuilder->BuildMeshForLODActor(ActorItem->LODActor.Get(), LevelItem->LODLevelIndex);
 					}
 				}
+
+				SlowTask.EnterProgressFrame();
 			}
 
 			SetForcedLODLevel(ForcedLODLevel);
diff --git a/Engine/Source/Editor/MergeActors/Private/MergeActorsModule.cpp b/Engine/Source/Editor/MergeActors/Private/MergeActorsModule.cpp
--- a/Engine/Source/Editor/MergeActors/Private/MergeActorsModule.cpp
+++ b/Engine/Source/Editor/MergeActors/Private/MergeActorsModule.cpp
@@ -134,8 +134,10 @@ void FMergeActorsModule::ShutdownModule()
 {
 	if (FSlateApplication::IsInitialized())
 	{
-		FGlobalTabmanager::Get()->UnregisterNomadTabSpawner(MergeActorsTabName);
+		FGlobalTabmanager::Get()->UnregisterNomadTabSpawner(MergeActorsTabName);		
 	}
+
+	MergeActorsTools.Empty();
 }
 
 
diff --git a/Engine/Source/Editor/MergeActors/Private/MeshInstancingTool/MeshInstancingTool.cpp b/Engine/Source/Editor/MergeActors/Private/MeshInstancingTool/MeshInstancingTool.cpp
--- a/Engine/Source/Editor/MergeActors/Private/MeshInstancingTool/MeshInstancingTool.cpp
+++ b/Engine/Source/Editor/MergeActors/Private/MeshInstancingTool/MeshInstancingTool.cpp
@@ -21,14 +21,22 @@
 #include "ScopedTransaction.h"
 #include "MeshMergeModule.h"
 
-
 #define LOCTEXT_NAMESPACE "MeshInstancingTool"
 
+bool UMeshInstancingSettingsObject::bInitialized = false;
+UMeshInstancingSettingsObject* UMeshInstancingSettingsObject::DefaultSettings = nullptr;
+
 FMeshInstancingTool::FMeshInstancingTool()
 {
 	SettingsObject = UMeshInstancingSettingsObject::Get();
 }
 
+FMeshInstancingTool::~FMeshInstancingTool()
+{
+	UMeshInstancingSettingsObject::Destroy();
+	SettingsObject = nullptr;
+}
+
 TSharedRef<SWidget> FMeshInstancingTool::GetWidget()
 {
 	SAssignNew(InstancingDialog, SMeshInstancingDialog, this);
@@ -155,5 +163,4 @@ bool FMeshInstancingTool::CanMerge() const
 {	
 	return InstancingDialog->GetNumSelectedMeshComponents() >= 1;
 }
-
-#undef LOCTEXT_NAMESPACE
+#undef LOCTEXT_NAMESPACE
\ No newline at end of file
diff --git a/Engine/Source/Editor/MergeActors/Private/MeshInstancingTool/MeshInstancingTool.h b/Engine/Source/Editor/MergeActors/Private/MeshInstancingTool/MeshInstancingTool.h
--- a/Engine/Source/Editor/MergeActors/Private/MeshInstancingTool/MeshInstancingTool.h
+++ b/Engine/Source/Editor/MergeActors/Private/MeshInstancingTool/MeshInstancingTool.h
@@ -24,9 +24,36 @@ public:
 
 	static UMeshInstancingSettingsObject* Get()
 	{
-		return GetMutableDefault<UMeshInstancingSettingsObject>();
+		if (!bInitialized)
+		{
+			DefaultSettings = DuplicateObject(GetMutableDefault<UMeshInstancingSettingsObject>(), nullptr);
+			DefaultSettings->AddToRoot();
+			bInitialized = true;
+		}
+
+		return DefaultSettings;
+	}
+
+	static void Destroy()
+	{
+		if (bInitialized)
+		{
+			if (DefaultSettings)
+			{
+				DefaultSettings->RemoveFromRoot();
+				DefaultSettings->MarkPendingKill();
+			}
+
+			DefaultSettings = nullptr;
+			bInitialized = false;
+		}
 	}
 
+protected:
+	static bool bInitialized;
+	// This is a singleton, duplicate default object
+	static UMeshInstancingSettingsObject* DefaultSettings;
+
 public:
 	UPROPERTY(EditAnywhere, meta = (ShowOnlyInnerProperties), Category = MergeSettings)
 	FMeshInstancingSettings Settings;
@@ -42,6 +69,7 @@ class FMeshInstancingTool : public IMergeActorsTool
 public:
 
 	FMeshInstancingTool();
+	~FMeshInstancingTool();
 
 	// IMergeActorsTool interface
 	virtual TSharedRef<SWidget> GetWidget() override;
diff --git a/Engine/Source/Editor/MergeActors/Private/MeshInstancingTool/SMeshInstancingDialog.cpp b/Engine/Source/Editor/MergeActors/Private/MeshInstancingTool/SMeshInstancingDialog.cpp
--- a/Engine/Source/Editor/MergeActors/Private/MeshInstancingTool/SMeshInstancingDialog.cpp
+++ b/Engine/Source/Editor/MergeActors/Private/MeshInstancingTool/SMeshInstancingDialog.cpp
@@ -340,6 +340,7 @@ void SMeshInstancingDialog::CreateSettingsView()
 	DetailsViewArgs.bLockable = true;
 	DetailsViewArgs.NameAreaSettings = FDetailsViewArgs::ComponentsAndActorsUseNameArea;
 	DetailsViewArgs.bCustomNameAreaLocation = false;
+	DetailsViewArgs.NameAreaSettings = FDetailsViewArgs::ENameAreaSettings::HideNameArea;
 	DetailsViewArgs.bCustomFilterAreaLocation = true;
 	DetailsViewArgs.DefaultsOnlyVisibility = EEditDefaultsOnlyNodeVisibility::Hide;
 	
diff --git a/Engine/Source/Editor/MergeActors/Private/MeshMergingTool/MeshMergingTool.cpp b/Engine/Source/Editor/MergeActors/Private/MeshMergingTool/MeshMergingTool.cpp
--- a/Engine/Source/Editor/MergeActors/Private/MeshMergingTool/MeshMergingTool.cpp
+++ b/Engine/Source/Editor/MergeActors/Private/MeshMergingTool/MeshMergingTool.cpp
@@ -24,12 +24,21 @@
 
 #define LOCTEXT_NAMESPACE "MeshMergingTool"
 
+UMeshMergingSettingsObject* UMeshMergingSettingsObject::DefaultSettings = nullptr;
+bool UMeshMergingSettingsObject::bInitialized = false;
+
 FMeshMergingTool::FMeshMergingTool()
 	: bReplaceSourceActors(false)
 {
 	SettingsObject = UMeshMergingSettingsObject::Get();
 }
 
+FMeshMergingTool::~FMeshMergingTool()
+{
+	UMeshMergingSettingsObject::Destroy();
+	SettingsObject = nullptr;
+}
+
 TSharedRef<SWidget> FMeshMergingTool::GetWidget()
 {
 	SAssignNew(MergingDialog, SMeshMergingDialog, this);
@@ -172,3 +181,6 @@ bool FMeshMergingTool::CanMerge() const
 }
 
 #undef LOCTEXT_NAMESPACE
+
+
+
diff --git a/Engine/Source/Editor/MergeActors/Private/MeshMergingTool/MeshMergingTool.h b/Engine/Source/Editor/MergeActors/Private/MeshMergingTool/MeshMergingTool.h
--- a/Engine/Source/Editor/MergeActors/Private/MeshMergingTool/MeshMergingTool.h
+++ b/Engine/Source/Editor/MergeActors/Private/MeshMergingTool/MeshMergingTool.h
@@ -25,20 +25,40 @@ public:
 		Settings.LODSelectionType = EMeshLODSelectionType::AllLODs;
 	}
 
-	static UMeshMergingSettingsObject* Get()
-	{
-		static bool bInitialized = false;
-		// This is a singleton, use default object
-		UMeshMergingSettingsObject* DefaultSettings = GetMutableDefault<UMeshMergingSettingsObject>();
+	
 
+	static UMeshMergingSettingsObject* Get()
+	{		
+		// This is a singleton, duplicate default object
 		if (!bInitialized)
 		{
+			DefaultSettings = DuplicateObject(GetMutableDefault<UMeshMergingSettingsObject>(), nullptr);
+			DefaultSettings->AddToRoot();			
 			bInitialized = true;
 		}
 
 		return DefaultSettings;
 	}
 
+	static void Destroy()
+	{
+		if (bInitialized)
+		{
+			if (DefaultSettings)
+			{
+				DefaultSettings->RemoveFromRoot();
+				DefaultSettings->MarkPendingKill();
+			}
+			
+			DefaultSettings = nullptr;
+			bInitialized = false;
+		}
+	}
+
+protected:
+	static UMeshMergingSettingsObject* DefaultSettings;
+	static bool bInitialized;
+
 public:
 	UPROPERTY(editanywhere, meta = (ShowOnlyInnerProperties), Category = MergeSettings)
 	FMeshMergingSettings Settings;
@@ -54,6 +74,7 @@ class FMeshMergingTool : public IMergeActorsTool
 public:
 
 	FMeshMergingTool();
+	~FMeshMergingTool();
 
 	// IMergeActorsTool interface
 	virtual TSharedRef<SWidget> GetWidget() override;
diff --git a/Engine/Source/Editor/MergeActors/Private/MeshMergingTool/SMeshMergingDialog.cpp b/Engine/Source/Editor/MergeActors/Private/MeshMergingTool/SMeshMergingDialog.cpp
--- a/Engine/Source/Editor/MergeActors/Private/MeshMergingTool/SMeshMergingDialog.cpp
+++ b/Engine/Source/Editor/MergeActors/Private/MeshMergingTool/SMeshMergingDialog.cpp
@@ -210,6 +210,7 @@ void SMeshMergingDialog::CreateSettingsView()
 	DetailsViewArgs.bLockable = true;
 	DetailsViewArgs.NameAreaSettings = FDetailsViewArgs::ComponentsAndActorsUseNameArea;
 	DetailsViewArgs.bCustomNameAreaLocation = false;
+	DetailsViewArgs.NameAreaSettings = FDetailsViewArgs::ENameAreaSettings::HideNameArea;
 	DetailsViewArgs.bCustomFilterAreaLocation = true;
 	DetailsViewArgs.DefaultsOnlyVisibility = EEditDefaultsOnlyNodeVisibility::Hide;
 	
diff --git a/Engine/Source/Editor/MergeActors/Private/MeshProxyTool/MeshProxyTool.cpp b/Engine/Source/Editor/MergeActors/Private/MeshProxyTool/MeshProxyTool.cpp
--- a/Engine/Source/Editor/MergeActors/Private/MeshProxyTool/MeshProxyTool.cpp
+++ b/Engine/Source/Editor/MergeActors/Private/MeshProxyTool/MeshProxyTool.cpp
@@ -23,12 +23,20 @@
 
 #define LOCTEXT_NAMESPACE "MeshProxyTool"
 
+bool UMeshProxySettingsObject::bInitialized = false;
+UMeshProxySettingsObject* UMeshProxySettingsObject::DefaultSettings = nullptr;
 
 FMeshProxyTool::FMeshProxyTool()
 {
 	SettingsObject = UMeshProxySettingsObject::Get();
 }
 
+FMeshProxyTool::~FMeshProxyTool()
+{
+	UMeshProxySettingsObject::Destroy();
+	SettingsObject = nullptr;
+}
+
 TSharedRef<SWidget> FMeshProxyTool::GetWidget()
 {
 	SAssignNew(ProxyDialog, SMeshProxyDialog, this);
@@ -278,3 +286,4 @@ TSharedRef<SWidget> FThirdPartyMeshProxyTool::GetWidget()
 }
 
 #undef LOCTEXT_NAMESPACE
+
diff --git a/Engine/Source/Editor/MergeActors/Private/MeshProxyTool/MeshProxyTool.h b/Engine/Source/Editor/MergeActors/Private/MeshProxyTool/MeshProxyTool.h
--- a/Engine/Source/Editor/MergeActors/Private/MeshProxyTool/MeshProxyTool.h
+++ b/Engine/Source/Editor/MergeActors/Private/MeshProxyTool/MeshProxyTool.h
@@ -20,19 +20,36 @@ public:
 	}
 
 	static UMeshProxySettingsObject* Get()
-	{
-		static bool bInitialized = false;
-		// This is a singleton, use default object
-		UMeshProxySettingsObject* DefaultSettings = GetMutableDefault<UMeshProxySettingsObject>();
-
+	{	
+		// This is a singleton, duplicate default object
 		if (!bInitialized)
 		{
+			DefaultSettings = DuplicateObject(GetMutableDefault<UMeshProxySettingsObject>(), nullptr);
+			DefaultSettings->AddToRoot();
 			bInitialized = true;
 		}
 
 		return DefaultSettings;
 	}
 
+	static void Destroy()
+	{
+		if (bInitialized)
+		{
+			if (DefaultSettings)
+			{
+				DefaultSettings->RemoveFromRoot();
+				DefaultSettings->MarkPendingKill();
+			}
+
+			DefaultSettings = nullptr;
+			bInitialized = false;
+		}
+	}
+
+protected:
+	static bool bInitialized;
+	static UMeshProxySettingsObject* DefaultSettings;
 public:
 	UPROPERTY(editanywhere, meta = (ShowOnlyInnerProperties), Category = ProxySettings)
 	FMeshProxySettings  Settings;
@@ -47,6 +64,7 @@ class FMeshProxyTool : public IMergeActorsTool
 
 public:
 	FMeshProxyTool();
+	~FMeshProxyTool();
 
 	// IMergeActorsTool interface
 	virtual TSharedRef<SWidget> GetWidget() override;
diff --git a/Engine/Source/Editor/MergeActors/Private/MeshProxyTool/SMeshProxyDialog.cpp b/Engine/Source/Editor/MergeActors/Private/MeshProxyTool/SMeshProxyDialog.cpp
--- a/Engine/Source/Editor/MergeActors/Private/MeshProxyTool/SMeshProxyDialog.cpp
+++ b/Engine/Source/Editor/MergeActors/Private/MeshProxyTool/SMeshProxyDialog.cpp
@@ -191,6 +191,7 @@ void SMeshProxyDialog::CreateSettingsView()
 	DetailsViewArgs.bCustomNameAreaLocation = false;
 	DetailsViewArgs.bCustomFilterAreaLocation = true;
 	DetailsViewArgs.DefaultsOnlyVisibility = EEditDefaultsOnlyNodeVisibility::Hide;
+	DetailsViewArgs.NameAreaSettings = FDetailsViewArgs::ENameAreaSettings::HideNameArea;
 
 
 	// Tiny hack to hide this setting, since we have no way / value to go off to 
diff --git a/Engine/Source/Editor/Persona/Private/AnimationEditorViewportClient.cpp b/Engine/Source/Editor/Persona/Private/AnimationEditorViewportClient.cpp
--- a/Engine/Source/Editor/Persona/Private/AnimationEditorViewportClient.cpp
+++ b/Engine/Source/Editor/Persona/Private/AnimationEditorViewportClient.cpp
@@ -862,11 +862,12 @@ FText FAnimationViewportClient::GetDisplayInfo(bool bDisplayAllInfo) const
 				NumSectionsInUse
 				));
 
+			TArray<FTransform> LocalBoneTransforms = PreviewMeshComponent->GetBoneSpaceTransforms();
 			if (PreviewMeshComponent->BonesOfInterest.Num() > 0)
 			{
 				int32 BoneIndex = PreviewMeshComponent->BonesOfInterest[0];
 				FTransform ReferenceTransform = PreviewMeshComponent->GetReferenceSkeleton().GetRefBonePose()[BoneIndex];
-				FTransform LocalTransform = PreviewMeshComponent->BoneSpaceTransforms[BoneIndex];
+				FTransform LocalTransform = LocalBoneTransforms[BoneIndex];
 				FTransform ComponentTransform = PreviewMeshComponent->GetDrawTransform(BoneIndex);
 
 				TextValue = ConcatenateLine(TextValue, FText::Format(LOCTEXT("LocalTransform", "Local: {0}"), FText::FromString(LocalTransform.ToHumanReadableString())));
diff --git a/Engine/Source/Editor/Persona/Private/PersonaMeshDetails.cpp b/Engine/Source/Editor/Persona/Private/PersonaMeshDetails.cpp
--- a/Engine/Source/Editor/Persona/Private/PersonaMeshDetails.cpp
+++ b/Engine/Source/Editor/Persona/Private/PersonaMeshDetails.cpp
@@ -1268,21 +1268,7 @@ TSharedRef<SWidget> FPersonaMeshDetails::CreateSkinWeightProfileMenuContent()
 			MeshDetailLayout->ForceRefreshDetails();
 		}
 	})));
-
-	// Menu entry for copying skin weights from another Skeletal Mesh
-	AddProfileMenuBuilder.AddMenuEntry(LOCTEXT("CopyOverrideLabel", "Copy Skin Weights"), LOCTEXT("CopyOverrideToolTip", "Copy Skin Weights from another Skeletal Mesh"),
-		FSlateIcon(), FUIAction(FExecuteAction::CreateLambda([this, WeakSkeletalMeshPtr = SkeletalMeshPtr]()
-	{
-		if (USkeletalMesh* SkeletalMesh = WeakSkeletalMeshPtr.Get())
-		{
-			FScopedTransaction ScopedTransaction(LOCTEXT("CopySkinWeightProfile", "Create Skin Weight Profile with Skin Weights from another Skeletal Mesh"));
-			SkeletalMesh->Modify();
-
-			FSkinWeightProfileHelpers::CopySkinWeightProfile(SkeletalMesh);
-			MeshDetailLayout->ForceRefreshDetails();
-		}
-	})));
-
+	
 	// Add extra (sub)-menus for previously added Skin Weight Profiles
 	if (USkeletalMesh* Mesh = SkeletalMeshPtr.Get())
 	{
diff --git a/Engine/Source/Editor/Persona/Private/PersonaModule.cpp b/Engine/Source/Editor/Persona/Private/PersonaModule.cpp
--- a/Engine/Source/Editor/Persona/Private/PersonaModule.cpp
+++ b/Engine/Source/Editor/Persona/Private/PersonaModule.cpp
@@ -175,6 +175,7 @@ void FPersonaModule::ShutdownModule()
 		PropertyModule.UnregisterCustomClassLayout("SkeletalMeshSocket");
 		PropertyModule.UnregisterCustomClassLayout("EditorNotifyObject");
 		PropertyModule.UnregisterCustomClassLayout("AnimGraphNode_Base");
+		PropertyModule.UnregisterCustomClassLayout("AnimInstance");
 		PropertyModule.UnregisterCustomClassLayout("BlendSpaceBase");
 
 		PropertyModule.UnregisterCustomPropertyTypeLayout("InputScaleBias");
@@ -183,8 +184,8 @@ void FPersonaModule::ShutdownModule()
 		PropertyModule.UnregisterCustomPropertyTypeLayout("BlendParameter");
 		PropertyModule.UnregisterCustomPropertyTypeLayout("InterpolationParameter");
 
-		PropertyModule.UnregisterCustomClassLayout("SkeletalMeshSamplingRegionBoneFilter");
-		PropertyModule.UnregisterCustomClassLayout("SkeletalMeshSamplingRegionMaterialFilter");
+		PropertyModule.UnregisterCustomPropertyTypeLayout("SkeletalMeshSamplingRegionBoneFilter");
+		PropertyModule.UnregisterCustomPropertyTypeLayout("SkeletalMeshSamplingRegionMaterialFilter");
 	}
 }
 
diff --git a/Engine/Source/Editor/Persona/Private/SSkinWeightProfileImportOptions.cpp b/Engine/Source/Editor/Persona/Private/SSkinWeightProfileImportOptions.cpp
--- a/Engine/Source/Editor/Persona/Private/SSkinWeightProfileImportOptions.cpp
+++ b/Engine/Source/Editor/Persona/Private/SSkinWeightProfileImportOptions.cpp
@@ -19,26 +19,8 @@
 
 void FSkinWeightImportOptionsCustomization::CustomizeDetails(IDetailLayoutBuilder& DetailBuilder)
 {
-	TSharedRef<IPropertyHandle> FilePathHandle = DetailBuilder.GetProperty(GET_MEMBER_NAME_CHECKED(USkinWeightImportOptions, FilePath));
-	TSharedRef<IPropertyHandle> SourceMeshHandle = DetailBuilder.GetProperty(GET_MEMBER_NAME_CHECKED(USkinWeightImportOptions, SourceMesh));
 	TSharedPtr<IPropertyHandle> LODIndexHandle = DetailBuilder.GetProperty(GET_MEMBER_NAME_CHECKED(USkinWeightImportOptions, LODIndex));
-
-	FString FilePath;
-
-	// In case we are importing from a file do not show the Skeletal Mesh property and vice-versa
-	if (FilePathHandle->GetValue(FilePath) == FPropertyAccess::Success)
-	{
-		if (!FilePath.IsEmpty())
-		{
-			DetailBuilder.HideProperty(SourceMeshHandle);
-		}
-		else
-		{
-			DetailBuilder.HideProperty(FilePathHandle);
-			DetailBuilder.HideProperty(LODIndexHandle);
-		}
-	}
-
+	
 	ProfileNameHandle = DetailBuilder.GetProperty(GET_MEMBER_NAME_CHECKED(USkinWeightImportOptions, ProfileName));
 	DetailBuilder.HideProperty(ProfileNameHandle);
 
@@ -164,12 +146,8 @@ void SSkinWeightProfileImportOptions::Construct(const FArguments& InArgs)
 				.OnClicked(this, &SSkinWeightProfileImportOptions::OnImport)
 				.IsEnabled_Lambda([this]() -> bool 
 				{
-					const bool bCopyOperation = ImportSettings->FilePath.IsEmpty();
 					const bool bValidProfileName = DetailsCustomization->IsProfileNameValid(DetailsCustomization->OnGetProfileName().ToString());
-					const bool bValidSkeletalMesh = ImportSettings->SourceMesh != nullptr;
-
-					// Need a valid profile name, and in case we are copying we'll need to have a valid source skeletal mesh as well 
-					return bValidProfileName && (!bCopyOperation || bValidSkeletalMesh);
+					return bValidProfileName;
 				})
 			]
 			+ SUniformGridPanel::Slot(1, 0)
diff --git a/Engine/Source/Editor/Persona/Private/SkinWeightProfileHelpers.cpp b/Engine/Source/Editor/Persona/Private/SkinWeightProfileHelpers.cpp
--- a/Engine/Source/Editor/Persona/Private/SkinWeightProfileHelpers.cpp
+++ b/Engine/Source/Editor/Persona/Private/SkinWeightProfileHelpers.cpp
@@ -39,7 +39,6 @@ void FSkinWeightProfileHelpers::ImportSkinWeightProfile(USkeletalMesh* InSkeleta
 
 			TSharedPtr<SSkinWeightProfileImportOptions> OptionsWidget;
 			ImportSettings->FilePath = PickedFileName;
-			ImportSettings->SourceMesh = nullptr;
 			ImportSettings->LODIndex = 0;
 			Window->SetContent
 			(
@@ -65,87 +64,30 @@ void FSkinWeightProfileHelpers::ImportSkinWeightProfile(USkeletalMesh* InSkeleta
 		
 		if (bShouldImport)
 		{
-			const FName ProfileName(*ImportSettings->ProfileName);
-			if (ImportSettings->SourceMesh == nullptr)
+			const FName ProfileName(*ImportSettings->ProfileName);			
+			// Try and import the skin weight profile from the provided FBX file path
+			const bool bReregisterComponent = true;
+			const bool bResult = FLODUtilities::ImportAlternateSkinWeight(InSkeletalMesh, PickedFileName, ImportSettings->LODIndex, ProfileName, bReregisterComponent);
+			if (bResult)
 			{
-				// Try and import the skin weight profile from the provided FBX file path
-				const bool bReregisterComponent = true;
-				const bool bResult = FLODUtilities::ImportAlternateSkinWeight(InSkeletalMesh, PickedFileName, ImportSettings->LODIndex, ProfileName, bReregisterComponent);
-				if (bResult)
-				{
-					FNotificationInfo NotificationInfo(FText::GetEmpty());
-					NotificationInfo.Text = LOCTEXT("ImportSuccessful", "Skin Weights imported successfully!");
-					NotificationInfo.ExpireDuration = 2.5f;
-					FSlateNotificationManager::Get().AddNotification(NotificationInfo);
-				}
-				else
-				{
-					FNotificationInfo NotificationInfo(FText::GetEmpty());
-					NotificationInfo.Text = LOCTEXT("ImportFailed", "Failed to import Skin Weights!");
-					NotificationInfo.ExpireDuration = 2.5f;
-					FSlateNotificationManager::Get().AddNotification(NotificationInfo);
-				}
+				FNotificationInfo NotificationInfo(FText::GetEmpty());
+				NotificationInfo.Text = LOCTEXT("ImportSuccessful", "Skin Weights imported successfully!");
+				NotificationInfo.ExpireDuration = 2.5f;
+				FSlateNotificationManager::Get().AddNotification(NotificationInfo);
 			}
 			else
 			{
-				if (USkeletalMesh* Mesh = ImportSettings->SourceMesh)
-				{
-					if (CopySkinWeightsToProfile(Mesh, InSkeletalMesh, LOD0Index, ProfileName))
-					{
-						// Add a new profile entry
-						FSkinWeightProfileInfo SkeletalMeshProfile;
-						SkeletalMeshProfile.DefaultProfile = (InSkeletalMesh->GetNumSkinWeightProfiles() == 0);
-						SkeletalMeshProfile.DefaultProfileFromLODIndex = LOD0Index;
-						SkeletalMeshProfile.Name = ProfileName;
-						InSkeletalMesh->AddSkinWeightProfile(SkeletalMeshProfile);
-
-						// Enforce a post-edit call to ensure the mesh data is cached
-						FPropertyChangedEvent PropertyChangeEvent(nullptr);
-						InSkeletalMesh->PostEditChangeProperty(PropertyChangeEvent);
-					}
-				}				
+				FNotificationInfo NotificationInfo(FText::GetEmpty());
+				NotificationInfo.Text = LOCTEXT("ImportFailed", "Failed to import Skin Weights!");
+				NotificationInfo.ExpireDuration = 2.5f;
+				FSlateNotificationManager::Get().AddNotification(NotificationInfo);
 			}
+			
 			FLODUtilities::RegenerateDependentLODs(InSkeletalMesh, 0);
 		}
 	}
 }
 
-bool FSkinWeightProfileHelpers::CopySkinWeightsToProfile(USkeletalMesh* SourceMesh, USkeletalMesh* TargetMesh, const int32 LODIndex, const FName ProfileName)
-{
-	bool bCopied = false;
-	const FSkeletalMeshModel* SourceModel = SourceMesh->GetImportedModel();
-	FSkeletalMeshModel* TargetModel = TargetMesh->GetImportedModel();
-	if (SourceModel && TargetModel)
-	{
-		if (SourceModel->LODModels.IsValidIndex(LODIndex) && TargetModel->LODModels.IsValidIndex(LODIndex))
-		{
-			const FSkeletalMeshLODModel& SourceLODModel = SourceModel->LODModels[LODIndex];
-			FSkeletalMeshLODModel& TargetLODModel = TargetModel->LODModels[LODIndex];
-			if (SourceLODModel.NumVertices == TargetLODModel.NumVertices)
-			{
-				// Copy actual per-lod data into profile entry
-				TArray<FSoftSkinVertex> SourceVertices;
-				SourceLODModel.GetVertices(SourceVertices);
-
-				// Add a new imported skin weight data entry
-				FImportedSkinWeightProfileData& ProfileData = TargetLODModel.SkinWeightProfiles.Add(ProfileName);
-				TArray<FRawSkinWeight>& SkinWeights = ProfileData.SkinWeights;
-				SkinWeights.Reserve(SourceVertices.Num());
-				for (const FSoftSkinVertex& Vertex : SourceVertices)
-				{
-					FRawSkinWeight& SkinWeight = SkinWeights.AddDefaulted_GetRef();
-					FMemory::Memcpy(SkinWeight.InfluenceBones, Vertex.InfluenceBones);
-					FMemory::Memcpy(SkinWeight.InfluenceWeights, Vertex.InfluenceWeights);
-				}
-
-				bCopied = true;
-			}
-		}
-	}
-
-	return bCopied;
-}
-
 void FSkinWeightProfileHelpers::ImportSkinWeightProfileLOD(USkeletalMesh* InSkeletalMesh, FName ProfileName, int32 LODIndex)
 {
 	if (InSkeletalMesh)
@@ -174,68 +116,6 @@ void FSkinWeightProfileHelpers::ImportSkinWeightProfileLOD(USkeletalMesh* InSkel
 	}
 }
 
-void FSkinWeightProfileHelpers::CopySkinWeightProfile(USkeletalMesh* InSkeletalMesh)
-{
-	if (InSkeletalMesh)
-	{
-		bool bShouldImport = false;
-		USkinWeightImportOptions* ImportSettings = GetMutableDefault<USkinWeightImportOptions>();
-
-		// Show settings dialog	
-		TSharedRef<SWindow> Window = SNew(SWindow)
-			.Title(LOCTEXT("WindowTitle", "Skin Weight Profile Import Options"))
-			.SizingRule(ESizingRule::Autosized);
-
-		TSharedPtr<SSkinWeightProfileImportOptions> OptionsWidget;
-		ImportSettings->FilePath.Empty();
-		ImportSettings->SourceMesh = nullptr;
-		Window->SetContent
-		(
-			SAssignNew(OptionsWidget, SSkinWeightProfileImportOptions).WidgetWindow(Window)
-			.ImportSettings(ImportSettings)
-			.WidgetWindow(Window)
-			.SkeletalMesh(InSkeletalMesh)
-		);
-
-		TSharedPtr<SWindow> ParentWindow;
-
-		if (FModuleManager::Get().IsModuleLoaded("MainFrame"))
-		{
-			IMainFrameModule& MainFrame = FModuleManager::LoadModuleChecked<IMainFrameModule>("MainFrame");
-			ParentWindow = MainFrame.GetParentWindow();
-		}
-
-		FSlateApplication::Get().AddModalWindow(Window, ParentWindow, false);
-
-		// Poll the result of the window interaction, user could have cancelled operation
-		bShouldImport = OptionsWidget->ShouldImport();
-
-		if (bShouldImport)
-		{
-			const int32 LOD0Index = 0;
-			const FName ProfileName(*ImportSettings->ProfileName);
-			if (USkeletalMesh* Mesh = ImportSettings->SourceMesh)
-			{				
-				if (CopySkinWeightsToProfile(Mesh, InSkeletalMesh, LOD0Index, ProfileName))
-				{
-					// Add a new profile entry
-					FSkinWeightProfileInfo SkeletalMeshProfile;
-					SkeletalMeshProfile.DefaultProfile = (InSkeletalMesh->GetNumSkinWeightProfiles() == 0);
-					SkeletalMeshProfile.DefaultProfileFromLODIndex = LOD0Index;
-					SkeletalMeshProfile.Name = ProfileName;
-					InSkeletalMesh->AddSkinWeightProfile(SkeletalMeshProfile);
-
-					// Enforce a post-edit call to ensure the mesh data is cached
-					FPropertyChangedEvent PropertyChangeEvent(nullptr);
-					InSkeletalMesh->PostEditChangeProperty(PropertyChangeEvent);
-
-					// Make sure we regenerate any LOD data that is based off the now re imported LOD
-					FLODUtilities::RegenerateDependentLODs(InSkeletalMesh, LOD0Index);
-				}
-			}
-		}
-	}
-}
 
 void FSkinWeightProfileHelpers::ReimportSkinWeightProfileLOD(USkeletalMesh* InSkeletalMesh, const FName& InProfileName, const int32 LODIndex)
 {
diff --git a/Engine/Source/Editor/Persona/Public/SSkinWeightProfileImportOptions.h b/Engine/Source/Editor/Persona/Public/SSkinWeightProfileImportOptions.h
--- a/Engine/Source/Editor/Persona/Public/SSkinWeightProfileImportOptions.h
+++ b/Engine/Source/Editor/Persona/Public/SSkinWeightProfileImportOptions.h
@@ -32,13 +32,9 @@ public:
 	/** File path to FBX file containing Mesh with alternative set of Skin Weights */
 	UPROPERTY(VisibleAnywhere, Category = SkinWeights)
 	FString FilePath;
-
-	/** Skeletal Mesh from which to copy the Skin Weights used for this Profile */
-	UPROPERTY(EditAnywhere, Category = SkinWeights)
-	USkeletalMesh* SourceMesh;
-
+	
 	/** Target LOD index this file corresponds to */
-	UPROPERTY(EditAnywhere, Category = SkinWeights)
+	UPROPERTY(EditAnywhere, Category = SkinWeights, meta=(DisplayName="LOD Index"))
 	int32 LODIndex;
 };
 
diff --git a/Engine/Source/Editor/Persona/Public/SkinWeightProfileHelpers.h b/Engine/Source/Editor/Persona/Public/SkinWeightProfileHelpers.h
--- a/Engine/Source/Editor/Persona/Public/SkinWeightProfileHelpers.h
+++ b/Engine/Source/Editor/Persona/Public/SkinWeightProfileHelpers.h
@@ -16,9 +16,6 @@ struct FSkinWeightProfileHelpers
 	/** Tries to import a new set of Skin Weights for the given Skeletal Mesh at the given LOD index from an FBX file */
 	static void ImportSkinWeightProfileLOD(USkeletalMesh* InSkeletalMesh, FName ProfileName, int32 LODIndex);
 
-	/** Tries to copy a new set of Skin Weights for the given Skeletal Mesh from another Skeletal Mesh */
-	static void CopySkinWeightProfile(USkeletalMesh* InSkeletalMesh);
-
 	/** Tries to re import the previously imported skin weights for the given Skeletal Mesh, Profile name and LOD Index */
 	static void ReimportSkinWeightProfileLOD(USkeletalMesh* InSkeletalMesh, const FName& InProfileName, const int32 LODIndex);
 
diff --git a/Engine/Source/Editor/PropertyEditor/Private/SPerPlatformPropertiesWidget.cpp b/Engine/Source/Editor/PropertyEditor/Private/SPerPlatformPropertiesWidget.cpp
new file mode 100644
--- /dev/null
+++ b/Engine/Source/Editor/PropertyEditor/Private/SPerPlatformPropertiesWidget.cpp
@@ -0,0 +1,274 @@
+// Copyright 1998-2019 Epic Games, Inc. All Rights Reserved.
+
+#include "SPerPlatformPropertiesWidget.h"
+#include "Layout/Margin.h"
+#include "Fonts/SlateFontInfo.h"
+#include "Styling/CoreStyle.h"
+#include "Framework/MultiBox/MultiBoxBuilder.h"
+#include "Widgets/Input/SButton.h"
+#include "Widgets/Input/SComboButton.h"
+#include "PlatformInfo.h"
+#include "Widgets/Layout/SWrapBox.h"
+#include "EditorStyleSet.h"
+#include "Widgets/Layout/SBox.h"
+#include "DetailLayoutBuilder.h"
+#include "Widgets/Images/SImage.h"
+
+
+void SPerPlatformPropertiesWidget::Construct(const typename SPerPlatformPropertiesWidget::FArguments& InArgs)
+{
+	this->OnGenerateWidget = InArgs._OnGenerateWidget;
+	this->OnAddPlatform = InArgs._OnAddPlatform;
+	this->OnRemovePlatform = InArgs._OnRemovePlatform;
+	this->PlatformOverrideNames = InArgs._PlatformOverrideNames;
+
+	ConstructChildren();
+
+	// this widget has to check platform count change from outside to ensure the widget represents latest update
+	RegisterActiveTimer(FMath::RandRange(2.f, 5.f), FWidgetActiveTimerDelegate::CreateSP(this, &SPerPlatformPropertiesWidget::CheckPlatformCount));
+}
+
+void SPerPlatformPropertiesWidget::ConstructChildren()
+{
+	TSharedPtr<SWrapBox> WrapBox;
+
+	TArray<FName> PlatformOverrides = PlatformOverrideNames.Get();
+	LastPlatformOverrideNames = PlatformOverrides.Num();
+
+	ChildSlot
+	.VAlign(VAlign_Center)
+	.HAlign(HAlign_Fill)
+	[
+		SAssignNew(WrapBox, SWrapBox)
+		.UseAllottedWidth(true)
+	];
+
+	if (OnGenerateWidget.IsBound())
+	{
+		// Build Platform menu
+		bAddedMenuItem = false;
+		FMenuBuilder AddPlatformMenuBuilder(true, nullptr, nullptr, true);
+
+		// Platform (group) names
+		const TArray<FName>& PlatformGroupNameArray = PlatformInfo::GetAllPlatformGroupNames();
+		const TArray<FName>& VanillaPlatformNameArray = PlatformInfo::GetAllVanillaPlatformNames();
+
+		// Sanitized platform names
+		TArray<FName> BasePlatformNameArray;
+		// Mapping from platform group name to individual platforms
+		TMultiMap<FName, FName> GroupToPlatform;
+						
+		// Create mapping from platform to platform groups and remove postfixes and invalid platform names
+		const TArray<FString> Filters = { TEXT("NoEditor"), TEXT("Client"), TEXT("Server"), TEXT("AllDesktop") };
+		for (const FName& PlatformName : VanillaPlatformNameArray)
+		{
+			const PlatformInfo::FPlatformInfo* PlatformInfo = PlatformInfo::FindPlatformInfo(PlatformName);
+			FString PlatformNameString = PlatformName.ToString();
+			for ( const FString& Filter : Filters)
+			{
+				const int32 Position = PlatformNameString.Find(Filter);
+				if (Position != INDEX_NONE)
+				{
+					PlatformNameString.RemoveAt(Position, Filter.Len());
+					break;
+				}
+			}
+
+			// Add filtered name if it isn't already set, and also add to group mapping
+			const FName FilteredName = FName(*PlatformNameString);				
+			if (PlatformNameString.Len() && !PlatformOverrides.Contains(FilteredName))
+			{	
+				BasePlatformNameArray.AddUnique(FilteredName);
+				GroupToPlatform.AddUnique(PlatformInfo->PlatformGroupName, FilteredName);
+			}
+		}
+
+		// Create section for platform groups 
+		const FName PlatformGroupSection(TEXT("PlatformGroupSection"));
+		AddPlatformMenuBuilder.BeginSection(PlatformGroupSection, FText::FromString(TEXT("Platform Groups")));
+		for (const FName& GroupName : PlatformGroupNameArray)
+		{				
+			if (!PlatformOverrides.Contains(GroupName))
+			{					
+				const FTextFormat Format = NSLOCTEXT("SPerPlatformPropertiesWidget", "AddOverrideGroupFor", "Add Override for Platforms part of the {0} Platform Group");
+				AddPlatformToMenu(GroupName, Format, AddPlatformMenuBuilder);
+				bAddedMenuItem = true;
+			}
+		}				
+		AddPlatformMenuBuilder.EndSection();
+			
+		for (const FName& GroupName : PlatformGroupNameArray)
+		{
+			// Create a section for each platform group and their respective platforms
+			AddPlatformMenuBuilder.BeginSection(GroupName, FText::FromName(GroupName));
+
+			TArray<FName> PlatformNames;
+			GroupToPlatform.MultiFind(GroupName, PlatformNames);
+
+			const FTextFormat Format = NSLOCTEXT("SPerPlatformPropertiesWidget", "AddOverrideFor", "Add Override specifically for {0}");
+			for (const FName& PlatformName : PlatformNames)
+			{
+				AddPlatformToMenu(PlatformName, Format, AddPlatformMenuBuilder);
+			}
+
+			bAddedMenuItem |= PlatformNames.Num() > 0;
+
+			AddPlatformMenuBuilder.EndSection();
+		}
+
+		// Default control
+		WrapBox->AddSlot()
+		[
+			MakePerPlatformWidget(NAME_None, NSLOCTEXT("SPerPlatformPropertiesWidget", "DefaultPlatform", "Default"), PlatformOverrides, AddPlatformMenuBuilder)
+		];
+
+		for (FName PlatformName : PlatformOverrides)
+		{
+			WrapBox->AddSlot()
+			[
+				MakePerPlatformWidget(PlatformName, FText::AsCultureInvariant(PlatformName.ToString()), PlatformOverrides, AddPlatformMenuBuilder)
+			];
+		}
+	}
+	else
+	{
+		WrapBox->AddSlot()
+		[
+			SNew(STextBlock)
+			.Text(NSLOCTEXT("SPerPlatformPropertiesWidget", "OnGenerateWidgetWarning", "No OnGenerateWidget() Provided"))
+			.ColorAndOpacity(FLinearColor::Red)
+		];
+	}
+}
+
+void SPerPlatformPropertiesWidget::AddPlatformToMenu(const FName& PlatformName, const FTextFormat Format, FMenuBuilder &AddPlatformMenuBuilder)
+{		
+	const FText MenuText = FText::Format(FText::FromString(TEXT("{0}")), FText::AsCultureInvariant(PlatformName.ToString()));
+	const FText MenuTooltipText = FText::Format(Format, FText::AsCultureInvariant(PlatformName.ToString()));
+	AddPlatformMenuBuilder.AddMenuEntry(
+		MenuText,
+		MenuTooltipText,
+		FSlateIcon(FEditorStyle::GetStyleSetName(), "PerPlatformWidget.AddPlatform"),
+		FUIAction(FExecuteAction::CreateSP(this, &SPerPlatformPropertiesWidget::AddPlatform, PlatformName))
+	);
+}
+
+void SPerPlatformPropertiesWidget::AddPlatform(FName PlatformName)
+{
+	if (OnAddPlatform.IsBound() && OnAddPlatform.Execute(PlatformName))
+	{
+		ConstructChildren();
+		Invalidate(EInvalidateWidget::LayoutAndVolatility);
+	}
+}
+
+FReply SPerPlatformPropertiesWidget::RemovePlatform(FName PlatformName)
+{
+	if (OnRemovePlatform.IsBound() && OnRemovePlatform.Execute(PlatformName))
+	{
+		ConstructChildren();
+		Invalidate(EInvalidateWidget::LayoutAndVolatility);
+	}
+	return FReply::Handled();
+}
+
+EActiveTimerReturnType SPerPlatformPropertiesWidget::CheckPlatformCount(double InCurrentTime, float InDeltaSeconds)
+{
+	// @fixme: the platform count is fixed and locally cached
+	// so if you change this outside of editor, this widget won't update
+	// this timer is the one checking to see if platform count has changed
+	// if so, it will reconstruct
+	TArray<FName> PlatformOverrides = PlatformOverrideNames.Get();
+	if (LastPlatformOverrideNames != PlatformOverrides.Num())
+	{
+		ConstructChildren();
+	}
+
+	return EActiveTimerReturnType::Continue;
+}
+
+
+TSharedRef<SWidget> SPerPlatformPropertiesWidget::MakePerPlatformWidget(FName InName, FText InDisplayText, const TArray<FName>& InPlatformOverrides, FMenuBuilder& InAddPlatformMenuBuilder)
+{
+	TSharedPtr<SHorizontalBox> HorizontalBox;
+
+	TSharedRef<SWidget> Widget = 
+		SNew(SBox)
+		.ToolTipText(InName == NAME_None ? 
+			NSLOCTEXT("SPerPlatformPropertiesWidget", "DefaultPlatformDesc", "This property can have per-platform or platform group overrides.\nThis is the default value used when no override has been set for a platform or platform group.") : 
+			FText::Format(NSLOCTEXT("SPerPlatformPropertiesWidget", "PerPlatformDesc", "Override for {0}"), InDisplayText))
+		.Padding(FMargin(0.0f, 2.0f, 4.0f, 2.0f))
+		.MinDesiredWidth(50.0f)
+		[
+			SNew(SVerticalBox)
+			+SVerticalBox::Slot()
+			.AutoHeight()
+			[
+				SAssignNew(HorizontalBox, SHorizontalBox)
+				+SHorizontalBox::Slot()
+				.FillWidth(1.0f)
+				.Padding(0.0f, 0.0f, 2.0f, 2.0f)
+				.VAlign(VAlign_Center)
+				[
+					SNew(STextBlock)
+					.Font(IDetailLayoutBuilder::GetDetailFont())
+					.Text(InDisplayText)
+				]
+			]
+			+SVerticalBox::Slot()
+			.AutoHeight()
+			[
+				OnGenerateWidget.Execute(InName)
+			]
+		];
+
+	if(InName != NAME_None)
+	{
+		HorizontalBox->AddSlot()
+			.AutoWidth()
+			.Padding(2.0f, 0.0f, 0.0f, 2.0f)
+			.VAlign(VAlign_Center)
+			[
+				SNew(SButton)
+				.ContentPadding(2.0f)
+				.ButtonStyle(FEditorStyle::Get(), "HoverHintOnly")
+				.OnClicked(this, &SPerPlatformPropertiesWidget::RemovePlatform, InName)
+				.ToolTipText(FText::Format(NSLOCTEXT("SPerPlatformPropertiesWidget", "RemoveOverrideFor", "Remove Override for {0}"), InDisplayText))
+				.ForegroundColor(FSlateColor::UseForeground())
+				.HAlign(HAlign_Center)
+				.VAlign(VAlign_Center)
+				.Content()
+				[
+					SNew(SImage)
+					.Image(FEditorStyle::GetBrush("Icons.Cross"))
+				]
+			];
+	}
+	else
+	{
+		HorizontalBox->AddSlot()
+			.AutoWidth()
+			.Padding(2.0f, 0.0f, 0.0f, 2.0f)
+			.VAlign(VAlign_Center)
+			[
+				SNew(SComboButton)
+				.Visibility_Lambda([this](){ return bAddedMenuItem ? EVisibility::Visible : EVisibility::Hidden; })
+				.ButtonStyle(FEditorStyle::Get(), "HoverHintOnly")
+				.ContentPadding(2.0f)
+				.ForegroundColor(FSlateColor::UseForeground())
+				.HasDownArrow(false)
+				.ButtonContent()
+				[
+					SNew(SImage)
+					.Image(FEditorStyle::GetBrush("PropertyWindow.Button_AddToArray"))
+				]
+				.MenuContent()
+				[
+					InAddPlatformMenuBuilder.MakeWidget()
+				]
+				.ToolTipText(NSLOCTEXT("SPerPlatformPropertiesWidget", "AddOverrideToolTip", "Add an override for a specific platform or platform group"))	
+			];
+	}
+
+	return Widget;
+}
\ No newline at end of file
diff --git a/Engine/Source/Editor/PropertyEditor/PropertyEditor.Build.cs b/Engine/Source/Editor/PropertyEditor/PropertyEditor.Build.cs
--- a/Engine/Source/Editor/PropertyEditor/PropertyEditor.Build.cs
+++ b/Engine/Source/Editor/PropertyEditor/PropertyEditor.Build.cs
@@ -61,6 +61,7 @@ public class PropertyEditor : ModuleRules
                 "RHI",
 				"ConfigEditor",
                 "SceneOutliner",
+				"DesktopPlatform",
             }
         );
 
diff --git a/Engine/Source/Editor/PropertyEditor/Public/SPerPlatformPropertiesWidget.h b/Engine/Source/Editor/PropertyEditor/Public/SPerPlatformPropertiesWidget.h
--- a/Engine/Source/Editor/PropertyEditor/Public/SPerPlatformPropertiesWidget.h
+++ b/Engine/Source/Editor/PropertyEditor/Public/SPerPlatformPropertiesWidget.h
@@ -3,23 +3,17 @@
 #pragma once
 
 #include "CoreMinimal.h"
-#include "Layout/Margin.h"
-#include "Fonts/SlateFontInfo.h"
 #include "Widgets/DeclarativeSyntaxSupport.h"
 #include "Widgets/SWidget.h"
 #include "Widgets/SCompoundWidget.h"
-#include "Styling/CoreStyle.h"
 #include "Framework/SlateDelegates.h"
-#include "Framework/MultiBox/MultiBoxBuilder.h"
-#include "Widgets/Input/SButton.h"
-#include "Widgets/Input/SComboButton.h"
-#include "PlatformInfo.h"
 
+class FMenuBuilder;
 
 /**
 * SPerPlatformPropertiesWidget
 */
-class SPerPlatformPropertiesWidget : public SCompoundWidget
+class PROPERTYEDITOR_API SPerPlatformPropertiesWidget : public SCompoundWidget
 {
 public:
 	typedef typename TSlateDelegates<FName>::FOnGenerateWidget FOnGenerateWidget;
@@ -43,262 +37,26 @@ public:
 	*
 	* @param	InArgs	The declaration data for this widget
 	*/
-	void Construct(const typename SPerPlatformPropertiesWidget::FArguments& InArgs)
-	{
-		this->OnGenerateWidget = InArgs._OnGenerateWidget;
-		this->OnAddPlatform = InArgs._OnAddPlatform;
-		this->OnRemovePlatform = InArgs._OnRemovePlatform;
-		this->PlatformOverrideNames = InArgs._PlatformOverrideNames;
+	void Construct(const typename SPerPlatformPropertiesWidget::FArguments& InArgs);
 
-		ConstructChildren();
-
-		// this widget has to check platform count change from outside to ensure the widget represents latest update
-		RegisterActiveTimer(FMath::RandRange(2.f, 5.f), FWidgetActiveTimerDelegate::CreateSP(this, &SPerPlatformPropertiesWidget::CheckPlatformCount));
-	}
-
-	void ConstructChildren()
-	{
-		TSharedPtr<SHorizontalBox> HorizontalBox;
-		TSharedPtr<SVerticalBox> VerticalBox;
-
-		TArray<FName> PlatformOverrides = PlatformOverrideNames.Get();
-		LastPlatformOverrideNames = PlatformOverrides.Num();
-		ChildSlot
-		.VAlign(VAlign_Fill)
-		.HAlign(HAlign_Fill)
-		[
-			SAssignNew(HorizontalBox, SHorizontalBox)
-			+SHorizontalBox::Slot()
-			.FillWidth(1.0f)
-			[
-				SAssignNew(VerticalBox, SVerticalBox)
-				+ SVerticalBox::Slot()
-				.AutoHeight()
-				[
-					SNew(STextBlock)
-					.Visibility(PlatformOverrides.Num() ? EVisibility::Visible : EVisibility::Collapsed)
-					.Font(IDetailLayoutBuilder::GetDetailFont())
-					.Text(NSLOCTEXT("SPerPlatformPropertiesWidget", "DefaultPlatform", "Default"))
-				]
-			]
-		];
-
-		if (OnGenerateWidget.IsBound())
-		{
-			// Default control
-			VerticalBox->AddSlot()
-			.AutoHeight()
-			.VAlign(EVerticalAlignment::VAlign_Bottom)
-			[
-				OnGenerateWidget.Execute(NAME_None)
-			];
-
-			// Add Platform menu
-			bool bAddedMenuItem = false;
-			FMenuBuilder AddPlatformMenuBuilder(true, nullptr, nullptr, true);
-						
-			// Platform (group) names
-			const TArray<FName>& PlatformGroupNameArray = PlatformInfo::GetAllPlatformGroupNames();
-			const TArray<FName>& VanillaPlatformNameArray = PlatformInfo::GetAllVanillaPlatformNames();			
-
-			// Sanitized platform names
-			TArray<FName> BasePlatformNameArray;
-			// Mapping from platform group name to individual platforms
-			TMultiMap<FName, FName> GroupToPlatform;
-						
-			// Create mapping from platform to platform groups and remove postfixes and invalid platform names
-			const TArray<FString> Filters = { TEXT("NoEditor"), TEXT("Client"), TEXT("Server"), TEXT("AllDesktop") };
-			for (const FName& PlatformName : VanillaPlatformNameArray)
-			{
-				const PlatformInfo::FPlatformInfo* PlatformInfo = PlatformInfo::FindPlatformInfo(PlatformName);
-				FString PlatformNameString = PlatformName.ToString();
-				for ( const FString& Filter : Filters)
-				{
-					const int32 Position = PlatformNameString.Find(Filter);
-					if (Position != INDEX_NONE)
-					{
-						PlatformNameString.RemoveAt(Position, Filter.Len());
-						break;
-					}
-				}
-
-				// Add filtered name if it isn't already set, and also add to group mapping
-				const FName FilteredName = FName(*PlatformNameString);				
-				if (PlatformNameString.Len() && !PlatformOverrides.Contains(FilteredName))
-				{	
-					BasePlatformNameArray.AddUnique(FilteredName);
-					GroupToPlatform.AddUnique(PlatformInfo->PlatformGroupName, FilteredName);
-				}
-			}
-
-			// Create section for platform groups 
-			const FName PlatformGroupSection(TEXT("PlatformGroupSection"));
-			AddPlatformMenuBuilder.BeginSection(PlatformGroupSection, FText::FromString(TEXT("Platform Groups")));
-			for (const FName& GroupName : PlatformGroupNameArray)
-			{				
-				if (!PlatformOverrides.Contains(GroupName))
-				{					
-					const FTextFormat Format = NSLOCTEXT("SPerPlatformPropertiesWidget", "AddOverrideGroupFor", "Add Override for Platforms part of the {0} Platform Group");
-					AddPlatformToMenu(GroupName, Format, AddPlatformMenuBuilder);
-					bAddedMenuItem = true;
-				}
-			}				
-			AddPlatformMenuBuilder.EndSection();
-			
-			for (const FName& GroupName : PlatformGroupNameArray)
-			{
-				// Create a section for each platform group and their respective platforms
-				AddPlatformMenuBuilder.BeginSection(GroupName, FText::FromName(GroupName));
-
-				TArray<FName> PlatformNames;
-				GroupToPlatform.MultiFind(GroupName, PlatformNames);
-
-				const FTextFormat Format = NSLOCTEXT("SPerPlatformPropertiesWidget", "AddOverrideFor", "Add Override specifically for {0}");
-				for (const FName& PlatformName : PlatformNames)
-				{
-					AddPlatformToMenu(PlatformName, Format, AddPlatformMenuBuilder);
-				}
-
-				bAddedMenuItem |= PlatformNames.Num() > 0;
-
-				AddPlatformMenuBuilder.EndSection();
-			}
-			
-			if (bAddedMenuItem)
-			{
-				HorizontalBox->AddSlot()
-				.AutoWidth()
-				[
-					SNew(SComboButton)
-					.VAlign(EVerticalAlignment::VAlign_Bottom)
-					.ButtonStyle(FEditorStyle::Get(), "HoverHintOnly")
-					.ContentPadding(4.0f)
-					.ForegroundColor(FSlateColor::UseForeground())
-					.HasDownArrow(false)
-					.ButtonContent()
-					[
-						SNew(SImage)
-						.Image(FEditorStyle::GetBrush("PropertyWindow.Button_AddToArray"))
-					]
-					.MenuContent()
-					[
-						AddPlatformMenuBuilder.MakeWidget()
-					]
-					.ToolTipText(NSLOCTEXT("SPerPlatformPropertiesWidget", "AddOverrideToolTip", "Add an override for a specific platform or platform group"))
-				];
-			}
-
-			for (FName PlatformName : PlatformOverrides)
-			{
-				HorizontalBox->AddSlot()
-				.Padding(FMargin(5.0f, 0.0f, 0.0f, 0.0f))
-				.AutoWidth()
-				[
-					SNew(SVerticalBox)
-					+ SVerticalBox::Slot()
-					.AutoHeight()
-					.VAlign(EVerticalAlignment::VAlign_Bottom)
-					[
-						SNew(SHorizontalBox)
-						+ SHorizontalBox::Slot()
-						.FillWidth(1.0f)
-						[
-							SNew(STextBlock)
-							.Font(IDetailLayoutBuilder::GetDetailFont())
-							.Text(FText::FromName(PlatformName))
-						]
-
-						+ SHorizontalBox::Slot()
-						.HAlign(EHorizontalAlignment::HAlign_Right)
-						.AutoWidth()
-						[
-							SNew(SButton)
-							.ContentPadding(FMargin(1.0f, 0.0f, 1.0f, 0.0f))
-							.ButtonStyle(FEditorStyle::Get(), "NoBorder")
-							.OnClicked(this, &SPerPlatformPropertiesWidget::RemovePlatform, PlatformName)
-							.ToolTipText(FText::Format(NSLOCTEXT("SPerPlatformPropertiesWidget", "RemoveOverrideFor", "Remove Override for {0}"), FText::AsCultureInvariant(PlatformName.ToString())))
-							.ForegroundColor(FSlateColor::UseForeground())
-							.HAlign(HAlign_Center)
-							.VAlign(VAlign_Center)
-							.Content()
-							[
-								SNew(SImage)
-								.Image(FEditorStyle::GetBrush("Icons.Cross"))
-							]
-						]
-					]
-
-					+ SVerticalBox::Slot()
-					.AutoHeight()
-					[
-						OnGenerateWidget.Execute(PlatformName)
-					]
-				];
-			}
-		}
-		else
-		{
-			VerticalBox->AddSlot()
-			.AutoHeight()
-			[
-				SNew(STextBlock)
-				.Text(NSLOCTEXT("SPerPlatformPropertiesWidget", "OnGenerateWidgetWarning", "No OnGenerateWidget() Provided"))
-				.ColorAndOpacity(FLinearColor::Red)
-			];
-		}
-	}
+protected:
+	void ConstructChildren();
 
-	void AddPlatformToMenu(const FName& PlatformName, const FTextFormat Format, FMenuBuilder &AddPlatformMenuBuilder)
-	{		
-		const FText MenuText = FText::Format(FText::FromString(TEXT("{0}")), FText::AsCultureInvariant(PlatformName.ToString()));
-		const FText MenuTooltipText = FText::Format(Format, FText::AsCultureInvariant(PlatformName.ToString()));
-		AddPlatformMenuBuilder.AddMenuEntry(
-			MenuText,
-			MenuTooltipText,
-			FSlateIcon(FEditorStyle::GetStyleSetName(), "PerPlatformWidget.AddPlatform"),
-			FUIAction(FExecuteAction::CreateSP(this, &SPerPlatformPropertiesWidget::AddPlatform, PlatformName))
-		);
-	}
+	TSharedRef<SWidget> MakePerPlatformWidget(FName InName, FText InDisplayText, const TArray<FName>& InPlatformOverrides, FMenuBuilder& InAddPlatformMenuBuilder);
 
-protected:
-	void AddPlatform(FName PlatformName)
-	{
-		if (OnAddPlatform.IsBound() && OnAddPlatform.Execute(PlatformName))
-		{
-			ConstructChildren();
-			Invalidate(EInvalidateWidget::LayoutAndVolatility);
-		}
-	}
+	void AddPlatform(FName PlatformName);
 
-	FReply RemovePlatform(FName PlatformName)
-	{
-		if (OnRemovePlatform.IsBound() && OnRemovePlatform.Execute(PlatformName))
-		{
-			ConstructChildren();
-			Invalidate(EInvalidateWidget::LayoutAndVolatility);
-		}
-		return FReply::Handled();
-	}
+	FReply RemovePlatform(FName PlatformName);
 
-	EActiveTimerReturnType CheckPlatformCount(double InCurrentTime, float InDeltaSeconds)
-	{
-		// @fixme: the platform count is fixed and locally cached
-		// so if you change this outside of editor, this widget won't update
-		// this timer is the one checking to see if platform count has changed
-		// if so, it will reconstruct
-		TArray<FName> PlatformOverrides = PlatformOverrideNames.Get();
-		if (LastPlatformOverrideNames != PlatformOverrides.Num())
-		{
-			ConstructChildren();
-		}
+	EActiveTimerReturnType CheckPlatformCount(double InCurrentTime, float InDeltaSeconds);
 
-		return EActiveTimerReturnType::Continue;
-	}
+	void AddPlatformToMenu(const FName& PlatformName, const FTextFormat Format, FMenuBuilder& AddPlatformMenuBuilder);
 
 	FOnGenerateWidget OnGenerateWidget;
 	FOnPlatformAction OnAddPlatform;
 	FOnPlatformAction OnRemovePlatform;
 	TAttribute<TArray<FName>> PlatformOverrideNames;
 	int32 LastPlatformOverrideNames;
+	bool bAddedMenuItem;
 };
 
diff --git a/Engine/Source/Editor/SkeletalMeshEditor/Private/SkeletalMeshEditor.cpp b/Engine/Source/Editor/SkeletalMeshEditor/Private/SkeletalMeshEditor.cpp
--- a/Engine/Source/Editor/SkeletalMeshEditor/Private/SkeletalMeshEditor.cpp
+++ b/Engine/Source/Editor/SkeletalMeshEditor/Private/SkeletalMeshEditor.cpp
@@ -883,8 +883,27 @@ UObject* FSkeletalMeshEditor::HandleGetAsset()
 
 bool FSkeletalMeshEditor::HandleReimportMeshInternal(int32 SourceFileIndex /*= INDEX_NONE*/, bool bWithNewFile /*= false*/)
 {
+	// Make sure we clear any skin weight profile currently being previewed
+	UDebugSkelMeshComponent* MeshComponent = GetPersonaToolkit()->GetPreviewScene()->GetPreviewMeshComponent();
+	const FName ProfileName = MeshComponent ? MeshComponent->GetCurrentSkinWeightProfileName() : NAME_None;
+	if (MeshComponent)
+	{
+		MeshComponent->ClearSkinWeightProfile();
+	}
+
 	// Reimport the asset
-	return FReimportManager::Instance()->Reimport(SkeletalMesh, true, true, TEXT(""), nullptr, SourceFileIndex, bWithNewFile);
+	const bool bResult = FReimportManager::Instance()->Reimport(SkeletalMesh, true, true, TEXT(""), nullptr, SourceFileIndex, bWithNewFile);
+
+	// Refresh skeleton tree
+	SkeletonTree->Refresh();
+
+	// Re-set skin weigh profile
+	if (MeshComponent && (ProfileName != NAME_None))
+	{
+		MeshComponent->SetSkinWeightProfile(ProfileName);
+	}
+
+	return bResult;
 }
 
 void FSkeletalMeshEditor::HandleReimportMesh(int32 SourceFileIndex /*= INDEX_NONE*/)
diff --git a/Engine/Source/Editor/SkeletonEditor/Private/BoneProxy.cpp b/Engine/Source/Editor/SkeletonEditor/Private/BoneProxy.cpp
--- a/Engine/Source/Editor/SkeletonEditor/Private/BoneProxy.cpp
+++ b/Engine/Source/Editor/SkeletonEditor/Private/BoneProxy.cpp
@@ -26,10 +26,12 @@ void UBoneProxy::Tick(float DeltaTime)
 	{
 		if (UDebugSkelMeshComponent* Component = SkelMeshComponent.Get())
 		{
+			TArray<FTransform> LocalBoneTransforms = Component->GetBoneSpaceTransforms();
+
 			int32 BoneIndex = Component->GetBoneIndex(BoneName);
-			if (Component->BoneSpaceTransforms.IsValidIndex(BoneIndex))
+			if (LocalBoneTransforms.IsValidIndex(BoneIndex))
 			{
-				FTransform LocalTransform = Component->BoneSpaceTransforms[BoneIndex];
+				FTransform LocalTransform = LocalBoneTransforms[BoneIndex];
 				FTransform BoneTransform = Component->GetBoneTransform(BoneIndex);
 
 				if (bLocalLocation)
diff --git a/Engine/Source/Editor/StaticMeshEditor/Private/StaticMeshEditorTools.cpp b/Engine/Source/Editor/StaticMeshEditor/Private/StaticMeshEditorTools.cpp
--- a/Engine/Source/Editor/StaticMeshEditor/Private/StaticMeshEditorTools.cpp
+++ b/Engine/Source/Editor/StaticMeshEditor/Private/StaticMeshEditorTools.cpp
@@ -29,6 +29,7 @@
 #include "Widgets/Input/SVectorInputBox.h"
 #include "Widgets/Input/SNumericEntryBox.h"
 #include "SPerPlatformPropertiesWidget.h"
+#include "PlatformInfo.h"
 
 #include "Runtime/Analytics/Analytics/Public/Interfaces/IAnalyticsProvider.h"
 #include "EngineAnalytics.h"
@@ -3067,6 +3068,7 @@ void FLevelOfDetailSettingsLayout::AddToDetailsPanel( IDetailLayoutBuilder& Deta
 	.MaxDesiredWidth((float)(PlatformNumber + 1)*125.0f)
 	[
 		SNew(SPerPlatformPropertiesWidget)
+		.IsEnabled(FLevelOfDetailSettingsLayout::GetLODCount() > 1)
 		.OnGenerateWidget(this, &FLevelOfDetailSettingsLayout::GetMinLODWidget)
 		.OnAddPlatform(this, &FLevelOfDetailSettingsLayout::AddMinLODPlatformOverride)
 		.OnRemovePlatform(this, &FLevelOfDetailSettingsLayout::RemoveMinLODPlatformOverride)
@@ -3085,6 +3087,7 @@ void FLevelOfDetailSettingsLayout::AddToDetailsPanel( IDetailLayoutBuilder& Deta
 	.MaxDesiredWidth((float)(PlatformNumber + 1)*125.0f)
 	[
 		SNew(SPerPlatformPropertiesWidget)
+		.IsEnabled(FLevelOfDetailSettingsLayout::GetLODCount() > 1)
 		.OnGenerateWidget(this, &FLevelOfDetailSettingsLayout::GetNumStreamedLODsWidget)
 		.OnAddPlatform(this, &FLevelOfDetailSettingsLayout::AddNumStreamedLODsPlatformOverride)
 		.OnRemovePlatform(this, &FLevelOfDetailSettingsLayout::RemoveNumStreamedLODsPlatformOverride)
@@ -3388,6 +3391,7 @@ void FLevelOfDetailSettingsLayout::AddLODLevelCategories( IDetailLayoutBuilder&
 			.MaxDesiredWidth((float)(PlatformNumber + 1)*125.0f)
 			[
 				SNew(SPerPlatformPropertiesWidget)
+				.IsEnabled(this, &FLevelOfDetailSettingsLayout::CanChangeLODScreenSize)
 				.OnGenerateWidget(this, &FLevelOfDetailSettingsLayout::GetLODScreenSizeWidget, LODIndex)
 				.OnAddPlatform(this, &FLevelOfDetailSettingsLayout::AddLODScreenSizePlatformOverride, LODIndex)
 				.OnRemovePlatform(this, &FLevelOfDetailSettingsLayout::RemoveLODScreenSizePlatformOverride, LODIndex)
diff --git a/Engine/Source/Editor/UnrealEd/Private/Fbx/FbxAnimationExport.cpp b/Engine/Source/Editor/UnrealEd/Private/Fbx/FbxAnimationExport.cpp
--- a/Engine/Source/Editor/UnrealEd/Private/Fbx/FbxAnimationExport.cpp
+++ b/Engine/Source/Editor/UnrealEd/Private/Fbx/FbxAnimationExport.cpp
@@ -512,6 +512,7 @@ void FFbxExporter::ExportAnimTrack(IAnimTrackAdapter& AnimTrackAdapter, AActor*
 	}
 
 	FTransform InitialInvParentTransform;
+	TArray<FTransform> LocalBoneTransforms = SkeletalMeshComponent->GetBoneSpaceTransforms();
 
 	float SampleTime;
 	for(SampleTime = AnimationStart; SampleTime <= AnimationEnd; SampleTime += SamplingRate)
@@ -567,7 +568,7 @@ void FFbxExporter::ExportAnimTrack(IAnimTrackAdapter& AnimTrackAdapter, AActor*
 				Curves[i]->KeyModifyBegin();
 			}
 
-			FTransform BoneTransform = SkeletalMeshComponent->BoneSpaceTransforms[BoneIndex];
+			FTransform BoneTransform = LocalBoneTransforms[BoneIndex];
 
 			if (GetExportOptions()->MapSkeletalMotionToRoot && BoneIndex == 0)
 			{
diff --git a/Engine/Source/Editor/UnrealEd/Private/LODCluster.cpp b/Engine/Source/Editor/UnrealEd/Private/LODCluster.cpp
--- a/Engine/Source/Editor/UnrealEd/Private/LODCluster.cpp
+++ b/Engine/Source/Editor/UnrealEd/Private/LODCluster.cpp
@@ -47,7 +47,7 @@ const float CalculateOverlap(const FSphere& ASphere, const float AFillingFactor,
 	}
 
 	float Distance = (ASphere.Center-BSphere.Center).Size();
-	check (Distance > 0.f);
+	check(!FMath::IsNearlyZero(Distance));
 
 	float ARadius = ASphere.W;
 	float BRadius = BSphere.W;
@@ -67,13 +67,13 @@ const float CalculateOverlap(const FSphere& ASphere, const float AFillingFactor,
 	float OverlapRedius = FMath::Sqrt(OverlapRadius1) / OverlapRadius2;
 	float OverlapRediusSq = OverlapRedius*OverlapRedius;
 
-	check (OverlapRadius1 >= 0.f);
+	check(!FMath::IsNearlyZero(OverlapRadius1));
 
 	float ConstPI = PI/6.0f;
 	float AVolume = ConstPI*(3*OverlapRediusSq + ACapHeight*ACapHeight) * ACapHeight;
 	float BVolume = ConstPI*(3*OverlapRediusSq + BCapHeight*BCapHeight) * BCapHeight;
 
-	check (AVolume > 0.f &&  BVolume > 0.f);
+	check(!FMath::IsNearlyZero(AVolume) && !FMath::IsNearlyZero(BVolume));
 
 	float TotalVolume = AFillingFactor*AVolume + BFillingFactor*BVolume;
 	return TotalVolume;
diff --git a/Engine/Source/Editor/UnrealEd/Private/SkeletalMeshEdit.cpp b/Engine/Source/Editor/UnrealEd/Private/SkeletalMeshEdit.cpp
--- a/Engine/Source/Editor/UnrealEd/Private/SkeletalMeshEdit.cpp
+++ b/Engine/Source/Editor/UnrealEd/Private/SkeletalMeshEdit.cpp
@@ -1691,10 +1691,6 @@ bool UnFbx::FFbxImporter::ImportAnimation(USkeleton* Skeleton, UAnimSequence * D
 
 		GWarn->EndSlowTask();
 	}
-	else
-	{
-		DestSeq->RecycleAnimSequence();
-	}
 
 	// compress animation
 	{
diff --git a/Engine/Source/Programs/UnrealHeaderTool/Private/BaseParser.cpp b/Engine/Source/Programs/UnrealHeaderTool/Private/BaseParser.cpp
--- a/Engine/Source/Programs/UnrealHeaderTool/Private/BaseParser.cpp
+++ b/Engine/Source/Programs/UnrealHeaderTool/Private/BaseParser.cpp
@@ -120,6 +120,9 @@ namespace
 
 			FMetadataKeyword& EarlyAccessFeature = Dictionary.Add(TEXT("EarlyAccessPreview"), EMetadataValueArgument::None);
 			EarlyAccessFeature.InsertAddAction(TEXT("DevelopmentStatus"), TEXT("EarlyAccess"));
+
+			FMetadataKeyword& DocumentationPolicy = Dictionary.Add(TEXT("DocumentationPolicy"), EMetadataValueArgument::None);
+			DocumentationPolicy.InsertAddAction(TEXT("DocumentationPolicy"), TEXT("Strict"));
 		}
 
 		return Dictionary.Find(Keyword);
diff --git a/Engine/Source/Programs/UnrealHeaderTool/Private/HeaderParser.cpp b/Engine/Source/Programs/UnrealHeaderTool/Private/HeaderParser.cpp
--- a/Engine/Source/Programs/UnrealHeaderTool/Private/HeaderParser.cpp
+++ b/Engine/Source/Programs/UnrealHeaderTool/Private/HeaderParser.cpp
@@ -66,13 +66,17 @@ private:
 
 enum {MAX_ARRAY_SIZE=2048};
 
+static const FName NAME_Comment(TEXT("Comment"));
 static const FName NAME_ToolTip(TEXT("ToolTip"));
+static const FName NAME_DocumentationPolicy(TEXT("DocumentationPolicy"));
 EGeneratedCodeVersion FHeaderParser::DefaultGeneratedCodeVersion = EGeneratedCodeVersion::V1;
 TArray<FString> FHeaderParser::StructsWithNoPrefix;
 TArray<FString> FHeaderParser::StructsWithTPrefix;
 TArray<FString> FHeaderParser::DelegateParameterCountStrings;
 TMap<FString, FString> FHeaderParser::TypeRedirectMap;
+TArray<FString> FHeaderParser::PropertyCPPTypesRequiringUIRanges = { TEXT("float"), TEXT("double") };
 TMap<UClass*, ClassDefinitionRange> ClassDefinitionRanges;
+
 /**
  * Dirty hack global variable to allow different result codes passed through
  * exceptions. Needs to be fixed in future versions of UHT.
@@ -1090,6 +1094,16 @@ namespace
 				}
 			}
 			break;
+
+			case ECheckedMetadataSpecifier::DocumentationPolicy:
+			{
+				const TCHAR* StrictValue = TEXT("Strict");
+				if (InValue != StrictValue)
+				{
+					FError::Throwf(TEXT("'%s' metadata was '%s' but it must be %s"), *InKey, *InValue, *StrictValue);
+				}
+			}
+			break;
 		}
 	}
 
@@ -1619,6 +1633,7 @@ UEnum* FHeaderParser::CompileEnum()
 
 	// Parse all enums tags.
 	FToken TagToken;
+	TArray<TMap<FName, FString>> EntryMetaData;
 
 	TArray<TPair<FName, int64>> EnumNames;
 	int64 CurrentEnumValue = 0;
@@ -1701,6 +1716,9 @@ UEnum* FHeaderParser::CompileEnum()
 			++CurrentEnumValue;
 		}
 
+		TagToken.MetaData.Add(TEXT("Name"), NewTag.ToString());
+		EntryMetaData.Add(TagToken.MetaData);
+
 		// check for metadata on this enum value
 		ParseFieldMetaData(TagToken.MetaData, TagToken.Identifier);
 		if (TagToken.MetaData.Num() > 0)
@@ -1765,6 +1783,7 @@ UEnum* FHeaderParser::CompileEnum()
 		FError::Throwf(TEXT("Unable to generate enum MAX entry '%s' due to name collision"), *MaxEnumItem.ToString());
 	}
 
+	CheckDocumentationPolicyForEnum(Enum, EnumValueMetaData, EntryMetaData);
 	return Enum;
 }
 
@@ -2024,6 +2043,79 @@ FString FHeaderParser::FormatCommentForToolTip(const FString& Input)
 	return Result;
 }
 
+TMap<FName, FString> FHeaderParser::GetParameterToolTipsFromFunctionComment(const FString& Input)
+{
+	TMap<FName, FString> Map;
+	TArray<FString> Params;
+
+	int32 Offset = 0;
+	int32 ParamStart = -1;
+	int32 ParamLength = 0;
+	FString ParamPrefix;
+	while (Offset < Input.Len())
+	{
+		if (ParamStart == -1)
+		{
+			if (Input.Find(TEXT("@param"), ESearchCase::CaseSensitive, ESearchDir::FromStart, Offset) == Offset)
+			{
+				ParamStart = Offset + 6;
+				Offset += 6;
+				continue;
+			}
+			else if (Input.Find(TEXT("@return"), ESearchCase::CaseSensitive, ESearchDir::FromStart, Offset) == Offset)
+			{
+				ParamStart = Offset + 7;
+				Offset += 7;
+				ParamPrefix = TEXT("ReturnValue ");
+				continue;
+			}
+			Offset++;
+			continue;
+		}
+
+		TCHAR CurrentChar = Input[Offset];
+		if ((CurrentChar == '\n') || (CurrentChar == '\r'))
+		{
+			if (ParamLength > 0)
+			{
+				Params.Add(ParamPrefix + Input.Mid(ParamStart, ParamLength));
+				ParamStart = -1;
+				ParamLength = 0;
+			}
+			ParamPrefix.Reset();
+		}
+		else
+		{
+			ParamLength++;
+		}
+		Offset++;
+	}
+	if (ParamLength > 0)
+	{
+		Params.Add(ParamPrefix + Input.Mid(ParamStart, ParamLength));
+	}
+
+	for (FString Param : Params)
+	{
+		Param = Param.ConvertTabsToSpaces(4).TrimStartAndEnd();
+
+		int32 FirstSpaceIndex = -1;
+		if (!Param.FindChar(' ', FirstSpaceIndex))
+		{
+			continue;
+		}
+
+		FString ParamName = Param.Left(FirstSpaceIndex);
+		FString ParamToolTip = Param.Mid(FirstSpaceIndex + 1);
+		ParamToolTip.TrimStartInline();
+
+		Map.Add(*ParamName, ParamToolTip);
+	}
+
+	return Map;
+}
+
+
 void FHeaderParser::AddFormattedPrevCommentAsTooltipMetaData(TMap<FName, FString>& MetaData)
 {
 	// Don't add a tooltip if one already exists.
@@ -2032,6 +2124,12 @@ void FHeaderParser::AddFormattedPrevCommentAsTooltipMetaData(TMap<FName, FString
 		return;
 	}
 
+	// Add the comment if it is not empty
+	if (!PrevComment.IsEmpty())
+	{
+		MetaData.Add(NAME_Comment, *PrevComment);
+	}
+
 	// Don't add a tooltip if the comment is empty after formatting.
 	FString FormattedComment = FormatCommentForToolTip(PrevComment);
 	if (!FormattedComment.Len())
@@ -2535,6 +2633,7 @@ UScriptStruct* FHeaderParser::CompileStructDeclaration(FClasses& AllClasses)
 	// Link the properties within the struct
 	Struct->StaticLink(true);
 
+	CheckDocumentationPolicyForStruct(Struct, MetaData);
 	return Struct;
 }
 
@@ -5857,6 +5956,7 @@ UClass* FHeaderParser::CompileClassDeclaration(FClasses& AllClasses)
 		}
 	}
 
+	CheckDocumentationPolicyForStruct(Class, MetaData);
 	return Class;
 }
 
@@ -6653,7 +6753,7 @@ void FHeaderParser::CompileFunctionDeclaration(FClasses& AllClasses)
 			{
 				UE_LOG_ERROR_UHT(TEXT("An explicit Category specifier is required for Blueprint accessible functions in an Engine module."));
 			}
-		} 
+		}
 	}
 
 	// Verify interfaces with respect to their blueprint accessible functions
@@ -7044,6 +7144,9 @@ void FHeaderParser::CompileFunctionDeclaration(FClasses& AllClasses)
 		// Put the token back so we can continue parsing as normal
 		UngetToken(Token);
 	}
+
+	// perform documentation policy tests
+	CheckDocumentationPolicyForFunc(GetCurrentClass(), FuncInfo.FunctionReference, MetaData);
 }
 
 /** Parses optional metadata text. */
@@ -9504,6 +9607,274 @@ void FHeaderParser::PostPopNestInterface(FClasses& AllClasses, UClass* CurrentIn
 	}
 }
 
+FDocumentationPolicy FHeaderParser::GetDocumentationPolicyFromName(const FString& PolicyName)
+{
+	FDocumentationPolicy DocumentationPolicy;
+	if (PolicyName.Equals(TEXT("Strict")))
+	{
+		DocumentationPolicy.bClassOrStructCommentRequired = true;
+		DocumentationPolicy.bFunctionToolTipsRequired = true;
+		DocumentationPolicy.bMemberToolTipsRequired = true;
+		DocumentationPolicy.bParameterToolTipsRequired = true;
+		DocumentationPolicy.bFloatRangesRequired = true;
+	}
+	else
+	{
+		FError::Throwf(TEXT("Documentation Policy '%s' not yet supported"), *PolicyName);
+	}
+	return DocumentationPolicy;
+}
+
+
+FDocumentationPolicy FHeaderParser::GetDocumentationPolicyForStruct(UStruct* Struct)
+{
+	check(Struct!= nullptr);
+
+	FDocumentationPolicy DocumentationPolicy;
+	FString DocumentationPolicyName;
+	if (Struct->GetStringMetaDataHierarchical(NAME_DocumentationPolicy, &DocumentationPolicyName))
+	{
+		DocumentationPolicy = GetDocumentationPolicyFromName(DocumentationPolicyName);
+	}
+	return DocumentationPolicy;
+}
+
+void FHeaderParser::CheckDocumentationPolicyForEnum(UEnum* Enum, const TMap<FName, FString>& MetaData, const TArray<TMap<FName, FString>>& Entries)
+{
+	check(Enum != nullptr);
+
+	const FString* DocumentationPolicyName = MetaData.Find(NAME_DocumentationPolicy);
+	if (DocumentationPolicyName == nullptr)
+	{
+		return;
+	}
+
+	check(!DocumentationPolicyName->IsEmpty());
+
+	FDocumentationPolicy DocumentationPolicy = GetDocumentationPolicyFromName(*DocumentationPolicyName);
+	if (DocumentationPolicy.bClassOrStructCommentRequired)
+	{
+		const FString* EnumToolTip = MetaData.Find(NAME_ToolTip);
+		if (EnumToolTip == nullptr)
+		{
+			FError::Throwf(TEXT("Enum '%s' does not provide a tooltip / comment (DocumentationPolicy)."), *Enum->GetName());
+		}
+	}
+
+	TMap<FString, FString> ToolTipToEntry;
+	for (const TMap<FName, FString>& Entry : Entries)
+	{
+		const FString* EntryName = Entry.Find(TEXT("Name"));
+		if (EntryName == nullptr)
+		{
+			continue;
+		}
+
+		const FString* ToolTip = Entry.Find(NAME_ToolTip);
+		if (ToolTip == nullptr)
+		{
+			FError::Throwf(TEXT("Enum entry '%s::%s' does not provide a tooltip / comment (DocumentationPolicy)."), *Enum->GetName(), **EntryName);
+		}
+
+		const FString* ExistingEntry = ToolTipToEntry.Find(*ToolTip);
+		if (ExistingEntry != nullptr)
+		{
+			FError::Throwf(TEXT("Enum entries '%s::%s' and '%s::%s' have identical tooltips / comments (DocumentationPolicy)."), *Enum->GetName(), **ExistingEntry, *Enum->GetName(), **EntryName);
+		}
+		ToolTipToEntry.Add(*ToolTip, *EntryName);
+	}
+	
+}
+
+void FHeaderParser::CheckDocumentationPolicyForStruct(UStruct* Struct, const TMap<FName, FString>& MetaData)
+{
+	check(Struct != nullptr);
+
+	FDocumentationPolicy DocumentationPolicy = GetDocumentationPolicyForStruct(Struct);
+	if (DocumentationPolicy.bClassOrStructCommentRequired)
+	{
+		const FString* ClassTooltip = MetaData.Find(NAME_ToolTip);
+		if (ClassTooltip == nullptr)
+		{
+			FError::Throwf(TEXT("Struct '%s' does not provide a tooltip / comment (DocumentationPolicy)."), *Struct->GetName());
+		}
+	}
+
+	if (DocumentationPolicy.bMemberToolTipsRequired)
+	{
+		TMap<FString, FName> ToolTipToPropertyName;
+		for (UProperty* Property : TFieldRange<UProperty>(Struct))
+		{
+			FString ToolTip = Property->GetToolTipText().ToString();
+			if (ToolTip.IsEmpty() || ToolTip.Equals(Property->GetName()))
+			{
+				FError::Throwf(TEXT("Property '%s::%s' does not provide a tooltip / comment (DocumentationPolicy)."), *Struct->GetName(), *Property->GetName());
+			}
+			const FName* ExistingPropertyName = ToolTipToPropertyName.Find(ToolTip);
+			if (ExistingPropertyName != nullptr)
+			{
+				FError::Throwf(TEXT("Property '%s::%s' and '%s::%s' are using identical tooltips (DocumentationPolicy)."), *Struct->GetName(), *ExistingPropertyName->ToString(), *Struct->GetName(), *Property->GetName());
+			}
+			ToolTipToPropertyName.Add(ToolTip, Property->GetFName());
+		}
+	}
+
+	if (DocumentationPolicy.bFloatRangesRequired)
+	{
+		for (UProperty* Property : TFieldRange<UProperty>(Struct))
+		{
+			if(DoesCPPTypeRequireDocumentation(Property->GetCPPType()))
+			{
+				const FString& UIMin = Property->GetMetaData(TEXT("UIMin"));
+				const FString& UIMax = Property->GetMetaData(TEXT("UIMax"));
+				if(!CheckUIMinMaxRangeFromMetaData(UIMin, UIMax))
+				{
+					FError::Throwf(TEXT("Property '%s::%s' does not provide a valid UIMin / UIMax (DocumentationPolicy)."), *Struct->GetName(), *Property->GetName());
+				}
+			}
+		}
+	}
+
+	// also compare all tooltips to see if they are unique
+	if (DocumentationPolicy.bFunctionToolTipsRequired)
+	{
+		UClass* Class = Cast<UClass>(Struct);
+		if (Class != nullptr)
+		{
+			TMap<FString, FString> ToolTipToFunc;
+			for (UField* Field = Class->Children; Field; Field = Field->Next)
+			{
+				UFunction* Func = Cast<UFunction>(Field);
+				if (Func != NULL)
+				{
+					FString ToolTip = Func->GetToolTipText().ToString();
+					if (ToolTip.IsEmpty())
+					{
+						FError::Throwf(TEXT("Function '%s::%s' does not provide a tooltip / comment (DocumentationPolicy)."), *Class->GetName(), *Func->GetName());
+					}
+					const FString* ExistingFuncName = ToolTipToFunc.Find(ToolTip);
+					if (ExistingFuncName != nullptr)
+					{
+						FError::Throwf(TEXT("Functions '%s::%s' and '%s::%s' uses identical tooltips / comments (DocumentationPolicy)."), *Class->GetName(), **ExistingFuncName, *Class->GetName(), *Func->GetName());
+					}
+					ToolTipToFunc.Add(ToolTip, Func->GetName());
+				}
+			}
+		}
+	}
+}
+
+bool FHeaderParser::DoesCPPTypeRequireDocumentation(const FString& CPPType)
+{
+	return PropertyCPPTypesRequiringUIRanges.Find(CPPType) != INDEX_NONE;
+}
+
+// Validates the documentation for a given method
+void FHeaderParser::CheckDocumentationPolicyForFunc(UClass* Class, UFunction* Func, const TMap<FName, FString>& MetaData)
+{
+	check(Class != nullptr);
+	check(Func != nullptr);
+
+	FDocumentationPolicy DocumentationPolicy = GetDocumentationPolicyForStruct(Class);
+	if (DocumentationPolicy.bFunctionToolTipsRequired)
+	{
+		const FString* FunctionTooltip = MetaData.Find(NAME_ToolTip);
+		if (FunctionTooltip == nullptr)
+		{
+			FError::Throwf(TEXT("Function '%s::%s' does not provide a tooltip / comment (DocumentationPolicy)."), *Class->GetName(), *Func->GetName());
+		}
+	}
+
+	if (DocumentationPolicy.bParameterToolTipsRequired)
+	{
+		const FString* FunctionComment = MetaData.Find(NAME_Comment);
+		if (FunctionComment == nullptr)
+		{
+			FError::Throwf(TEXT("Function '%s::%s' does not provide a comment (DocumentationPolicy)."), *Class->GetName(), *Func->GetName());
+		}
+		
+		TMap<FName, FString> ParamToolTips = GetParameterToolTipsFromFunctionComment(*FunctionComment);
+		bool HasAnyParamToolTips = ParamToolTips.Num() > 0;
+		if (ParamToolTips.Num() == 0)
+		{
+			const FString* ReturnValueToolTip = ParamToolTips.Find(TEXT("ReturnValue"));
+			if (ReturnValueToolTip != nullptr)
+			{
+				HasAnyParamToolTips = false;
+			}
+		}
+
+		// only apply the validation for parameter tooltips if a function has any @param statements at all.
+		if (HasAnyParamToolTips)
+		{
+			// ensure each parameter has a tooltip
+			TSet<FName> ExistingFields;
+			for (UProperty* Property : TFieldRange<UProperty>(Func))
+			{
+				FName ParamName = Property->GetFName();
+				if (ParamName.IsEqual(TEXT("ReturnValue")))
+				{
+					continue;
+				}
+				const FString* ParamToolTip = ParamToolTips.Find(ParamName);
+				if (ParamToolTip == nullptr)
+				{
+					FError::Throwf(TEXT("Function '%s::%s' doesn't provide a tooltip for parameter '%s' (DocumentationPolicy)."), *Class->GetName(), *Func->GetName(), *ParamName.ToString());
+				}
+				ExistingFields.Add(ParamName);
+			}
+
+			// ensure we don't have parameter tooltips for parameters that don't exist
+			for (TPair<FName, FString>& Pair : ParamToolTips)
+			{
+				const FName& ParamName = Pair.Key;
+				if (ParamName.IsEqual(TEXT("ReturnValue")))
+				{
+					continue;
+				}
+				if (!ExistingFields.Contains(ParamName))
+				{
+					FError::Throwf(TEXT("Function '%s::%s' provides a tooltip for an unknown parameter '%s' (DocumentationPolicy)."), *Class->GetName(), *Func->GetName(), *Pair.Key.ToString());
+				}
+			}
+
+			// check for duplicate tooltips
+			TMap<FString, FName> ToolTipToParam;
+			for (TPair<FName, FString>& Pair : ParamToolTips)
+			{
+				const FName& ParamName = Pair.Key;
+				if (ParamName.IsEqual(TEXT("ReturnValue")))
+				{
+					continue;
+				}
+				const FName* ExistingParam = ToolTipToParam.Find(Pair.Value);
+				if (ExistingParam != nullptr)
+				{
+					FError::Throwf(TEXT("Function '%s::%s' uses identical tooltips for parameters '%s' and '%s' (DocumentationPolicy)."), *Class->GetName(), *Func->GetName(), *ExistingParam->ToString(), *Pair.Key.ToString());
+				}
+				ToolTipToParam.Add(Pair.Value, Pair.Key);
+			}
+		}
+	}
+}
+
+bool FHeaderParser::CheckUIMinMaxRangeFromMetaData(const FString& UIMin, const FString& UIMax)
+{
+	if (UIMin.IsEmpty() || UIMax.IsEmpty())
+	{
+		return false;
+	}
+
+	double UIMinValue = FCString::Atod(*UIMin);
+	double UIMaxValue = FCString::Atod(*UIMax);
+	if (UIMin > UIMax) // note that we actually allow UIMin == UIMax to disable the range manually.
+	{
+		return false;
+	}
+
+	return true;
+}
+
 template <class TFunctionType>
 TFunctionType* CreateFunctionImpl(const FFuncInfo& FuncInfo, UObject* Outer, FScope* CurrentScope)
 {
@@ -9549,3 +9920,4 @@ UDelegateFunction* FHeaderParser::CreateDelegateFunction(const FFuncInfo &FuncIn
 	UObject* CurrentPackage = LocSourceFile ? LocSourceFile->GetPackage() : nullptr;
 	return CreateFunctionImpl<T>(FuncInfo, IsInAClass() ? (UObject*)GetCurrentClass() : CurrentPackage, GetCurrentScope());
 }
+
diff --git a/Engine/Source/Programs/UnrealHeaderTool/Private/HeaderParser.h b/Engine/Source/Programs/UnrealHeaderTool/Private/HeaderParser.h
--- a/Engine/Source/Programs/UnrealHeaderTool/Private/HeaderParser.h
+++ b/Engine/Source/Programs/UnrealHeaderTool/Private/HeaderParser.h
@@ -147,6 +147,15 @@ struct ClassDefinitionRange
 
 extern TMap<UClass*, ClassDefinitionRange> ClassDefinitionRanges;
 
+struct FDocumentationPolicy
+{
+	bool bClassOrStructCommentRequired = false;
+	bool bFunctionToolTipsRequired = false;
+	bool bMemberToolTipsRequired = false;
+	bool bParameterToolTipsRequired = false;
+	bool bFloatRangesRequired = false;
+};
+
 /////////////////////////////////////////////////////
 // FHeaderParser
 
@@ -469,6 +478,13 @@ protected:
 	 * @return					The input string, reformatted in such a way as to be appropriate for use as a tooltip.
 	 */
 	static FString FormatCommentForToolTip(const FString& Input);
+
+	/**
+	 * Retrieves parameter comments / tooltips from the function comment
+	 * @param		Input		An input string, expected to be a script comment.
+	 * @return					The map of parameter name to comment per parameter
+	*/
+	static TMap<FName, FString> GetParameterToolTipsFromFunctionComment(const FString& Input);
 	
 	/**
 	 * Begins the process of exporting C++ class declarations for native classes in the specified package
@@ -741,6 +757,30 @@ private:
 
 	// Verifies the target function meets the criteria for a replication notify callback
 	void VerifyRepNotifyCallback(UProperty* Property, UFunction* TargetFunction);
+
+	// Constructs the policy from a string
+	static FDocumentationPolicy GetDocumentationPolicyFromName(const FString& PolicyName);
+
+	// Constructs the policy for documentation checks for a given struct
+	static FDocumentationPolicy GetDocumentationPolicyForStruct(UStruct* Struct);
+
+	// Property types to provide UI Min and Max ranges
+	static TArray<FString> PropertyCPPTypesRequiringUIRanges;
+
+	// Returns true if a given CPP types required ui checking
+	static bool DoesCPPTypeRequireDocumentation(const FString& CPPType);
+
+	// Validates the documentation for a given enum
+	void CheckDocumentationPolicyForEnum(UEnum* Enum, const TMap<FName, FString>& MetaData, const TArray<TMap<FName, FString>>& Entries);
+
+	// Validates the documentation for a given struct
+	void CheckDocumentationPolicyForStruct(UStruct* Struct, const TMap<FName, FString>& MetaData);
+
+	// Validates the documentation for a given method
+	void CheckDocumentationPolicyForFunc(UClass* Class, UFunction* Func, const TMap<FName, FString>& MetaData);
+
+	// Checks if a valid range has been found on the provided metadata
+	bool CheckUIMinMaxRangeFromMetaData(const FString& UIMin, const FString& UIMax);
 };
 
 /////////////////////////////////////////////////////
diff --git a/Engine/Source/Programs/UnrealHeaderTool/Private/Specifiers/CheckedMetadataSpecifiers.def b/Engine/Source/Programs/UnrealHeaderTool/Private/Specifiers/CheckedMetadataSpecifiers.def
--- a/Engine/Source/Programs/UnrealHeaderTool/Private/Specifiers/CheckedMetadataSpecifiers.def
+++ b/Engine/Source/Programs/UnrealHeaderTool/Private/Specifiers/CheckedMetadataSpecifiers.def
@@ -7,6 +7,7 @@ CHECKED_METADATA_SPECIFIER(ClampMax)
 CHECKED_METADATA_SPECIFIER(ClampMin)
 CHECKED_METADATA_SPECIFIER(CommutativeAssociativeBinaryOperator)
 CHECKED_METADATA_SPECIFIER(DevelopmentStatus)
+CHECKED_METADATA_SPECIFIER(DocumentationPolicy)
 CHECKED_METADATA_SPECIFIER(ExpandEnumAsExecs)
 CHECKED_METADATA_SPECIFIER(UIMax)
 CHECKED_METADATA_SPECIFIER(UIMin)
diff --git a/Engine/Source/Runtime/AnimGraphRuntime/Private/AnimNodes/AnimNode_MultiWayBlend.cpp b/Engine/Source/Runtime/AnimGraphRuntime/Private/AnimNodes/AnimNode_MultiWayBlend.cpp
--- a/Engine/Source/Runtime/AnimGraphRuntime/Private/AnimNodes/AnimNode_MultiWayBlend.cpp
+++ b/Engine/Source/Runtime/AnimGraphRuntime/Private/AnimNodes/AnimNode_MultiWayBlend.cpp
@@ -136,6 +136,8 @@ void FAnimNode_MultiWayBlend::Evaluate_AnyThread(FPoseContext& Output)
 	if (SourcePoses.Num() > 0)
 	{
 		FAnimationRuntime::BlendPosesTogether(SourcePoses, SourceCurves, SourceWeights, Output.Pose, Output.Curve);
+		// normalize rotation - some cases, where additive is applied less than 1, it will use non normalized rotation
+		Output.Pose.NormalizeRotations();
 	}
 	else
 	{
diff --git a/Engine/Source/Runtime/AnimGraphRuntime/Private/AnimSequencerInstanceProxy.cpp b/Engine/Source/Runtime/AnimGraphRuntime/Private/AnimSequencerInstanceProxy.cpp
--- a/Engine/Source/Runtime/AnimGraphRuntime/Private/AnimSequencerInstanceProxy.cpp
+++ b/Engine/Source/Runtime/AnimGraphRuntime/Private/AnimSequencerInstanceProxy.cpp
@@ -144,6 +144,13 @@ void FAnimSequencerInstanceProxy::UpdateAnimTrack(UAnimSequenceBase* InAnimSeque
 	// if moving to 0.f, we mark this to teleport. Otherwise, do not use explicit time
 	FAnimNode_MultiWayBlend& BlendNode = (PlayerState->bAdditive) ? AdditiveBlendNode : FullBodyBlendNode;
 	BlendNode.DesiredAlphas[PlayerState->PoseIndex] = Weight;
+
+	// if additive, apply alpha value correctly
+	// this will be used when apply additive is blending correct total alpha to additive
+	if (PlayerState->bAdditive)
+	{
+		SequencerRootNode.Alpha = BlendNode.GetTotalAlpha();
+	}
 }
 
 void FAnimSequencerInstanceProxy::EnsureAnimTrack(UAnimSequenceBase* InAnimSequence, uint32 SequenceId)
diff --git a/Engine/Source/Runtime/AnimGraphRuntime/Public/AnimNodes/AnimNode_MultiWayBlend.h b/Engine/Source/Runtime/AnimGraphRuntime/Public/AnimNodes/AnimNode_MultiWayBlend.h
--- a/Engine/Source/Runtime/AnimGraphRuntime/Public/AnimNodes/AnimNode_MultiWayBlend.h
+++ b/Engine/Source/Runtime/AnimGraphRuntime/Public/AnimNodes/AnimNode_MultiWayBlend.h
@@ -69,8 +69,8 @@ public:
 		Poses.Reset();
 		DesiredAlphas.Reset();
 	}
-private:
-	float GetTotalAlpha()
+
+	float GetTotalAlpha() const
 	{
 		float TotalAlpha = 0.f;
 
@@ -82,6 +82,7 @@ private:
 		return TotalAlpha;
 	}
 
+private:
 	// process new weights and then return out
 	void UpdateCachedAlphas();
 };
diff --git a/Engine/Source/Runtime/AnimGraphRuntime/Public/BoneControllers/AnimNode_AnimDynamics.h b/Engine/Source/Runtime/AnimGraphRuntime/Public/BoneControllers/AnimNode_AnimDynamics.h
--- a/Engine/Source/Runtime/AnimGraphRuntime/Public/BoneControllers/AnimNode_AnimDynamics.h
+++ b/Engine/Source/Runtime/AnimGraphRuntime/Public/BoneControllers/AnimNode_AnimDynamics.h
@@ -285,15 +285,19 @@ struct ANIMGRAPHRUNTIME_API FAnimNode_AnimDynamics : public FAnimNode_SkeletalCo
 	UPROPERTY(EditAnywhere, BlueprintReadWrite, Category = Setup, meta = (PinHiddenByDefault, EditCondition = "bUseGravityOverride"))
 	FVector GravityOverride;
 
-	/** Spring constant to use when calculating linear springs, higher values mean a stronger spring.*/
-	UPROPERTY(EditAnywhere, BlueprintReadWrite, Category = Setup, meta = (PinHiddenByDefault, DisplayAfter="bAngularSpring"))
+	/** 
+	 * Spring constant to use when calculating linear springs, higher values mean a stronger spring.
+	 * You need to enable the Linear Spring checkbox for this to have an effect.
+	 */
+	UPROPERTY(EditAnywhere, BlueprintReadWrite, Category = Setup, meta = (PinHiddenByDefault, DisplayAfter="bAngularSpring", EditCondition = "bLinearSpring"))
 	float LinearSpringConstant;
 
 	/** 
 	 * Spring constant to use when calculating angular springs, higher values mean a stronger spring.
+	 * You need to enable the Angular Spring checkbox for this to have an effect.
 	 * Note: Make sure to also set the Angular Target Axis and Angular Target in the Constraint Setup for this to have an effect.
 	 */
-	UPROPERTY(EditAnywhere, BlueprintReadWrite, Category = Setup, meta = (PinHiddenByDefault))
+	UPROPERTY(EditAnywhere, BlueprintReadWrite, Category = Setup, meta = (PinHiddenByDefault, EditCondition = "bAngularSpring"))
 	float AngularSpringConstant;
 
 	/** Scale to apply to calculated wind velocities in the solver */
diff --git a/Engine/Source/Runtime/AnimGraphRuntime/Public/KismetAnimationLibrary.h b/Engine/Source/Runtime/AnimGraphRuntime/Public/KismetAnimationLibrary.h
--- a/Engine/Source/Runtime/AnimGraphRuntime/Public/KismetAnimationLibrary.h
+++ b/Engine/Source/Runtime/AnimGraphRuntime/Public/KismetAnimationLibrary.h
@@ -14,7 +14,7 @@ class USkeletalMeshComponent;
 /**
  *	A library of the most common animation blueprint functions.
  */
-UCLASS(meta=(ScriptName="AnimGraphLibrary"))
+UCLASS(meta=(ScriptName="AnimGraphLibrary", DocumentationPolicy="Strict"))
 class ANIMGRAPHRUNTIME_API UKismetAnimationLibrary : public UBlueprintFunctionLibrary
 {
 	GENERATED_UCLASS_BODY()
diff --git a/Engine/Source/Runtime/AnimGraphRuntime/Public/KismetAnimationTypes.h b/Engine/Source/Runtime/AnimGraphRuntime/Public/KismetAnimationTypes.h
--- a/Engine/Source/Runtime/AnimGraphRuntime/Public/KismetAnimationTypes.h
+++ b/Engine/Source/Runtime/AnimGraphRuntime/Public/KismetAnimationTypes.h
@@ -13,7 +13,7 @@
  *	The FPositionArray also tracks the last index used to allow for 
  *	reuse of entries (instead of appending to the array all the time).
  */
-USTRUCT(BlueprintType)
+USTRUCT(BlueprintType, meta = (DocumentationPolicy = "Strict"))
 struct FPositionHistory
 {
 	GENERATED_BODY()
diff --git a/Engine/Source/Runtime/CoreUObject/Public/UObject/ObjectMacros.h b/Engine/Source/Runtime/CoreUObject/Public/UObject/ObjectMacros.h
--- a/Engine/Source/Runtime/CoreUObject/Public/UObject/ObjectMacros.h
+++ b/Engine/Source/Runtime/CoreUObject/Public/UObject/ObjectMacros.h
@@ -987,6 +987,9 @@ namespace UM
 
 		/// A short tooltip that is used in some contexts where the full tooltip might be overwhelming (such as the parent class picker dialog)
 		ShortTooltip,
+
+		/// A setting to determine validation of tooltips and comments. Needs to be set to "Strict"
+		DocumentationPolicy,
 	};
 
 	// Metadata usable in UCLASS
diff --git a/Engine/Source/Runtime/Engine/Classes/Animation/AnimNode_SequencePlayer.h b/Engine/Source/Runtime/Engine/Classes/Animation/AnimNode_SequencePlayer.h
--- a/Engine/Source/Runtime/Engine/Classes/Animation/AnimNode_SequencePlayer.h
+++ b/Engine/Source/Runtime/Engine/Classes/Animation/AnimNode_SequencePlayer.h
@@ -18,7 +18,7 @@ struct ENGINE_API FAnimNode_SequencePlayer : public FAnimNode_AssetPlayerBase
 	GENERATED_USTRUCT_BODY()
 public:
 	// The animation sequence asset to play
-	UPROPERTY(EditAnywhere, BlueprintReadWrite, Category = Settings, meta = (PinHiddenByDefault))
+	UPROPERTY(EditAnywhere, BlueprintReadWrite, Category = Settings, meta = (PinHiddenByDefault, DisallowedClasses="AnimMontage"))
 	UAnimSequenceBase* Sequence;
 
 	// The Basis in which the PlayRate is expressed in. This is used to rescale PlayRate inputs.
diff --git a/Engine/Source/Runtime/Engine/Classes/Animation/AnimNotifies/AnimNotifyState_TimedParticleEffect.h b/Engine/Source/Runtime/Engine/Classes/Animation/AnimNotifies/AnimNotifyState_TimedParticleEffect.h
--- a/Engine/Source/Runtime/Engine/Classes/Animation/AnimNotifies/AnimNotifyState_TimedParticleEffect.h
+++ b/Engine/Source/Runtime/Engine/Classes/Animation/AnimNotifies/AnimNotifyState_TimedParticleEffect.h
@@ -13,8 +13,8 @@ class USkeletalMeshComponent;
 // Timed Particle Effect Notify
 // Allows a looping particle effect to be played in an animation that will activate
 // at the beginning of the notify and deactivate at the end.
-UCLASS(MinimalAPI, Blueprintable, meta = (DisplayName = "Timed Particle Effect"))
-class UAnimNotifyState_TimedParticleEffect : public UAnimNotifyState
+UCLASS(Blueprintable, meta = (DisplayName = "Timed Particle Effect"))
+class ENGINE_API UAnimNotifyState_TimedParticleEffect : public UAnimNotifyState
 {
 	GENERATED_UCLASS_BODY()
 
@@ -64,6 +64,6 @@ class UAnimNotifyState_TimedParticleEffect : public UAnimNotifyState
 	// Overridden from UAnimNotifyState to provide custom notify name.
 	FString GetNotifyName_Implementation() const override;
 
-private:
+protected:
 	bool ValidateParameters(USkeletalMeshComponent* MeshComp);
 };
diff --git a/Engine/Source/Runtime/Engine/Classes/Animation/BlendSpaceBase.h b/Engine/Source/Runtime/Engine/Classes/Animation/BlendSpaceBase.h
--- a/Engine/Source/Runtime/Engine/Classes/Animation/BlendSpaceBase.h
+++ b/Engine/Source/Runtime/Engine/Classes/Animation/BlendSpaceBase.h
@@ -475,6 +475,9 @@ protected:
 	UPROPERTY(EditAnywhere, Category = BlendParametersTest)
 	struct FBlendParameter BlendParameters[3];
 
+	/** Reset to reference pose. It does apply different refpose based on additive or not*/
+	void ResetToRefPose(FCompactPose& OutPose);
+
 #if WITH_EDITOR
 private:
 	// Track whether we have updated markers so cached data can be updated
diff --git a/Engine/Source/Runtime/Engine/Classes/Components/PoseableMeshComponent.h b/Engine/Source/Runtime/Engine/Classes/Components/PoseableMeshComponent.h
--- a/Engine/Source/Runtime/Engine/Classes/Components/PoseableMeshComponent.h
+++ b/Engine/Source/Runtime/Engine/Classes/Components/PoseableMeshComponent.h
@@ -18,8 +18,15 @@ class ENGINE_API UPoseableMeshComponent : public USkinnedMeshComponent
 {
 	GENERATED_UCLASS_BODY()
 
+	// we need this for template function in the cpp
+	// skeletalmeshcomponent hides this property, so now creating extra getter for poseablemeshcomponent
+	// although it's fine for poseablemeshcomponent to have this as external
+	// note that their signature is different, but the template would workf ine
+	const TArray<FTransform>& GetBoneSpaceTransforms() const { return BoneSpaceTransforms; }
+
 	/** Temporary array of local-space (ie relative to parent bone) rotation/translation/scale for each bone. */
 	TArray<FTransform> BoneSpaceTransforms;
+public:
 
 	FBoneContainer RequiredBones;
 
@@ -51,7 +58,7 @@ class ENGINE_API UPoseableMeshComponent : public USkinnedMeshComponent
 	void ResetBoneTransformByName(FName BoneName);
 
 	UFUNCTION(BlueprintCallable, Category="Components|PoseableMesh")
-	void CopyPoseFromSkeletalComponent(const USkeletalMeshComponent* InComponentToCopy);
+	void CopyPoseFromSkeletalComponent(USkeletalMeshComponent* InComponentToCopy);
 
 	//~ Begin USkinnedMeshComponent Interface
 	virtual void RefreshBoneTransforms(FActorComponentTickFunction* TickFunction = NULL) override;
diff --git a/Engine/Source/Runtime/Engine/Classes/Components/SkeletalMeshComponent.h b/Engine/Source/Runtime/Engine/Classes/Components/SkeletalMeshComponent.h
--- a/Engine/Source/Runtime/Engine/Classes/Components/SkeletalMeshComponent.h
+++ b/Engine/Source/Runtime/Engine/Classes/Components/SkeletalMeshComponent.h
@@ -339,9 +339,18 @@ public:
 	UPROPERTY(EditAnywhere, BlueprintReadWrite, Category=Animation, meta=(ShowOnlyInnerProperties))
 	struct FSingleAnimationPlayData AnimationData;
 
-	/** Temporary array of local-space (relative to parent bone) rotation/translation for each bone. */
+	// this is explicit copy because this buffer is reused during evaluation
+	// we want to have reference and emptied during evaluation
+	TArray<FTransform> GetBoneSpaceTransforms();
+
+	/** 
+	 * Temporary array of local-space (relative to parent bone) rotation/translation for each bone. 
+	 * This property is not safe to access during evaluation, so we created wrapper.
+	 */
+	UE_DEPRECATED(4.23, "Direct access to this property is deprecated, please use GetBoneSpaceTransforms instead. We will move to private in the future.")
 	TArray<FTransform> BoneSpaceTransforms;
-	
+
+public:
 	/** Offset of the root bone from the reference pose. Used to offset bounding box. */
 	UPROPERTY(transient)
 	FVector RootBoneTranslation;
@@ -1242,6 +1251,11 @@ private:
 	 **/
 	TMap<FName, float>	MorphTargetCurves;
 
+	/** 
+	 * Temporary storage for Curve UIDList of evaluating Animation 
+	 */
+	TArray<uint16> CachedCurveUIDList;
+
 public:
 	const TMap<FName, float>& GetMorphTargetCurves() const { return MorphTargetCurves;  }
 	// 
diff --git a/Engine/Source/Runtime/Engine/Classes/Engine/MaterialMerging.h b/Engine/Source/Runtime/Engine/Classes/Engine/MaterialMerging.h
--- a/Engine/Source/Runtime/Engine/Classes/Engine/MaterialMerging.h
+++ b/Engine/Source/Runtime/Engine/Classes/Engine/MaterialMerging.h
@@ -32,105 +32,115 @@ struct FMaterialProxySettings
 	// Size of generated BaseColor map
 	UPROPERTY(Category = Material, BlueprintReadWrite, EditAnywhere, meta =(ClampMin = "1", UIMin = "1"))
 	FIntPoint TextureSize;
-
+	
+	// Gutter space to take into account 
 	UPROPERTY(Category = Material, BlueprintReadWrite, AdvancedDisplay, EditAnywhere, meta=(DisplayAfter="TextureSizingType"))
 	float GutterSpace;
 
-	// Metallic constant
+	// Constant value to use for the Metallic property
 	UPROPERTY(Category = Material, BlueprintReadWrite, EditAnywhere, meta = (DisplayAfter="bMetallicMap", ClampMin = "0", ClampMax = "1", UIMin = "0", UIMax = "1", editcondition = "!bMetallicMap"))
 	float MetallicConstant;
 
-	// Roughness constant
+	// Constant value to use for the Roughness property
 	UPROPERTY(Category = Material, BlueprintReadWrite, EditAnywhere, meta = (DisplayAfter="bRoughnessMap", ClampMin = "0", ClampMax = "1", UIMin = "0", UIMax = "1", editcondition = "!bRoughnessMap"))
 	float RoughnessConstant;
 
-	// Specular constant
+	// Constant value to use for the Specular property
 	UPROPERTY(Category = Material, BlueprintReadWrite, EditAnywhere, meta = (DisplayAfter="bSpecularMap", ClampMin = "0", ClampMax = "1", UIMin = "0", UIMax = "1", editcondition = "!bSpecularMap"))
 	float SpecularConstant;
 
+	// Constant value to use for the Opacity property
 	UPROPERTY(Category = Material, BlueprintReadWrite, EditAnywhere, meta = (DisplayAfter="bOpacityMap", ClampMin = "0", ClampMax = "1", UIMin = "0", UIMax = "1", editcondition = "!bOpacityMap"))
 	float OpacityConstant;
 
+	// Constant value to use for the Opacity mask property
 	UPROPERTY(Category = Material, BlueprintReadWrite, EditAnywhere, meta = (DisplayAfter="bOpacityMaskMap", ClampMin = "0", ClampMax = "1", UIMin = "0", UIMax = "1", editcondition = "!bOpacityMaskMap"))
 	float OpacityMaskConstant;
 
+	// Constant value to use for the Ambient Occlusion property
 	UPROPERTY(Category = Material, BlueprintReadWrite, EditAnywhere, meta = (DisplayAfter="bAmbientOcclusionMap", ClampMin = "0", ClampMax = "1", UIMin = "0", UIMax = "1", editcondition = "!bAmbientOcclusionMap"))
 	float AmbientOcclusionConstant;
 
+	// Method that should be used to generate the sizes of the output textures
 	UPROPERTY(Category = Material, BlueprintReadWrite, EditAnywhere)
 	TEnumAsByte<ETextureSizingType> TextureSizingType;
 
 	UPROPERTY()
 	TEnumAsByte<EMaterialMergeType> MaterialMergeType;
 
+	// Target blend mode for the generated material
 	UPROPERTY(Category = Material, BlueprintReadWrite, AdvancedDisplay, EditAnywhere, meta=(DisplayAfter="AmbientOcclusionTextureSize"))
 	TEnumAsByte<EBlendMode> BlendMode;
 
+	// Whether or not to allow the generated material can be two-sided
 	UPROPERTY(Category = Material, BlueprintReadWrite, EditAnywhere, meta = (DisplayAfter = "BlendMode"))
 	uint8 bAllowTwoSidedMaterial : 1;
 
-	// Whether to generate normal map
+	// Whether to generate a texture for the Normal property
 	UPROPERTY(Category = Material, BlueprintReadWrite, EditAnywhere)
 	uint8 bNormalMap:1;
 
-	// Whether to generate metallic map
+	// Whether to generate a texture for the Metallic property
 	UPROPERTY(Category = Material, BlueprintReadWrite, EditAnywhere)
 	uint8 bMetallicMap:1;
 
-	// Whether to generate roughness map
+	// Whether to generate a texture for the Roughness property
 	UPROPERTY(Category = Material, BlueprintReadWrite, EditAnywhere)
 	uint8 bRoughnessMap:1;
 
-	// Whether to generate specular map
+	// Whether to generate a texture for the Specular property
 	UPROPERTY(Category = Material, BlueprintReadWrite, EditAnywhere)
 	uint8 bSpecularMap:1;
 
-	// Whether to generate emissive map
+	// Whether to generate a texture for the Emissive property
 	UPROPERTY(Category = Material, BlueprintReadWrite, EditAnywhere)
 	uint8 bEmissiveMap:1;
 
-	// Whether to generate opacity map
+	// Whether to generate a texture for the Opacity property
 	UPROPERTY(Category = Material, BlueprintReadWrite, EditAnywhere)
 	uint8 bOpacityMap:1;
 
+	// Whether to generate a texture for the Opacity Mask property
 	UPROPERTY(Category = Material, BlueprintReadWrite, EditAnywhere)
 	uint8 bOpacityMaskMap:1;
 
+	// Whether to generate a texture for the Ambient Occlusion property
 	UPROPERTY(Category = Material, BlueprintReadWrite, EditAnywhere)
 	uint8 bAmbientOcclusionMap:1;
 
-	// Override diffuse map size
+	// Override Diffuse texture size
 	UPROPERTY(Category = Material, BlueprintReadWrite, AdvancedDisplay, EditAnywhere, meta = (ClampMin = "1", UIMin = "1"))
 	FIntPoint DiffuseTextureSize;
 
-	// Override normal map size
+	// Override Normal texture size
 	UPROPERTY(Category = Material, BlueprintReadWrite, AdvancedDisplay, EditAnywhere, meta = (ClampMin = "1", UIMin = "1"))
 	FIntPoint NormalTextureSize;
 
-	// Override metallic map size
+	// Override Metallic texture size
 	UPROPERTY(Category = Material, BlueprintReadWrite, AdvancedDisplay, EditAnywhere, meta = (ClampMin = "1", UIMin = "1"))
 	FIntPoint MetallicTextureSize;
 
-	// Override roughness map size
+	// Override Roughness texture size
 	UPROPERTY(Category = Material, BlueprintReadWrite, AdvancedDisplay, EditAnywhere, meta = (ClampMin = "1", UIMin = "1"))
 	FIntPoint RoughnessTextureSize;
 
-	// Override specular map size
+	// Override Specular texture size
 	UPROPERTY(Category = Material, BlueprintReadWrite, AdvancedDisplay, EditAnywhere, meta = (ClampMin = "1", UIMin = "1"))
 	FIntPoint SpecularTextureSize;
 
-	// Override emissive map size
+	// Override Emissive texture size
 	UPROPERTY(Category = Material, BlueprintReadWrite, AdvancedDisplay, EditAnywhere, meta = (ClampMin = "1", UIMin = "1"))
 	FIntPoint EmissiveTextureSize;
 
-	// Override opacity map size
+	// Override Opacity texture size
 	UPROPERTY(Category = Material, BlueprintReadWrite, AdvancedDisplay, EditAnywhere, meta = (ClampMin = "1", UIMin = "1"))
 	FIntPoint OpacityTextureSize;
 	
-	// Override opacity map size
+	// Override Opacity Mask texture size
 	UPROPERTY(Category = Material, BlueprintReadWrite, AdvancedDisplay, EditAnywhere, meta = (ClampMin = "1", UIMin = "1"))
 	FIntPoint OpacityMaskTextureSize;
 
+	// Override Ambient Occlusion texture size
 	UPROPERTY(Category = Material, BlueprintReadWrite, AdvancedDisplay, EditAnywhere, meta = (ClampMin = "1", UIMin = "1"))
 	FIntPoint AmbientOcclusionTextureSize;
 
diff --git a/Engine/Source/Runtime/Engine/Classes/Engine/MeshMerging.h b/Engine/Source/Runtime/Engine/Classes/Engine/MeshMerging.h
--- a/Engine/Source/Runtime/Engine/Classes/Engine/MeshMerging.h
+++ b/Engine/Source/Runtime/Engine/Classes/Engine/MeshMerging.h
@@ -425,7 +425,7 @@ struct FMeshMergingSettings
 {
 	GENERATED_USTRUCT_BODY()
 
-	/** Target lightmap resolution */
+	/** The lightmap resolution used both for generating lightmap UV coordinates, and also set on the generated static mesh */
 	UPROPERTY(EditAnywhere, AdvancedDisplay, BlueprintReadWrite, Category = MeshSettings, meta=(ClampMax = 4096, EditCondition = "!bComputedLightMapResolution", DisplayAfter="bGenerateLightMapUV"))
 	int32 TargetLightMapResolution;
 
@@ -445,6 +445,7 @@ struct FMeshMergingSettings
 	UPROPERTY(EditAnywhere, Category = MeshSettings, BlueprintReadWrite, meta = (DisplayAfter="LODSelectionType", ClampMin = "0", ClampMax = "7", UIMin = "0", UIMax = "7", EnumCondition = 1))
 	int32 SpecificLOD;
 
+	/** Which selection mode should be used when generating the merged static mesh */
 	UPROPERTY(EditAnywhere, BlueprintReadWrite, Category = MeshSettings, meta = (DisplayAfter="bBakeVertexDataToMesh"))
 	EMeshLODSelectionType LODSelectionType;
 
@@ -492,6 +493,7 @@ struct FMeshMergingSettings
 	UPROPERTY(EditAnywhere, BlueprintReadWrite, Category = LandscapeCulling)
 	uint8 bUseLandscapeCulling:1;
 
+	/** Whether or not to include any imposter LODs that are part of the source static meshes */
 	UPROPERTY(EditAnywhere, BlueprintReadWrite, Category = MeshSettings)
 	uint8 bIncludeImposters:1;
 
diff --git a/Engine/Source/Runtime/Engine/Classes/Engine/SkeletalMesh.h b/Engine/Source/Runtime/Engine/Classes/Engine/SkeletalMesh.h
--- a/Engine/Source/Runtime/Engine/Classes/Engine/SkeletalMesh.h
+++ b/Engine/Source/Runtime/Engine/Classes/Engine/SkeletalMesh.h
@@ -721,6 +721,13 @@ public:
 	UPROPERTY(Category=Mesh, BlueprintReadWrite)
 	TArray<UMorphTarget*> MorphTargets;
 
+	/**
+	 *	Returns the list of all morph targets of this skeletal mesh
+	 *  @return	The list of morph targets
+	 */
+	UFUNCTION(BlueprintPure, Category = Mesh, meta = (DisplayName = "GetAllMorphTargetNames", ScriptName = "GetAllMorphTargetNames", Keywords = "morph shape"))
+	TArray<FString> K2_GetAllMorphTargetNames() const;
+
 	/** A fence which is used to keep track of the rendering thread releasing the static mesh resources. */
 	FRenderCommandFence ReleaseResourcesFence;
 
diff --git a/Engine/Source/Runtime/Engine/Private/Animation/AnimComposite.cpp b/Engine/Source/Runtime/Engine/Private/Animation/AnimComposite.cpp
--- a/Engine/Source/Runtime/Engine/Private/Animation/AnimComposite.cpp
+++ b/Engine/Source/Runtime/Engine/Private/Animation/AnimComposite.cpp
@@ -26,7 +26,7 @@ void UAnimComposite::ReplaceReferredAnimations(const TMap<UAnimationAsset*, UAni
 
 bool UAnimComposite::IsNotifyAvailable() const
 {
-	return (Super::IsNotifyAvailable() || AnimationTrack.IsNotifyAvailable());
+	return (SequenceLength > 0.f && (Super::IsNotifyAvailable() || AnimationTrack.IsNotifyAvailable()));
 }
 
 void UAnimComposite::GetAnimNotifiesFromDeltaPositions(const float& PreviousPosition, const float & CurrentPosition, TArray<FAnimNotifyEventReference>& OutActiveNotifies) const
diff --git a/Engine/Source/Runtime/Engine/Private/Animation/AnimInstance.cpp b/Engine/Source/Runtime/Engine/Private/Animation/AnimInstance.cpp
--- a/Engine/Source/Runtime/Engine/Private/Animation/AnimInstance.cpp
+++ b/Engine/Source/Runtime/Engine/Private/Animation/AnimInstance.cpp
@@ -1578,7 +1578,7 @@ void UAnimInstance::TriggerMontageEndedEvent(const FQueuedMontageEndedEvent& Mon
 
 	if (SkelMeshComp != nullptr)
 	{
-		for (int32 Index = ActiveAnimNotifyState.Num() - 1; Index >= 0; --Index)
+		for (int32 Index = ActiveAnimNotifyState.Num() - 1; ActiveAnimNotifyState.IsValidIndex(Index); --Index)
 		{
 			const FAnimNotifyEvent& AnimNotifyEvent = ActiveAnimNotifyState[Index];
 			UAnimMontage* NotifyMontage = Cast<UAnimMontage>(AnimNotifyEvent.NotifyStateClass->GetOuter());
@@ -1589,7 +1589,19 @@ void UAnimInstance::TriggerMontageEndedEvent(const FQueuedMontageEndedEvent& Mon
 				{
 					AnimNotifyEvent.NotifyStateClass->NotifyEnd(SkelMeshComp, NotifyMontage);
 				}
-				ActiveAnimNotifyState.RemoveAtSwap(Index);
+
+				if (ActiveAnimNotifyState.IsValidIndex(Index))
+				{
+					ActiveAnimNotifyState.RemoveAtSwap(Index);
+				}
+				else
+				{
+					// The NotifyEnd callback above may have triggered actor destruction and the tear down
+					// of this instance via UninitializeAnimation which empties ActiveAnimNotifyState.
+					// If that happened, we should stop iterating the ActiveAnimNotifyState array and bail 
+					// without attempting to send MontageEnded events.
+					return;
+				}
 			}
 		}
 	}
diff --git a/Engine/Source/Runtime/Engine/Private/Animation/AnimNode_SequencePlayer.cpp b/Engine/Source/Runtime/Engine/Private/Animation/AnimNode_SequencePlayer.cpp
--- a/Engine/Source/Runtime/Engine/Private/Animation/AnimNode_SequencePlayer.cpp
+++ b/Engine/Source/Runtime/Engine/Private/Animation/AnimNode_SequencePlayer.cpp
@@ -32,6 +32,12 @@ void FAnimNode_SequencePlayer::Initialize_AnyThread(const FAnimationInitializeCo
 	FAnimNode_AssetPlayerBase::Initialize_AnyThread(Context);
 
 	GetEvaluateGraphExposedInputs().Execute(Context);
+
+	if(Sequence && !ensureMsgf(!Sequence->IsA<UAnimMontage>(), TEXT("Sequence players do not support anim montages.")))
+	{
+		Sequence = nullptr;
+	}
+
 	InternalTimeAccumulator = StartPosition;
 	PlayRateScaleBiasClamp.Reinitialize();
 
@@ -55,6 +61,11 @@ void FAnimNode_SequencePlayer::UpdateAssetPlayer(const FAnimationUpdateContext&
 {
 	GetEvaluateGraphExposedInputs().Execute(Context);
 
+	if(Sequence && !ensureMsgf(!Sequence->IsA<UAnimMontage>(), TEXT("Sequence players do not support anim montages.")))
+	{
+		Sequence = nullptr;
+	}
+
 	if ((Sequence != nullptr) && (Context.AnimInstanceProxy->IsSkeletonCompatible(Sequence->GetSkeleton())))
 	{
 		InternalTimeAccumulator = FMath::Clamp(InternalTimeAccumulator, 0.f, Sequence->SequenceLength);
diff --git a/Engine/Source/Runtime/Engine/Private/Animation/AnimSequence.cpp b/Engine/Source/Runtime/Engine/Private/Animation/AnimSequence.cpp
--- a/Engine/Source/Runtime/Engine/Private/Animation/AnimSequence.cpp
+++ b/Engine/Source/Runtime/Engine/Private/Animation/AnimSequence.cpp
@@ -4843,7 +4843,7 @@ bool UAnimSequence::CreateAnimation(USkeletalMeshComponent* MeshComponent)
 		RawAnimationData.AddZeroed(NumBones);
 		AnimationTrackNames.AddUninitialized(NumBones);
 
-		const TArray<FTransform>& BoneSpaceTransforms = MeshComponent->BoneSpaceTransforms;
+		const TArray<FTransform> BoneSpaceTransforms = MeshComponent->GetBoneSpaceTransforms();
 
 		check(BoneSpaceTransforms.Num() >= NumBones);
 
diff --git a/Engine/Source/Runtime/Engine/Private/Animation/BlendSpaceBase.cpp b/Engine/Source/Runtime/Engine/Private/Animation/BlendSpaceBase.cpp
--- a/Engine/Source/Runtime/Engine/Private/Animation/BlendSpaceBase.cpp
+++ b/Engine/Source/Runtime/Engine/Private/Animation/BlendSpaceBase.cpp
@@ -553,6 +553,18 @@ bool UBlendSpaceBase::IsValidAdditiveType(EAdditiveAnimationType AdditiveType) c
 	return false;
 }
 
+void UBlendSpaceBase::ResetToRefPose(FCompactPose& OutPose)
+{
+	if (IsValidAdditive())
+	{
+		OutPose.ResetToAdditiveIdentity();
+	}
+	else
+	{
+		OutPose.ResetToRefPose();
+	}
+}
+
 void UBlendSpaceBase::GetAnimationPose(TArray<FBlendSampleData>& BlendSampleDataCache, /*out*/ FCompactPose& OutPose, /*out*/ FBlendedCurve& OutCurve)
 {
 	SCOPE_CYCLE_COUNTER(STAT_BlendSpace_GetAnimPose);
@@ -560,7 +572,7 @@ void UBlendSpaceBase::GetAnimationPose(TArray<FBlendSampleData>& BlendSampleData
 
 	if(BlendSampleDataCache.Num() == 0)
 	{
-		OutPose.ResetToRefPose();
+		ResetToRefPose(OutPose);
 		return;
 	}
 
@@ -605,12 +617,12 @@ void UBlendSpaceBase::GetAnimationPose(TArray<FBlendSampleData>& BlendSampleData
 			}
 			else
 			{
-				Pose.ResetToRefPose();
+				ResetToRefPose(Pose);
 			}
 		}
 		else
 		{
-			Pose.ResetToRefPose();
+			ResetToRefPose(Pose);
 		}
 	}
 
diff --git a/Engine/Source/Runtime/Engine/Private/Animation/PoseAsset.cpp b/Engine/Source/Runtime/Engine/Private/Animation/PoseAsset.cpp
--- a/Engine/Source/Runtime/Engine/Private/Animation/PoseAsset.cpp
+++ b/Engine/Source/Runtime/Engine/Private/Animation/PoseAsset.cpp
@@ -696,7 +696,7 @@ bool UPoseAsset::GetAnimationPose(struct FCompactPose& OutPose, FBlendedCurve& O
 					}
 					FAnimationRuntime::RetargetBoneTransform(MySkeleton, RetargetSource, BlendedBoneTransform[TrackIndex], LocalBoneIndices.SkeletonBoneIndex, LocalBoneIndices.CompactBoneIndex, RequiredBones, bAdditivePose);
 					OutPose[LocalBoneIndices.CompactBoneIndex] = BlendedBoneTransform[TrackIndex];
-					OutPose.NormalizeRotations();
+					OutPose[LocalBoneIndices.CompactBoneIndex].NormalizeRotation();
 				}
 			}
 
diff --git a/Engine/Source/Runtime/Engine/Private/Animation/Skeleton.cpp b/Engine/Source/Runtime/Engine/Private/Animation/Skeleton.cpp
--- a/Engine/Source/Runtime/Engine/Private/Animation/Skeleton.cpp
+++ b/Engine/Source/Runtime/Engine/Private/Animation/Skeleton.cpp
@@ -1538,10 +1538,17 @@ void USkeleton::AccumulateCurveMetaData(FName CurveName, bool bMaterialSet, bool
 			if (Mapping->Exists(CurveName))
 			{
 				FCurveMetaData* CurveMetaData = GetCurveMetaData(CurveName);
+				bool bOldMaterial = CurveMetaData->Type.bMaterial;
+				bool bOldMorphtarget = CurveMetaData->Type.bMorphtarget;
 				// we don't want to undo previous flags, if it was true, we just alolw more to it. 
 				CurveMetaData->Type.bMaterial |= bMaterialSet;
 				CurveMetaData->Type.bMorphtarget |= bMorphtargetSet;
-				MarkPackageDirty();
+
+				if (bOldMaterial != CurveMetaData->Type.bMaterial 
+					|| bOldMorphtarget != CurveMetaData->Type.bMorphtarget)
+				{
+					MarkPackageDirty();
+				}
 			}
 		}
 	}
diff --git a/Engine/Source/Runtime/Engine/Private/Components/PoseableMeshComponent.cpp b/Engine/Source/Runtime/Engine/Private/Components/PoseableMeshComponent.cpp
--- a/Engine/Source/Runtime/Engine/Private/Components/PoseableMeshComponent.cpp
+++ b/Engine/Source/Runtime/Engine/Private/Components/PoseableMeshComponent.cpp
@@ -202,7 +202,7 @@ FTransform GetBoneTransformByNameHelper(FName BoneName, EBoneSpaces::Type BoneSp
 	}*/
 
 	FA2CSPose CSPose;
-	CSPose.AllocateLocalPoses(RequiredBones, Component->BoneSpaceTransforms);
+	CSPose.AllocateLocalPoses(RequiredBones, Component->GetBoneSpaceTransforms());
 
 	if (BoneSpace == EBoneSpaces::ComponentSpace)
 	{
@@ -278,16 +278,16 @@ void UPoseableMeshComponent::ResetBoneTransformByName(FName BoneName)
 	}
 }
 
-void UPoseableMeshComponent::CopyPoseFromSkeletalComponent(const USkeletalMeshComponent* InComponentToCopy)
+void UPoseableMeshComponent::CopyPoseFromSkeletalComponent(USkeletalMeshComponent* InComponentToCopy)
 {
 	if(RequiredBones.IsValid())
 	{
-		if(this->SkeletalMesh == InComponentToCopy->SkeletalMesh)
+		TArray<FTransform> LocalTransforms = InComponentToCopy->GetBoneSpaceTransforms();
+		if(this->SkeletalMesh == InComponentToCopy->SkeletalMesh 
+			&& LocalTransforms.Num() == BoneSpaceTransforms.Num())
 		{
-			check(BoneSpaceTransforms.Num() == InComponentToCopy->BoneSpaceTransforms.Num());
-
-			// Quick path, we know everything matches, just copy the local atoms
-			BoneSpaceTransforms = InComponentToCopy->BoneSpaceTransforms;
+			
+			Exchange(BoneSpaceTransforms, LocalTransforms);
 		}
 		else
 		{
@@ -306,7 +306,7 @@ void UPoseableMeshComponent::CopyPoseFromSkeletalComponent(const USkeletalMeshCo
 
 				if(TargetBoneIndex != INDEX_NONE)
 				{
-					BoneSpaceTransforms[TargetBoneIndex] = InComponentToCopy->BoneSpaceTransforms[SourceBoneIndex];
+					BoneSpaceTransforms[TargetBoneIndex] = LocalTransforms[SourceBoneIndex];
 				}
 			}
 		}
diff --git a/Engine/Source/Runtime/Engine/Private/Components/SkeletalMeshComponent.cpp b/Engine/Source/Runtime/Engine/Private/Components/SkeletalMeshComponent.cpp
--- a/Engine/Source/Runtime/Engine/Private/Components/SkeletalMeshComponent.cpp
+++ b/Engine/Source/Runtime/Engine/Private/Components/SkeletalMeshComponent.cpp
@@ -654,9 +654,11 @@ void USkeletalMeshComponent::InitAnim(bool bForceReinit)
 				}
 				else
 				{
+					PRAGMA_DISABLE_DEPRECATION_WARNINGS
 					BoneSpaceTransforms = SkeletalMesh->RefSkeleton.GetRefBonePose();
 					//Mini RefreshBoneTransforms (the bit we actually care about)
 					FillComponentSpaceTransforms(SkeletalMesh, BoneSpaceTransforms, GetEditableComponentSpaceTransforms());
+					PRAGMA_ENABLE_DEPRECATION_WARNINGS
 					bNeedToFlipSpaceBaseBuffers = true; // Have updated space bases so need to flip
 					FlipEditableSpaceBases();
 				}
@@ -1468,6 +1470,11 @@ void USkeletalMeshComponent::RecalcRequiredCurves()
 		return;
 	}
 
+	if (SkeletalMesh->Skeleton)
+	{
+		CachedCurveUIDList = SkeletalMesh->Skeleton->GetDefaultCurveUIDList();
+	}
+
 	const FCurveEvaluationOption CurveEvalOption(bAllowAnimCurveEvaluation, &DisallowedAnimCurves, PredictedLODLevel);
 
 	// make sure animation requiredcurve to mark as dirty
@@ -1590,8 +1597,10 @@ void USkeletalMeshComponent::ComputeRequiredBones(TArray<FBoneIndexType>& OutReq
 
 	// Add in any bones that may be required when mirroring.
 	// JTODO: This is only required if there are mirroring nodes in the tree, but hard to know...
+	PRAGMA_DISABLE_DEPRECATION_WARNINGS
 	if (SkeletalMesh->SkelMirrorTable.Num() > 0 &&
 		SkeletalMesh->SkelMirrorTable.Num() == BoneSpaceTransforms.Num())
+	PRAGMA_ENABLE_DEPRECATION_WARNINGS
 	{
 		TArray<FBoneIndexType> MirroredDesiredBones;
 		MirroredDesiredBones.AddUninitialized(RequiredBones.Num());
@@ -1667,9 +1676,9 @@ void USkeletalMeshComponent::RecalcRequiredBones(int32 LODIndex)
 	}
 
 	ComputeRequiredBones(RequiredBones, FillComponentSpaceTransformsRequiredBones, LODIndex, /*bIgnorePhysicsAsset=*/ false);
-
+	PRAGMA_DISABLE_DEPRECATION_WARNINGS
 	BoneSpaceTransforms = SkeletalMesh->RefSkeleton.GetRefBonePose();
-
+	PRAGMA_ENABLE_DEPRECATION_WARNINGS
 	// make sure animation requiredBone to mark as dirty
 	if (AnimScriptInstance)
 	{
@@ -1728,11 +1737,7 @@ void USkeletalMeshComponent::EvaluateAnimation(const USkeletalMesh* InSkeletalMe
 	}
 	else
 	{
-		// unfortunately it's possible they might not have skeleton, in that case, we don't have any place to copy the curve from
-		if (InSkeletalMesh->Skeleton)
-		{
-			OutCurve.InitFrom(&InSkeletalMesh->Skeleton->GetDefaultCurveUIDList());
-		}
+		OutCurve.InitFrom(&CachedCurveUIDList);
 	}
 }
 
@@ -1978,11 +1983,13 @@ void USkeletalMeshComponent::RefreshBoneTransforms(FActorComponentTickFunction*
 
 	//Handle update rate optimization setup
 	//Dont mark cache as invalid if we aren't performing optimization anyway
+	PRAGMA_DISABLE_DEPRECATION_WARNINGS
 	const bool bInvalidCachedBones = bDoEvaluationRateOptimization &&
 									 ((BoneSpaceTransforms.Num() != SkeletalMesh->RefSkeleton.GetNum())
 									 || (BoneSpaceTransforms.Num() != CachedBoneSpaceTransforms.Num())
 									 || (GetNumComponentSpaceTransforms() != CachedComponentSpaceTransforms.Num()));
 
+	PRAGMA_ENABLE_DEPRECATION_WARNINGS
 
 	TArray<uint16> const* CurrentAnimCurveUIDFinder = (AnimScriptInstance) ? &AnimScriptInstance->GetRequiredBones().GetUIDToArrayLookupTable() : nullptr;
 	const bool bAnimInstanceHasCurveUIDList = CurrentAnimCurveUIDFinder != nullptr;
@@ -2109,8 +2116,10 @@ void USkeletalMeshComponent::RefreshBoneTransforms(FActorComponentTickFunction*
 
 				if(CachedBoneSpaceTransforms.Num())
 				{
+					PRAGMA_DISABLE_DEPRECATION_WARNINGS
 					BoneSpaceTransforms.Reset();
 					BoneSpaceTransforms.Append(CachedBoneSpaceTransforms);
+					PRAGMA_ENABLE_DEPRECATION_WARNINGS
 				}
 				if(CachedComponentSpaceTransforms.Num())
 				{
@@ -2149,7 +2158,9 @@ void USkeletalMeshComponent::SwapEvaluationContextBuffers()
 {
 	Exchange(AnimEvaluationContext.ComponentSpaceTransforms, GetEditableComponentSpaceTransforms());
 	Exchange(AnimEvaluationContext.CachedComponentSpaceTransforms, CachedComponentSpaceTransforms);
+	PRAGMA_DISABLE_DEPRECATION_WARNINGS
 	Exchange(AnimEvaluationContext.BoneSpaceTransforms, BoneSpaceTransforms);
+	PRAGMA_ENABLE_DEPRECATION_WARNINGS
 	Exchange(AnimEvaluationContext.CachedBoneSpaceTransforms, CachedBoneSpaceTransforms);
 	Exchange(AnimEvaluationContext.Curve, AnimCurves);
 	Exchange(AnimEvaluationContext.CachedCurve, CachedCurve);
@@ -2285,7 +2296,9 @@ void USkeletalMeshComponent::PostAnimEvaluation(FAnimationEvaluationContext& Eva
 			CachedComponentSpaceTransforms.Reset();
 			CachedComponentSpaceTransforms.Append(GetEditableComponentSpaceTransforms());
 			CachedBoneSpaceTransforms.Reset();
+			PRAGMA_DISABLE_DEPRECATION_WARNINGS
 			CachedBoneSpaceTransforms.Append(BoneSpaceTransforms);
+			PRAGMA_ENABLE_DEPRECATION_WARNINGS
 		}
 
 		if (EvaluationContext.bDoInterpolation)
@@ -2319,9 +2332,10 @@ void USkeletalMeshComponent::PostAnimEvaluation(FAnimationEvaluationContext& Eva
 				Alpha = ExternalInterpolationAlpha;
 			}
 
+			PRAGMA_DISABLE_DEPRECATION_WARNINGS
 			FAnimationRuntime::LerpBoneTransforms(BoneSpaceTransforms, CachedBoneSpaceTransforms, Alpha, RequiredBones);
 			FillComponentSpaceTransforms(SkeletalMesh, BoneSpaceTransforms, GetEditableComponentSpaceTransforms());
-
+			PRAGMA_ENABLE_DEPRECATION_WARNINGS
 			// interpolate curve
 			AnimCurves.LerpTo(CachedCurve, Alpha);
 		}
@@ -2594,23 +2608,28 @@ bool USkeletalMeshComponent::AllocateTransformData()
 	// Allocate transforms if not present.
 	if ( Super::AllocateTransformData() )
 	{
+		PRAGMA_DISABLE_DEPRECATION_WARNINGS
 		if(BoneSpaceTransforms.Num() != SkeletalMesh->RefSkeleton.GetNum() )
 		{
 			BoneSpaceTransforms = SkeletalMesh->RefSkeleton.GetRefBonePose();
 		}
+		PRAGMA_ENABLE_DEPRECATION_WARNINGS
 
 		return true;
 	}
 
+	PRAGMA_DISABLE_DEPRECATION_WARNINGS
 	BoneSpaceTransforms.Empty();
-	
+	PRAGMA_ENABLE_DEPRECATION_WARNINGS
 	return false;
 }
 
 void USkeletalMeshComponent::DeallocateTransformData()
 {
 	Super::DeallocateTransformData();
+	PRAGMA_DISABLE_DEPRECATION_WARNINGS
 	BoneSpaceTransforms.Empty();
+	PRAGMA_ENABLE_DEPRECATION_WARNINGS
 }
 
 void USkeletalMeshComponent::SetForceRefPose(bool bNewForceRefPose)
@@ -2771,9 +2790,9 @@ void USkeletalMeshComponent::HideBone( int32 BoneIndex, EPhysBodyOp PhysBodyOpti
 		return;
 	}
 
-	if (BoneSpaceTransforms.IsValidIndex(BoneIndex))
+	// if valid bone index
+	if (BoneIndex >= 0 && GetNumBones() > BoneIndex)
 	{
-		BoneSpaceTransforms[BoneIndex].SetScale3D(FVector::ZeroVector);
 		bRequiredBonesUpToDate = false;
 
 		if (PhysBodyOption != PBO_None)
@@ -2805,9 +2824,8 @@ void USkeletalMeshComponent::UnHideBone( int32 BoneIndex )
 		return;
 	}
 
-	if (BoneSpaceTransforms.IsValidIndex(BoneIndex))
+	if (BoneIndex >= 0 && GetNumBones() > BoneIndex)
 	{
-		BoneSpaceTransforms[BoneIndex].SetScale3D(FVector(1.0f));
 		bRequiredBonesUpToDate = false;
 
 		//FName HideBoneName = SkeletalMesh->RefSkeleton.GetBoneName(BoneIndex);
@@ -3952,4 +3970,17 @@ const TArray<FTransform>& USkeletalMeshComponent::GetCachedComponentSpaceTransfo
 	return CachedComponentSpaceTransforms;
 }
 
+TArray<FTransform> USkeletalMeshComponent::GetBoneSpaceTransforms() 
+{
+	// We may be doing parallel evaluation on the current anim instance
+	// Calling this here with true will block this init till that thread completes
+	// and it is safe to continue
+	const bool bBlockOnTask = true; // wait on evaluation task so it is safe to swap the buffers
+	const bool bPerformPostAnimEvaluation = true; // Do PostEvaluation so we make sure to swap the buffers back. 
+	HandleExistingParallelEvaluationTask(bBlockOnTask, bPerformPostAnimEvaluation);
+
+	PRAGMA_DISABLE_DEPRECATION_WARNINGS
+	return BoneSpaceTransforms;
+	PRAGMA_ENABLE_DEPRECATION_WARNINGS
+}
 #undef LOCTEXT_NAMESPACE
diff --git a/Engine/Source/Runtime/Engine/Private/PhysicsEngine/PhysAnim.cpp b/Engine/Source/Runtime/Engine/Private/PhysicsEngine/PhysAnim.cpp
--- a/Engine/Source/Runtime/Engine/Private/PhysicsEngine/PhysAnim.cpp
+++ b/Engine/Source/Runtime/Engine/Private/PhysicsEngine/PhysAnim.cpp
@@ -364,6 +364,7 @@ void USkeletalMeshComponent::BlendInPhysics(FTickFunction& ThisTickFunction)
 		const bool bParallelBlend = !!CVarUseParallelBlendPhysics.GetValueOnGameThread() && FApp::ShouldUseThreadingForPerformance();
 		if(bParallelBlend)
 		{
+			PRAGMA_DISABLE_DEPRECATION_WARNINGS
 			if (SkeletalMesh->RefSkeleton.GetNum() != AnimEvaluationContext.BoneSpaceTransforms.Num())
 			{
 				// Initialize Parallel Task arrays
@@ -372,6 +373,7 @@ void USkeletalMeshComponent::BlendInPhysics(FTickFunction& ThisTickFunction)
 
 			AnimEvaluationContext.BoneSpaceTransforms.Reset(BoneSpaceTransforms.Num());
 			AnimEvaluationContext.BoneSpaceTransforms.Append(BoneSpaceTransforms);
+			PRAGMA_ENABLE_DEPRECATION_WARNINGS
 
 			ParallelAnimationEvaluationTask = TGraphTask<FParallelBlendPhysicsTask>::CreateTask().ConstructAndDispatchWhenReady(this);
 
@@ -387,7 +389,9 @@ void USkeletalMeshComponent::BlendInPhysics(FTickFunction& ThisTickFunction)
 		}
 		else
 		{
+			PRAGMA_DISABLE_DEPRECATION_WARNINGS
 			PerformBlendPhysicsBones(RequiredBones, BoneSpaceTransforms);
+			PRAGMA_ENABLE_DEPRECATION_WARNINGS
 			FinalizeAnimationUpdate();
 		}
 	}
@@ -454,8 +458,10 @@ void USkeletalMeshComponent::FinalizeAnimationUpdate()
 
 void USkeletalMeshComponent::CompleteParallelBlendPhysics()
 {
+	PRAGMA_DISABLE_DEPRECATION_WARNINGS
 	Exchange(AnimEvaluationContext.BoneSpaceTransforms, AnimEvaluationContext.bDoInterpolation ? CachedBoneSpaceTransforms : BoneSpaceTransforms);
-		
+	PRAGMA_ENABLE_DEPRECATION_WARNINGS
+
 	FinalizeAnimationUpdate();
 
 	ParallelAnimationEvaluationTask.SafeRelease();
@@ -719,10 +725,12 @@ void USkeletalMeshComponent::UpdateRBJointMotors()
 				(GetBoneVisibilityStates()[BoneIndex] == BVS_Visible) &&
 				(CI->IsAngularOrientationDriveEnabled()) )
 			{
+				PRAGMA_DISABLE_DEPRECATION_WARNINGS
 				check(BoneIndex < BoneSpaceTransforms.Num());
 
 				// If we find the joint - get the local-space animation between this bone and its parent.
 				FQuat LocalQuat = BoneSpaceTransforms[BoneIndex].GetRotation();
+				PRAGMA_ENABLE_DEPRECATION_WARNINGS
 				FQuatRotationTranslationMatrix LocalRot(LocalQuat, FVector::ZeroVector);
 
 				// We loop from the graphics parent bone up to the bone that has the body which the joint is attached to, to calculate the relative transform.
@@ -734,16 +742,20 @@ void USkeletalMeshComponent::UpdateRBJointMotors()
 
 				while(!bFoundControlBody)
 				{
+					PRAGMA_DISABLE_DEPRECATION_WARNINGS
 					// Abort if we find a bone scaled to zero.
 					const FVector Scale3D = BoneSpaceTransforms[TestBoneIndex].GetScale3D();
+					PRAGMA_ENABLE_DEPRECATION_WARNINGS
 					const float ScaleSum = Scale3D.X + Scale3D.Y + Scale3D.Z;
 					if(ScaleSum < KINDA_SMALL_NUMBER)
 					{
 						break;
 					}
 
+					PRAGMA_DISABLE_DEPRECATION_WARNINGS
 					// Add the current animated local transform into the overall controlling body->parent bone TM
 					FMatrix RelTM = BoneSpaceTransforms[TestBoneIndex].ToMatrixNoScale();
+					PRAGMA_ENABLE_DEPRECATION_WARNINGS
 					RelTM.SetOrigin(FVector::ZeroVector);
 					ControlBodyToParentBoneTM = ControlBodyToParentBoneTM * RelTM;
 
diff --git a/Engine/Source/Runtime/Engine/Private/PhysicsEngine/PhysicalAnimationComponent.cpp b/Engine/Source/Runtime/Engine/Private/PhysicsEngine/PhysicalAnimationComponent.cpp
--- a/Engine/Source/Runtime/Engine/Private/PhysicsEngine/PhysicalAnimationComponent.cpp
+++ b/Engine/Source/Runtime/Engine/Private/PhysicsEngine/PhysicalAnimationComponent.cpp
@@ -216,10 +216,10 @@ FTransform ComputeWorldSpaceTargetTM(const USkeletalMeshComponent& SkeletalMeshC
 	return SpaceBases[BoneIndex] * SkeletalMeshComponent.GetComponentToWorld();
 }
 
-FTransform ComputeLocalSpaceTargetTM(const USkeletalMeshComponent& SkeletalMeshComponent, const UPhysicsAsset& PhysAsset, int32 BoneIndex)
+FTransform ComputeLocalSpaceTargetTM(const USkeletalMeshComponent& SkeletalMeshComponent, const UPhysicsAsset& PhysAsset, const TArray<FTransform>& LocalTransforms, int32 BoneIndex)
 {
 	const FReferenceSkeleton& RefSkeleton = SkeletalMeshComponent.SkeletalMesh->RefSkeleton;
-	FTransform AccumulatedDelta = SkeletalMeshComponent.BoneSpaceTransforms[BoneIndex];
+	FTransform AccumulatedDelta = LocalTransforms[BoneIndex];
 	int32 CurBoneIdx = BoneIndex;
 	while ((CurBoneIdx = RefSkeleton.GetParentIndex(CurBoneIdx)) != INDEX_NONE)
 	{
@@ -246,15 +246,15 @@ FTransform ComputeLocalSpaceTargetTM(const USkeletalMeshComponent& SkeletalMeshC
 			}
 		}
 
-		AccumulatedDelta = AccumulatedDelta * SkeletalMeshComponent.BoneSpaceTransforms[CurBoneIdx];
+		AccumulatedDelta = AccumulatedDelta * LocalTransforms[CurBoneIdx];
 	}
 
 	return FTransform::Identity;
 }
 
-FTransform ComputeTargetTM(const FPhysicalAnimationData& PhysAnimData, const USkeletalMeshComponent& SkeletalMeshComponent, const UPhysicsAsset& PhysAsset, const TArray<FTransform>& SpaceBases, int32 BoneIndex)
+FTransform ComputeTargetTM(const FPhysicalAnimationData& PhysAnimData, const USkeletalMeshComponent& SkeletalMeshComponent, const UPhysicsAsset& PhysAsset, const TArray<FTransform>& LocalTransforms, const TArray<FTransform>& SpaceBases, int32 BoneIndex)
 {
-	return PhysAnimData.bIsLocalSimulation ? ComputeLocalSpaceTargetTM(SkeletalMeshComponent, PhysAsset, BoneIndex) : ComputeWorldSpaceTargetTM(SkeletalMeshComponent, SpaceBases, BoneIndex);
+	return PhysAnimData.bIsLocalSimulation ? ComputeLocalSpaceTargetTM(SkeletalMeshComponent, PhysAsset, LocalTransforms, BoneIndex) : ComputeWorldSpaceTargetTM(SkeletalMeshComponent, SpaceBases, BoneIndex);
 }
 
 void UPhysicalAnimationComponent::UpdateTargetActors(ETeleportType TeleportType)
@@ -271,6 +271,7 @@ void UPhysicalAnimationComponent::UpdateTargetActors(ETeleportType TeleportType)
 #if WITH_PHYSX
 		FPhysicsCommand::ExecuteWrite(SkeletalMeshComponent, [&]()
 		{
+			TArray<FTransform> LocalTransforms = SkeletalMeshComponent->GetBoneSpaceTransforms();
 			for (int32 DataIdx = 0; DataIdx < DriveData.Num(); ++DataIdx)
 			{
 				const FPhysicalAnimationData& PhysAnimData = DriveData[DataIdx];
@@ -280,7 +281,7 @@ void UPhysicalAnimationComponent::UpdateTargetActors(ETeleportType TeleportType)
 					const int32 BoneIdx = RefSkeleton.FindBoneIndex(PhysAnimData.BodyName);
 					if (BoneIdx != INDEX_NONE)	//It's possible the skeletal mesh has changed out from under us. In that case we should probably reset, but at the very least don't do work on non-existent bones
 					{
-						const FTransform TargetTM = ComputeTargetTM(PhysAnimData, *SkeletalMeshComponent, *PhysAsset, SpaceBases, BoneIdx);
+						const FTransform TargetTM = ComputeTargetTM(PhysAnimData, *SkeletalMeshComponent, *PhysAsset, LocalTransforms, SpaceBases, BoneIdx);
 						TargetActor->setKinematicTarget(U2PTransform(TargetTM));	//TODO: this doesn't work with sub-stepping!
 						
 						if(TeleportType == ETeleportType::TeleportPhysics)
@@ -338,6 +339,8 @@ void UPhysicalAnimationComponent::UpdatePhysicsEngine()
 #if WITH_PHYSX
 		FPhysicsCommand::ExecuteWrite(SkeletalMeshComponent, [&]()
 		{
+			TArray<FTransform> LocalTransforms = SkeletalMeshComponent->GetBoneSpaceTransforms();
+
 			for (int32 DataIdx = 0; DataIdx < DriveData.Num(); ++DataIdx)
 			{
 				bool bNewConstraint = false;
@@ -368,7 +371,7 @@ void UPhysicalAnimationComponent::UpdatePhysicsEngine()
 							ConstraintInstance->SetRefFrame(EConstraintFrame::Frame1, FTransform::Identity);
 							ConstraintInstance->SetRefFrame(EConstraintFrame::Frame2, FTransform::Identity);
 
-							const FTransform TargetTM = ComputeTargetTM(PhysAnimData, *SkeletalMeshComponent, *PhysAsset, SpaceBases, ChildBody->InstanceBoneIndex);
+							const FTransform TargetTM = ComputeTargetTM(PhysAnimData, *SkeletalMeshComponent, *PhysAsset, LocalTransforms, SpaceBases, ChildBody->InstanceBoneIndex);
 
 							// Create kinematic actor we are going to create joint with. This will be moved around with calls to SetLocation/SetRotation.
 							PxScene* PScene = PRigidActor->getScene();
diff --git a/Engine/Source/Runtime/Engine/Private/SkeletalMesh.cpp b/Engine/Source/Runtime/Engine/Private/SkeletalMesh.cpp
--- a/Engine/Source/Runtime/Engine/Private/SkeletalMesh.cpp
+++ b/Engine/Source/Runtime/Engine/Private/SkeletalMesh.cpp
@@ -3233,6 +3233,17 @@ FText USkeletalMesh::GetSourceFileLabelFromIndex(int32 SourceFileIndex)
 }
 #endif //WITH_EDITOR
 
+
+TArray<FString> USkeletalMesh::K2_GetAllMorphTargetNames() const
+{
+	TArray<FString> Names;
+	for (UMorphTarget* MorphTarget : MorphTargets)
+	{
+		Names.Add(MorphTarget->GetFName().ToString());
+	}
+	return Names;
+}
+
 /*-----------------------------------------------------------------------------
 USkeletalMeshSocket
 -----------------------------------------------------------------------------*/
diff --git a/Engine/Source/Runtime/Engine/Private/SkeletalMeshMerge.cpp b/Engine/Source/Runtime/Engine/Private/SkeletalMeshMerge.cpp
--- a/Engine/Source/Runtime/Engine/Private/SkeletalMeshMerge.cpp
+++ b/Engine/Source/Runtime/Engine/Private/SkeletalMeshMerge.cpp
@@ -397,7 +397,8 @@ void FSkeletalMeshMerge::CopyVertexFromSource(VertexDataType& DestVert, const FS
 
 	// Copy all UVs that are available
 	uint32 LODNumTexCoords = SrcLODData.StaticVertexBuffers.StaticMeshVertexBuffer.GetNumTexCoords();
-	for (uint32 UVIndex = 0; UVIndex < LODNumTexCoords && UVIndex < VertexDataType::NumTexCoords; ++UVIndex)
+	const uint32 ValidLoopCount = FMath::Min(VertexDataType::NumTexCoords, LODNumTexCoords);
+	for (uint32 UVIndex = 0; UVIndex < ValidLoopCount; ++UVIndex)
 	{
 		FVector2D UVs = SrcLODData.StaticVertexBuffers.StaticMeshVertexBuffer.GetVertexUV_Typed<VertexDataType::StaticMeshVertexUVType>(SourceVertIdx, UVIndex);
 		if (UVIndex < (uint32)MergeSectionInfo.UVTransforms.Num())
@@ -407,6 +408,12 @@ void FSkeletalMeshMerge::CopyVertexFromSource(VertexDataType& DestVert, const FS
 		}
 		DestVert.UVs[UVIndex] = UVs;
 	}
+	
+	// now just fill up zero value if we didn't reach till end
+	for (uint32 UVIndex = ValidLoopCount; UVIndex < VertexDataType::NumTexCoords; ++UVIndex)
+	{
+		DestVert.UVs[UVIndex] = FVector2D::ZeroVector;
+	}
 }
 
 template<typename SkinWeightType, bool bHasExtraBoneInfluences>
@@ -695,7 +702,7 @@ void FSkeletalMeshMerge::GenerateLODModel( int32 LODIdx )
                         Section.DuplicatedVerticesBuffer.DupVertData = MergeSectionInfo.Section->DuplicatedVerticesBuffer.DupVertData;
                         Section.DuplicatedVerticesBuffer.DupVertIndexData = MergeSectionInfo.Section->DuplicatedVerticesBuffer.DupVertIndexData;
                         uint8* VertData = Section.DuplicatedVerticesBuffer.DupVertData.GetDataPointer();
-                        for (uint32 i = 0; i < MergeSectionInfo.Section->NumVertices; ++i)
+                        for (int32 i = 0; i < Section.DuplicatedVerticesBuffer.DupVertData.Num(); ++i)
                         {
                             *((uint32*)(VertData + i * sizeof(uint32))) += CurrentBaseVertexIndex - MergeSectionInfo.Section->BaseVertexIndex;
                         }
