Added UHT error for marking a struct property as BlueprintAssignable since it cannot be bound in blueprints.

diff --git a/Engine/Source/Programs/UnrealHeaderTool/Private/HeaderParser.cpp b/Engine/Source/Programs/UnrealHeaderTool/Private/HeaderParser.cpp
--- a/Engine/Source/Programs/UnrealHeaderTool/Private/HeaderParser.cpp
+++ b/Engine/Source/Programs/UnrealHeaderTool/Private/HeaderParser.cpp
@@ -3532,6 +3532,11 @@ void FHeaderParser::GetVarType(
 
 				case EVariableSpecifier::BlueprintAssignable:
 				{
+					if (OwnerStruct->IsA<UScriptStruct>())
+					{
+						UE_LOG_ERROR_UHT(TEXT("Cannot specify BlueprintAssignable for a struct member."))
+					}
+
 					Flags |= CPF_BlueprintAssignable;
 				}
 				break;
