Remove UHT error for now. Deal with it next year.
#rnx

diff --git a/Engine/Source/Programs/UnrealHeaderTool/Private/HeaderParser.cpp b/Engine/Source/Programs/UnrealHeaderTool/Private/HeaderParser.cpp
--- a/Engine/Source/Programs/UnrealHeaderTool/Private/HeaderParser.cpp
+++ b/Engine/Source/Programs/UnrealHeaderTool/Private/HeaderParser.cpp
@@ -3532,11 +3532,6 @@ void FHeaderParser::GetVarType(
 
 				case EVariableSpecifier::BlueprintAssignable:
 				{
-					if (OwnerStruct->IsA<UScriptStruct>())
-					{
-						UE_LOG_ERROR_UHT(TEXT("Cannot specify BlueprintAssignable for a struct member."))
-					}
-
 					Flags |= CPF_BlueprintAssignable;
 				}
 				break;
