Merging //UE4/Dev-Main to Dev-Editor (//UE4/Dev-Editor)

diff --git a/Engine/Source/Programs/AutomationTool/AutomationUtils/UE4Build.cs b/Engine/Source/Programs/AutomationTool/AutomationUtils/UE4Build.cs
--- a/Engine/Source/Programs/AutomationTool/AutomationUtils/UE4Build.cs
+++ b/Engine/Source/Programs/AutomationTool/AutomationUtils/UE4Build.cs
@@ -1145,7 +1145,8 @@ namespace AutomationTool
 			bool bForceUnity = ParseParam("ForceUnity") || InForceUnity;
 			bool bForceDebugInfo = ParseParam("ForceDebugInfo");
 			string XGEConsole = XGEConsoleExe();
-			bool bCanUseXGE = !ParseParam("NoXGE") && !InForceNoXGE && !String.IsNullOrEmpty(XGEConsole);
+			bool bDisableXGE = ParseParam("NoXGE") || InForceNoXGE;
+			bool bCanUseXGE = !bDisableXGE && !String.IsNullOrEmpty(XGEConsole);
 
 			// only run ParallelExecutor if not running XGE (and we've requested ParallelExecutor and it exists)
 			bool bCanUseParallelExecutor = !bCanUseXGE && InUseParallelExecutor && (HostPlatform.Current.HostEditorPlatform == UnrealTargetPlatform.Win64);
@@ -1208,7 +1209,7 @@ namespace AutomationTool
 			{
 				// When building a target for Mac or iOS, use UBT's -flushmac option to clean up the remote builder
 				bool bForceFlushMac = DeleteBuildProducts && (Target.Platform == UnrealTargetPlatform.Mac || Target.Platform == UnrealTargetPlatform.IOS);
-				BuildManifest Manifest = BuildWithUBT(Target.TargetName, Target.Platform, Target.Config, Target.UprojectPath, bForceMonolithic, bForceNonUnity, bForceDebugInfo, bForceFlushMac, true, Target.UBTArgs, bForceUnity);
+				BuildManifest Manifest = BuildWithUBT(Target.TargetName, Target.Platform, Target.Config, Target.UprojectPath, bForceMonolithic, bForceNonUnity, bForceDebugInfo, bForceFlushMac, bDisableXGE, Target.UBTArgs, bForceUnity);
 				if(InTargetToManifest != null)
 				{
 					InTargetToManifest[Target] = Manifest;
@@ -1238,7 +1239,7 @@ namespace AutomationTool
 				if (bCanUseXGE) 
 				{
 					XGETool = XGEConsoleExePath;
-					Args = "\"" + TaskFilePath + "\" /Rebuild /MaxCPUS=200";
+					Args = "\"" + TaskFilePath + "\" /Rebuild /NoLogo /ShowAgent /ShowTime";
 					if (ParseParam("StopOnErrors"))
 					{
 						Args += " /StopOnErrors";
diff --git a/Engine/Source/Programs/UnrealHeaderTool/Private/HeaderParser.cpp b/Engine/Source/Programs/UnrealHeaderTool/Private/HeaderParser.cpp
--- a/Engine/Source/Programs/UnrealHeaderTool/Private/HeaderParser.cpp
+++ b/Engine/Source/Programs/UnrealHeaderTool/Private/HeaderParser.cpp
@@ -5325,7 +5325,7 @@ bool FHeaderParser::CompileDeclaration(FClasses& AllClasses, TArray<UDelegateFun
 							}
 							else
 							{
-								FError::Throwf(TEXT("Serialize functions must be defined outside of all compiler define blocks, except for WITH_EDITORONLY_DATA"));
+								FError::Throwf(TEXT("Serialize functions must not be inside preprocessor blocks, except for WITH_EDITORONLY_DATA"));
 							}
 						}
 					}
@@ -8375,7 +8375,32 @@ void FHeaderParser::SimplifiedClassParse(const TCHAR* Filename, const TCHAR* InB
 							|| FindInitialStr(FoundSubstr, TrimmedStrLine, TEXT("UDELEGATE"))
 							|| FindInitialStr(FoundSubstr, TrimmedStrLine, TEXT("UFUNCTION")))
 						{
-							FFileLineException::Throwf(Filename, CurrentLine, TEXT("%s inside this preprocessor block will be skipped"), FoundSubstr);
+							FFileLineException::Throwf(Filename, CurrentLine, TEXT("%s must not be inside preprocessor blocks, except for WITH_EDITORONLY_DATA"), FoundSubstr);
+						}
+
+						// Try and determine if this line contains something like a serialize function
+						if (TrimmedStrLine.Len() > 0)
+						{
+							static const FString Str_Void = TEXT("void");
+							static const FString Str_Serialize = TEXT("Serialize(");
+							static const FString Str_FArchive = TEXT("FArchive");
+							static const FString Str_FStructuredArchive = TEXT("FStructuredArchive::FSlot");
+
+							int32 Pos = 0;
+							if ((Pos = TrimmedStrLine.Find(Str_Void, ESearchCase::CaseSensitive, ESearchDir::FromStart, Pos)) != -1)
+							{
+								Pos += Str_Void.Len();
+								if ((Pos = TrimmedStrLine.Find(Str_Serialize, ESearchCase::CaseSensitive, ESearchDir::FromStart, Pos)) != -1)
+								{
+									Pos += Str_Serialize.Len();
+
+									if (((TrimmedStrLine.Find(Str_FArchive, ESearchCase::CaseSensitive, ESearchDir::FromStart, Pos)) != -1) ||
+										((TrimmedStrLine.Find(Str_FStructuredArchive, ESearchCase::CaseSensitive, ESearchDir::FromStart, Pos)) != -1))
+									{
+										FFileLineException::Throwf(Filename, CurrentLine, TEXT("'%s' must not be inside preprocessor blocks, except for WITH_EDITORONLY_DATA"), *TrimmedStrLine);
+									}
+								}
+							}
 						}
 					}
 				}
