Added some test cases for UHT parsing of Serialize functions

diff --git a/Engine/Source/Programs/UnrealHeaderTool/Resources/UHTDebugging.manifest b/Engine/Source/Programs/UnrealHeaderTool/Resources/UHTDebugging.manifest
--- a/Engine/Source/Programs/UnrealHeaderTool/Resources/UHTDebugging.manifest
+++ b/Engine/Source/Programs/UnrealHeaderTool/Resources/UHTDebugging.manifest
@@ -33,7 +33,8 @@
 				".\\UHTDebugging\\TestInterface.h",
 				".\\UHTDebugging\\TestInterfaceObject.h",
 				".\\UHTDebugging\\StructOnlyHeader.h",
-				".\\UHTDebugging\\EnumOnlyHeader.h"
+				".\\UHTDebugging\\EnumOnlyHeader.h",
+				".\\UHTDebugging\\TestObjectSerializers.h"
 			],
 			"PublicHeaders": [
 			],
diff --git a/Engine/Source/Programs/UnrealHeaderTool/Resources/UHTDebugging/TestObjectSerializers.h b/Engine/Source/Programs/UnrealHeaderTool/Resources/UHTDebugging/TestObjectSerializers.h
new file mode 100644
--- /dev/null
+++ b/Engine/Source/Programs/UnrealHeaderTool/Resources/UHTDebugging/TestObjectSerializers.h
@@ -0,0 +1,116 @@
+// Copyright 1998-2018 Epic Games, Inc. All Rights Reserved.
+
+#pragma once
+
+#include "EnumOnlyHeader.h"
+#include "TestObjectSerializers.generated.h"
+
+UCLASS()
+class UTestObject_NoSerializers : public UObject
+{
+	GENERATED_UCLASS_BODY()
+};
+
+UCLASS()
+class UTestObject_FArchive : public UObject
+{
+	GENERATED_UCLASS_BODY()
+	virtual void Serialize(FArchive& Ar) override;
+};
+
+UCLASS()
+class UTestObject_FStructuredArchive : public UObject
+{
+	GENERATED_UCLASS_BODY()
+	virtual void Serialize(FStructuredArchive::FSlot Slot) override;
+};
+
+UCLASS()
+class UTestObject_BothArchives : public UObject
+{
+	GENERATED_UCLASS_BODY()
+	virtual void Serialize(FArchive& Ar) override;
+	virtual void Serialize(FStructuredArchive::FSlot Slot) override;
+};
+
+UCLASS()
+class UTestObject_ArchiveInEditorOnlyDataDefine : public UObject
+{
+	GENERATED_UCLASS_BODY()
+	
+#if WITH_EDITORONLY_DATA
+	virtual void Serialize(FArchive& Ar) override;
+#endif
+};
+
+#if 0
+// Should fail - no Serialize functions inside WITH_EDITOR
+UCLASS()
+class UTestObject_ArchiveInWithEditor : public UObject
+{
+	GENERATED_UCLASS_BODY()
+
+#if WITH_EDITOR
+	virtual void Serialize(FArchive& Ar) override;
+#endif
+};
+#endif
+
+#if 0
+// Should fail - no Serialize functions inside WITH_EDITOR
+UCLASS()
+class UTestObject_StructuredArchiveInWithEditor : public UObject
+{
+	GENERATED_UCLASS_BODY()
+
+#if WITH_EDITOR
+		virtual void Serialize(FStructuredArchive::FSlot Slot) override;
+#endif
+};
+#endif
+
+#if 0
+#define SOME_RANDOM_DEFINE 1
+// Should fail - no Serialize functions inside arbitrary preprocessor blocks
+UCLASS()
+class UTestObject_ArchiveInPreprocessorBlock : public UObject
+{
+	GENERATED_UCLASS_BODY()
+
+#if SOME_RANDOM_DEFINE
+	void Serialize(FArchive& Ar) override;
+#endif
+};
+#undef SOME_RANDOM_DEFINE
+#endif
+
+#if 0
+#define SOME_RANDOM_DEFINE 1
+// Should fail - no Serialize functions inside arbitrary preprocessor blocks
+UCLASS()
+class UTestObject_StructuredArchiveInPreprocessorBlock : public UObject
+{
+	GENERATED_UCLASS_BODY()
+
+#if SOME_RANDOM_DEFINE
+		void Serialize(FStructuredArchive::FSlot Slot) override;
+#endif
+};
+#undef SOME_RANDOM_DEFINE
+#endif
+
+#if 0
+#define SOME_RANDOM_DEFINE 1
+// Should fail - no uproperties inside arbitrary preprocessor blocks
+UCLASS()
+class UTestObject_UPropertyInPreprocessorBlock : public UObject
+{
+	GENERATED_UCLASS_BODY()
+
+#if SOME_RANDOM_DEFINE
+	UPROPERTY()
+	uint32 TestProperty;
+#endif
+};
+#undef SOME_RANDOM_DEFINE
+#endif
\ No newline at end of file
