Made UHT GetSourceFilesWithInheritanceOrdering O(N) instead of O(N^2)

diff --git a/Engine/Source/Programs/UnrealHeaderTool/Private/HeaderParser.cpp b/Engine/Source/Programs/UnrealHeaderTool/Private/HeaderParser.cpp
--- a/Engine/Source/Programs/UnrealHeaderTool/Private/HeaderParser.cpp
+++ b/Engine/Source/Programs/UnrealHeaderTool/Private/HeaderParser.cpp
@@ -7710,7 +7710,7 @@ bool FHeaderParser::DependentClassNameFromHeader(const TCHAR* HeaderFilename, FS
  */
 TArray<FUnrealSourceFile*> GetSourceFilesWithInheritanceOrdering(UPackage* CurrentPackage, FClasses& AllClasses)
 {
-	TArray<FUnrealSourceFile*> SourceFiles;
+	TSet<FUnrealSourceFile*> SourceFiles;
 
 	TArray<FClass*> Classes = AllClasses.GetClassesInPackage();
 
@@ -7745,7 +7745,7 @@ TArray<FUnrealSourceFile*> GetSourceFilesWithInheritanceOrdering(UPackage* Curre
 		}
 	}
 
-	return SourceFiles;
+	return SourceFiles.Array();
 }
 
 // Begins the process of exporting C++ class declarations for native classes in the specified package
