UHT header resolve optimization

GUnrealSourceFilesMap used to contain absolute paths. Map lookup was done via filename, which always failed. It then traversed the entire map and tried to match via module relative paths, which also always failed. It finally prepended a slash to the filename and did another full scan and checked if entries ended with the slash + filename. Changed so that the map contains filename only to start with.

diff --git a/Engine/Source/Programs/UnrealHeaderTool/Private/CodeGenerator.cpp b/Engine/Source/Programs/UnrealHeaderTool/Private/CodeGenerator.cpp
--- a/Engine/Source/Programs/UnrealHeaderTool/Private/CodeGenerator.cpp
+++ b/Engine/Source/Programs/UnrealHeaderTool/Private/CodeGenerator.cpp
@@ -5649,7 +5649,7 @@ ECompilationResult::Type PreparseModules(const FString& ModuleInfoPath, int32& N
 					TSharedRef<FUnrealSourceFile> UnrealSourceFile = PerformInitialParseOnHeader(Package, *RawFilename, RF_Public | RF_Standalone, *HeaderFile);
 					FUnrealSourceFile* UnrealSourceFilePtr = &UnrealSourceFile.Get();
 					TArray<UClass*> DefinedClasses = UnrealSourceFile->GetDefinedClasses();
-					GUnrealSourceFilesMap.Add(RawFilename, UnrealSourceFile);
+					GUnrealSourceFilesMap.Add(FPaths::GetCleanFilename(RawFilename), UnrealSourceFile);
 
 					if (CurrentlyProcessing == PublicClassesHeaders)
 					{
diff --git a/Engine/Source/Programs/UnrealHeaderTool/Private/HeaderProvider.cpp b/Engine/Source/Programs/UnrealHeaderTool/Private/HeaderProvider.cpp
--- a/Engine/Source/Programs/UnrealHeaderTool/Private/HeaderProvider.cpp
+++ b/Engine/Source/Programs/UnrealHeaderTool/Private/HeaderProvider.cpp
@@ -13,22 +13,6 @@ FHeaderProvider::FHeaderProvider(EHeaderProviderSourceType InType, FString&& InI
 
 }
 
-template <class Predicate>
-bool TryFindSourceFileWithPredicate(FUnrealSourceFile*& Out, Predicate Pred)
-{
-	for (const TPair<FString, TSharedRef<FUnrealSourceFile>>& MapPair : GUnrealSourceFilesMap)
-	{
-		if (Pred(MapPair.Value.Get()))
-		{
-			Out = &(MapPair.Value.Get());
-			return true;
-		}
-	}
-
-	return false;
-}
-
-
 FUnrealSourceFile* FHeaderProvider::Resolve()
 {
 	if (Type != EHeaderProviderSourceType::Resolved)
@@ -44,29 +28,9 @@ FUnrealSourceFile* FHeaderProvider::Resolve()
 				}
 			}
 		}
-		else
+		else if (const TSharedRef<FUnrealSourceFile>* Source = GUnrealSourceFilesMap.Find(Id))
 		{
-			auto* SourceFileSharedPtr = GUnrealSourceFilesMap.Find(Id);
-
-			if (SourceFileSharedPtr != nullptr)
-			{
-				Cache = &SourceFileSharedPtr->Get();
-			}
-			else
-			{
-				if (!TryFindSourceFileWithPredicate(Cache, [this](const FUnrealSourceFile& SourceFile) { return SourceFile.GetIncludePath() == Id; }))
-				{
-					FString SlashId     = TEXT("/") + Id;
-					FString BackslashId = TEXT("\\") + Id;
-					TryFindSourceFileWithPredicate(
-						Cache,
-						[&SlashId, &BackslashId](const FUnrealSourceFile& SourceFile)
-						{
-							return SourceFile.GetFilename().EndsWith(SlashId) || SourceFile.GetFilename().EndsWith(BackslashId);
-						}
-					);
-				}
-			}
+			Cache = &Source->Get();
 		}
 
 		Type = EHeaderProviderSourceType::Resolved;
