Made header resolving another 3x faster

Avoid constructing lots of unnecessary FName.
Also removes some dead code.
Removes another 1.5% of UHT walltime.

diff --git a/Engine/Source/Programs/UnrealHeaderTool/Private/HeaderProvider.cpp b/Engine/Source/Programs/UnrealHeaderTool/Private/HeaderProvider.cpp
--- a/Engine/Source/Programs/UnrealHeaderTool/Private/HeaderProvider.cpp
+++ b/Engine/Source/Programs/UnrealHeaderTool/Private/HeaderProvider.cpp
@@ -19,9 +19,10 @@ FUnrealSourceFile* FHeaderProvider::Resolve()
 	{
 		if (Type == EHeaderProviderSourceType::ClassName)
 		{
+			FName IdName(*Id);
 			for (const auto& Pair : GTypeDefinitionInfoMap)
 			{
-				if (Pair.Key->GetFName() == *Id)
+				if (Pair.Key->GetFName() == IdName)
 				{
 					Cache = &Pair.Value->GetUnrealSourceFile();
 					break;
@@ -49,11 +50,6 @@ const FString& FHeaderProvider::GetId() const
 	return Id;
 }
 
-EHeaderProviderSourceType FHeaderProvider::GetType() const
-{
-	return Type;
-}
-
 bool operator==(const FHeaderProvider& A, const FHeaderProvider& B)
 {
 	return A.Id == B.Id && A.Type == B.Type;
diff --git a/Engine/Source/Programs/UnrealHeaderTool/Private/HeaderProvider.h b/Engine/Source/Programs/UnrealHeaderTool/Private/HeaderProvider.h
--- a/Engine/Source/Programs/UnrealHeaderTool/Private/HeaderProvider.h
+++ b/Engine/Source/Programs/UnrealHeaderTool/Private/HeaderProvider.h
@@ -11,8 +11,7 @@ enum class EHeaderProviderSourceType
 {
 	ClassName,
 	FileName,
-	Resolved,
-	Invalid,
+	Resolved
 };
 
 class FHeaderProvider
@@ -20,11 +19,6 @@ class FHeaderProvider
 	friend bool operator==(const FHeaderProvider& A, const FHeaderProvider& B);
 public:
 	FHeaderProvider(EHeaderProviderSourceType Type, FString&& Id);
-	FHeaderProvider()
-		: Type(EHeaderProviderSourceType::Invalid)
-		, Id(FString())
-		, Cache(nullptr)
-	{ }
 
 	FUnrealSourceFile* Resolve();
 
@@ -32,8 +26,6 @@ public:
 
 	const FString& GetId() const;
 
-	EHeaderProviderSourceType GetType() const;
-
 private:
 	EHeaderProviderSourceType Type;
 	FString Id;
