Made UHT package init code generation ~100x faster

Created and update a package -> exported fields map instead of traversing the entire reflection system for every package.
Decreases UHT null build walltime by ~11% or about 1s for my setup.

Also added timer and log entry for UnrealHeaderTool_Main() to help measure improvements.

diff --git a/Engine/Source/Programs/UnrealHeaderTool/Private/ClassMaps.cpp b/Engine/Source/Programs/UnrealHeaderTool/Private/ClassMaps.cpp
--- a/Engine/Source/Programs/UnrealHeaderTool/Private/ClassMaps.cpp
+++ b/Engine/Source/Programs/UnrealHeaderTool/Private/ClassMaps.cpp
@@ -6,6 +6,7 @@
 
 TMap<FString, TSharedRef<FUnrealSourceFile> > GUnrealSourceFilesMap;
 TMap<UField*, TSharedRef<FUnrealTypeDefinitionInfo> > GTypeDefinitionInfoMap;
+TMap<const UPackage*, TArray<UField*>> GPackageSingletons;
 TMap<UClass*, FString> GClassStrippedHeaderTextMap;
 TMap<UClass*, FString> GClassHeaderNameWithNoPathMap;
 TSet<FUnrealSourceFile*> GPublicSourceFileSet;
diff --git a/Engine/Source/Programs/UnrealHeaderTool/Private/ClassMaps.h b/Engine/Source/Programs/UnrealHeaderTool/Private/ClassMaps.h
--- a/Engine/Source/Programs/UnrealHeaderTool/Private/ClassMaps.h
+++ b/Engine/Source/Programs/UnrealHeaderTool/Private/ClassMaps.h
@@ -21,6 +21,7 @@ class FUnrealTypeDefinitionInfo;
 
 extern TMap<FString, TSharedRef<FUnrealSourceFile> > GUnrealSourceFilesMap;
 extern TMap<UField*, TSharedRef<FUnrealTypeDefinitionInfo> > GTypeDefinitionInfoMap;
+extern TMap<const UPackage*, TArray<UField*>> GPackageSingletons;
 extern TMap<UClass*, FString> GClassStrippedHeaderTextMap;
 extern TMap<UClass*, FString> GClassHeaderNameWithNoPathMap;
 extern TSet<FUnrealSourceFile*> GPublicSourceFileSet;
diff --git a/Engine/Source/Programs/UnrealHeaderTool/Private/CodeGenerator.cpp b/Engine/Source/Programs/UnrealHeaderTool/Private/CodeGenerator.cpp
--- a/Engine/Source/Programs/UnrealHeaderTool/Private/CodeGenerator.cpp
+++ b/Engine/Source/Programs/UnrealHeaderTool/Private/CodeGenerator.cpp
@@ -1836,31 +1836,15 @@ void FNativeClassHeaderGenerator::ExportGeneratedPackageInitCode(FOutputDevice&
 	FString ApiString = GetAPIString();
 	FString SingletonName = GetPackageSingletonName(InPackage);
 
-	TArray<UField*> SingletonsToOutput;
-	for (UField* ScriptType : TObjectRange<UField>())
+	const TArray<UField*>* SingletonsToOutput = GPackageSingletons.Find(InPackage);
+	if (SingletonsToOutput)
 	{
-		if (ScriptType->GetOutermost() != InPackage)
+		for (UField* ScriptType : *SingletonsToOutput)
 		{
-			continue;
-		}
-
-		if (ScriptType->IsA<UDelegateFunction>() || (ScriptType->IsA<UScriptStruct>() && (((UScriptStruct*)ScriptType)->StructFlags & STRUCT_NoExport)))
-		{
-			UField* FieldOuter = Cast<UField>(ScriptType->GetOuter());
-			if (!FieldOuter || !FClass::IsDynamic(FieldOuter))
-			{
-				SingletonsToOutput.Add(ScriptType);
-			}
+			Out.Log(FTypeSingletonCache::Get(ScriptType, true).GetExternDecl());
 		}
 	}
 
-	for (UField* ScriptType : SingletonsToOutput)
-	{
-		const FString& Extern = FTypeSingletonCache::Get(ScriptType, true).GetExternDecl();
-
-		Out.Logf(TEXT("%s"), *Extern);
-	}
-
 	FOutputDeviceNull OutputDeviceNull;
 	FString MetaDataParams = OutputMetaDataCodeForObject(OutputDeviceNull, Out, InPackage, TEXT("Package_MetaDataParams"), TEXT(""), TEXT("\t\t\t"));
 
@@ -1871,10 +1855,10 @@ void FNativeClassHeaderGenerator::ExportGeneratedPackageInitCode(FOutputDevice&
 	Out.Logf(TEXT("\t\t{\r\n"));
 
 	FString SingletonRange;
-	if (SingletonsToOutput.Num() > 0)
+	if (SingletonsToOutput)
 	{
 		Out.Logf(TEXT("\t\t\tstatic UObject* (*const SingletonFuncArray[])() = {\r\n"));
-		for (UField* ScriptType : SingletonsToOutput)
+		for (UField* ScriptType : *SingletonsToOutput)
 		{
 			const FString& Name = FTypeSingletonCache::Get(ScriptType, true).GetName().LeftChop(2);
 
@@ -4901,6 +4885,47 @@ TArray<FUnrealSourceFile*> GetSourceFilesInDependencyOrder(const UPackage* Packa
 
 TMap<UClass*, FUnrealSourceFile*> GClassToSourceFileMap;
 
+static bool HasDynamicOuter(UField* Field)
+{
+	UField* FieldOuter = Cast<UField>(Field->GetOuter());
+	return FieldOuter && FClass::IsDynamic(FieldOuter);
+}
+
+static void RecordPackageSingletons(
+	const UPackage& Package,
+	const TArray<UScriptStruct*>& Structs,
+	const TArray<UDelegateFunction*>& Delegates)
+{
+	TArray<UField*> Singletons;
+	for (UScriptStruct* Struct : Structs)
+	{
+		if (Struct->StructFlags & STRUCT_NoExport && !HasDynamicOuter(Struct))
+		{
+			Singletons.Add(Struct);
+		}
+	}
+
+	for (UDelegateFunction* Delegate : Delegates)
+	{
+		if (!HasDynamicOuter(Delegate))
+		{
+			Singletons.Add(Delegate);
+		}
+	}
+
+	if (Singletons.Num())
+	{
+		TArray<UField*>& PackageSingletons = GPackageSingletons.FindOrAdd(&Package);
+		PackageSingletons.Append(Singletons);
+		Algo::Sort(PackageSingletons, [](UField* A, UField* B)
+		{
+			// Structs before delegates then UniqueId order
+			return	(uint64(A->IsA<UDelegateFunction>()) << 32) + A->GetUniqueID() <
+					(uint64(B->IsA<UDelegateFunction>()) << 32) + B->GetUniqueID();
+		});
+	}
+}
+
 // Constructor.
 FNativeClassHeaderGenerator::FNativeClassHeaderGenerator(
 	const UPackage* InPackage,
@@ -5002,6 +5027,8 @@ FNativeClassHeaderGenerator::FNativeClassHeaderGenerator(
 		TArray<UDelegateFunction*> DelegateFunctions;
 		SourceFile->GetScope()->SplitTypesIntoArrays(Enums, Structs, DelegateFunctions);
 
+		RecordPackageSingletons(*SourceFile->GetPackage(), Structs, DelegateFunctions);
+
 		// Reverse the containers as they come out in the reverse order of declaration
 		Algo::Reverse(Enums);
 		Algo::Reverse(Structs);
@@ -5767,6 +5794,10 @@ ECompilationResult::Type PreparseModules(const FString& ModuleInfoPath, int32& N
 
 ECompilationResult::Type UnrealHeaderTool_Main(const FString& ModuleInfoFilename)
 {
+	double MainTime = 0.0;
+	FDurationTimer MainTimer(MainTime);
+	MainTimer.Start();
+
 	check(GIsUCCMakeStandaloneHeaderGenerator);
 	ECompilationResult::Type Result = ECompilationResult::Succeeded;
 
@@ -5893,11 +5924,14 @@ ECompilationResult::Type UnrealHeaderTool_Main(const FString& ModuleInfoFilename
 	// Avoid TArray slack for meta data.
 	GScriptHelper.Shrink();
 
+	MainTimer.Stop();
+
 	UE_LOG(LogCompile, Log, TEXT("Preparsing %i modules took %.2f seconds"), GManifest.Modules.Num(), TotalModulePreparseTime);
 	UE_LOG(LogCompile, Log, TEXT("Parsing took %.2f seconds"), TotalParseAndCodegenTime - GHeaderCodeGenTime);
 	UE_LOG(LogCompile, Log, TEXT("Code generation took %.2f seconds"), GHeaderCodeGenTime);
 	UE_LOG(LogCompile, Log, TEXT("ScriptPlugin overhead was %.2f seconds"), GPluginOverheadTime);
 	UE_LOG(LogCompile, Log, TEXT("Macroize time was %.2f seconds"), GMacroizeTime);
+	UE_LOG(LogCompile, Log, TEXT("Total time was %.2f seconds"), MainTime);
 
 	if (bWriteContents)
 	{
