UHT rename CRC to hash, since CRC is no longer used

diff --git a/Engine/Source/Programs/UnrealHeaderTool/Private/ClassMaps.cpp b/Engine/Source/Programs/UnrealHeaderTool/Private/ClassMaps.cpp
--- a/Engine/Source/Programs/UnrealHeaderTool/Private/ClassMaps.cpp
+++ b/Engine/Source/Programs/UnrealHeaderTool/Private/ClassMaps.cpp
@@ -12,7 +12,7 @@ TMap<UClass*, FString> GClassHeaderNameWithNoPathMap;
 TSet<FUnrealSourceFile*> GPublicSourceFileSet;
 TMap<UProperty*, FString> GArrayDimensions;
 TMap<UPackage*,  const FManifestModule*> GPackageToManifestModuleMap;
-TMap<UField*, uint32> GGeneratedCodeCRCs;
+TMap<UField*, uint32> GGeneratedCodeHashes;
 TMap<UEnum*,  EUnderlyingEnumType> GEnumUnderlyingTypes;
 TMap<FName, TSharedRef<FClassDeclarationMetaData> > GClassDeclarations;
 TSet<UProperty*> GUnsizedProperties;
diff --git a/Engine/Source/Programs/UnrealHeaderTool/Private/ClassMaps.h b/Engine/Source/Programs/UnrealHeaderTool/Private/ClassMaps.h
--- a/Engine/Source/Programs/UnrealHeaderTool/Private/ClassMaps.h
+++ b/Engine/Source/Programs/UnrealHeaderTool/Private/ClassMaps.h
@@ -27,7 +27,7 @@ extern TMap<UClass*, FString> GClassHeaderNameWithNoPathMap;
 extern TSet<FUnrealSourceFile*> GPublicSourceFileSet;
 extern TMap<UProperty*, FString> GArrayDimensions;
 extern TMap<UPackage*,  const FManifestModule*> GPackageToManifestModuleMap;
-extern TMap<UField*, uint32> GGeneratedCodeCRCs;
+extern TMap<UField*, uint32> GGeneratedCodeHashes;
 extern TMap<UEnum*, EUnderlyingEnumType> GEnumUnderlyingTypes;
 extern TMap<FName, TSharedRef<FClassDeclarationMetaData> > GClassDeclarations;
 extern TSet<UProperty*> GUnsizedProperties;
diff --git a/Engine/Source/Programs/UnrealHeaderTool/Private/CodeGenerator.cpp b/Engine/Source/Programs/UnrealHeaderTool/Private/CodeGenerator.cpp
--- a/Engine/Source/Programs/UnrealHeaderTool/Private/CodeGenerator.cpp
+++ b/Engine/Source/Programs/UnrealHeaderTool/Private/CodeGenerator.cpp
@@ -322,14 +322,14 @@ FString Macroize(const TCHAR* MacroName, const TCHAR* StringToMacroize)
 	return FString::Printf(TEXT("#define %s%s\r\n"), MacroName, Result.Len() ? TEXT(" \\") : TEXT("")) + Result;
 }
 
-/** Generates a CRC tag string for the specified field */
-static FString GetGeneratedCodeCRCTag(UField* Field)
+/** Generates a Hash tag string for the specified field */
+static FString GetGeneratedCodeHashTag(UField* Field)
 {
 	FString Tag;
-	const uint32* FieldCrc = GGeneratedCodeCRCs.Find(Field);
-	if (FieldCrc)
+	const uint32* FieldHash = GGeneratedCodeHashes.Find(Field);
+	if (FieldHash)
 	{
-		Tag = FString::Printf(TEXT(" // %u"), *FieldCrc);
+		Tag = FString::Printf(TEXT(" // %u"), *FieldHash);
 	}
 	return Tag;
 }
@@ -892,7 +892,7 @@ void FNativeClassHeaderGenerator::PropertyNew(FOutputDevice& DeclOut, FOutputDev
 	const TCHAR*   UPropertyObjectFlags = FClass::IsOwnedByDynamicType(Prop) ? TEXT("RF_Public|RF_Transient") : TEXT("RF_Public|RF_Transient|RF_MarkAsNative");
 	EPropertyFlags PropFlags            = Prop->PropertyFlags & ~CPF_ComputedFlags;
 
-	FString PropTag        = GetGeneratedCodeCRCTag(Prop);
+	FString PropTag        = GetGeneratedCodeHashTag(Prop);
 	FString PropNotifyFunc = (Prop->RepNotifyFunc != NAME_None) ? CreateUTF8LiteralString(*Prop->RepNotifyFunc.ToString()) : TEXT("nullptr");
 
 	FString ArrayDim = (Prop->ArrayDim != 1) ? FString::Printf(TEXT("CPP_ARRAY_DIM(%s, %s)"), *PropNameDep, SourceStruct) : TEXT("1");
@@ -1831,7 +1831,7 @@ FString FNativeClassHeaderGenerator::GetPackageSingletonName(const UPackage* InP
 	return Result;
 }
 
-void FNativeClassHeaderGenerator::ExportGeneratedPackageInitCode(FOutputDevice& Out, const TCHAR* InDeclarations, const UPackage* InPackage, uint32 CRC)
+void FNativeClassHeaderGenerator::ExportGeneratedPackageInitCode(FOutputDevice& Out, const TCHAR* InDeclarations, const UPackage* InPackage, uint32 Hash)
 {
 	FString ApiString = GetAPIString();
 	FString SingletonName = GetPackageSingletonName(InPackage);
@@ -1876,8 +1876,8 @@ void FNativeClassHeaderGenerator::ExportGeneratedPackageInitCode(FOutputDevice&
 	Out.Logf(TEXT("\t\t\tstatic const UE4CodeGen_Private::FPackageParams PackageParams = {\r\n"));
 	Out.Logf(TEXT("\t\t\t\t%s,\r\n"), *CreateUTF8LiteralString(InPackage->GetName()));
 	Out.Logf(TEXT("\t\t\t\tPKG_CompiledIn | 0x%08X,\r\n"), InPackage->GetPackageFlags() & (PKG_ClientOptional | PKG_ServerSideOnly | PKG_EditorOnly | PKG_Developer));
-	Out.Logf(TEXT("\t\t\t\t0x%08X,\r\n"), CRC);
-	Out.Logf(TEXT("\t\t\t\t0x%08X,\r\n"), GenerateTextCRC(InDeclarations));
+	Out.Logf(TEXT("\t\t\t\t0x%08X,\r\n"), Hash);
+	Out.Logf(TEXT("\t\t\t\t0x%08X,\r\n"), GenerateTextHash(InDeclarations));
 	Out.Logf(TEXT("\t\t\t\t%s,\r\n"), *SingletonRange);
 	Out.Logf(TEXT("\t\t\t\t%s\r\n"), *MetaDataParams);
 	Out.Logf(TEXT("\t\t\t};\r\n"));
@@ -2014,7 +2014,7 @@ void FNativeClassHeaderGenerator::ExportNativeGeneratedInitCode(FOutputDevice& O
 					BEGIN_WRAP_EDITOR_ONLY(bIsEditorOnlyFunction),
 					*GetSingletonNameFuncAddr(Function),
 					*FNativeClassHeaderGenerator::GetUTF8OverriddenNameForLiteral(Function),
-					*GetGeneratedCodeCRCTag(Function),
+					*GetGeneratedCodeHashTag(Function),
 					END_WRAP_EDITOR_ONLY(bIsEditorOnlyFunction)
 				);
 			}
@@ -2156,26 +2156,26 @@ void FNativeClassHeaderGenerator::ExportNativeGeneratedInitCode(FOutputDevice& O
 	const FString InitSearchableValuesFunctionParam = InitSearchableValuesFunctionName.IsEmpty() ? FString(TEXT("nullptr")) :
 		FString::Printf(TEXT("&%s::%s"), ClassNameCPP, *InitSearchableValuesFunctionName);
 
-	// Append base class' CRC at the end of the generated code, this will force update derived classes
+	// Append base class' hash at the end of the generated code, this will force update derived classes
 	// when base class changes during hot-reload.
-	uint32 BaseClassCRC = 0;
+	uint32 BaseClassHash = 0;
 	if (Class->GetSuperClass() && !Class->GetSuperClass()->HasAnyClassFlags(CLASS_Intrinsic))
 	{
-		BaseClassCRC = GGeneratedCodeCRCs.FindChecked(Class->GetSuperClass());
+		BaseClassHash = GGeneratedCodeHashes.FindChecked(Class->GetSuperClass());
 	}
-	GeneratedClassRegisterFunctionText.Logf(TEXT("\r\n// %u\r\n"), BaseClassCRC);
+	GeneratedClassRegisterFunctionText.Logf(TEXT("\r\n// %u\r\n"), BaseClassHash);
 
-	// Calculate generated class initialization code CRC so that we know when it changes after hot-reload
-	uint32 ClassCrc = GenerateTextCRC(*GeneratedClassRegisterFunctionText);
-	GGeneratedCodeCRCs.Add(Class, ClassCrc);
+	// Calculate generated class initialization code hash so that we know when it changes after hot-reload
+	uint32 ClassHash = GenerateTextHash(*GeneratedClassRegisterFunctionText);
+	GGeneratedCodeHashes.Add(Class, ClassHash);
 	// Emit the IMPLEMENT_CLASS macro to go in the generated cpp file.
 	if (!bIsDynamic)
 	{
-		Out.Logf(TEXT("\tIMPLEMENT_CLASS(%s, %u);\r\n"), ClassNameCPP, ClassCrc);
+		Out.Logf(TEXT("\tIMPLEMENT_CLASS(%s, %u);\r\n"), ClassNameCPP, ClassHash);
 	}
 	else
 	{
-		Out.Logf(TEXT("\tIMPLEMENT_DYNAMIC_CLASS(%s, TEXT(\"%s\"), %u);\r\n"), ClassNameCPP, *OverriddenClassName, ClassCrc);
+		Out.Logf(TEXT("\tIMPLEMENT_DYNAMIC_CLASS(%s, TEXT(\"%s\"), %u);\r\n"), ClassNameCPP, *OverriddenClassName, ClassHash);
 	}
 
 	Out.Logf(TEXT("\tstatic FCompiledInDefer Z_CompiledInDefer_UClass_%s(%s, &%s::StaticClass, TEXT(\"%s\"), TEXT(\"%s\"), %s, %s, %s, %s);\r\n"),
@@ -2319,8 +2319,8 @@ void FNativeClassHeaderGenerator::ExportFunction(FOutputDevice& Out, const FUnre
 		CurrentFunctionText.Logf(EndEditorOnlyGuard);
 	}
 
-	uint32 FunctionCrc = GenerateTextCRC(*CurrentFunctionText);
-	GGeneratedCodeCRCs.Add(Function, FunctionCrc);
+	uint32 FunctionHash = GenerateTextHash(*CurrentFunctionText);
+	GGeneratedCodeHashes.Add(Function, FunctionHash);
 	Out.Log(CurrentFunctionText);
 }
 
@@ -3239,7 +3239,7 @@ void FNativeClassHeaderGenerator::ExportGeneratedStructBodyMacros(FOutputDevice&
 		const FString Macroized = Macroize(*MacroName, *CombinedLine);
 		OutGeneratedHeaderText.Log(*Macroized);
 
-		FString GetCRCName = FString::Printf(TEXT("Get_%s_CRC"), *SingletonName.LeftChop(2));
+		FString GetHashName = FString::Printf(TEXT("Get_%s_Hash"), *SingletonName.LeftChop(2));
 
 		Out.Logf(TEXT("class UScriptStruct* %s::StaticStruct()\r\n"), StructNameCPP);
 		Out.Logf(TEXT("{\r\n"));
@@ -3260,10 +3260,10 @@ void FNativeClassHeaderGenerator::ExportGeneratedStructBodyMacros(FOutputDevice&
 
 		Out.Logf(TEXT("\tif (!Singleton)\r\n"));
 		Out.Logf(TEXT("\t{\r\n"));
-		Out.Logf(TEXT("\t\textern %suint32 %s();\r\n"), *FriendApiString, *GetCRCName);
+		Out.Logf(TEXT("\t\textern %suint32 %s();\r\n"), *FriendApiString, *GetHashName);
 
 		Out.Logf(TEXT("\t\tSingleton = GetStaticStruct(%s, %s, TEXT(\"%s\"), sizeof(%s), %s());\r\n"),
-			*SingletonName.LeftChop(2), *OuterName, *ActualStructName, StructNameCPP, *GetCRCName);
+			*SingletonName.LeftChop(2), *OuterName, *ActualStructName, StructNameCPP, *GetHashName);
 
 		Out.Logf(TEXT("\t}\r\n"));
 		Out.Logf(TEXT("\treturn Singleton;\r\n"));
@@ -3395,23 +3395,23 @@ void FNativeClassHeaderGenerator::ExportGeneratedStructBodyMacros(FOutputDevice&
 		NoExportStructNameCPP = StructNameCPP;
 	}
 
-	FString CRCFuncName = FString::Printf(TEXT("Get_%s_CRC"), *SingletonName.Replace(TEXT("()"), TEXT(""), ESearchCase::CaseSensitive));
+	FString HashFuncName = FString::Printf(TEXT("Get_%s_Hash"), *SingletonName.Replace(TEXT("()"), TEXT(""), ESearchCase::CaseSensitive));
 	// Structs can either have a UClass or UPackage as outer (if declared in non-UClass header).
 	if (!bIsDynamic)
 	{
 		GeneratedStructRegisterFunctionText.Log (TEXT("#if WITH_HOT_RELOAD\r\n"));
-		GeneratedStructRegisterFunctionText.Logf(TEXT("\t\textern uint32 %s();\r\n"), *CRCFuncName);
+		GeneratedStructRegisterFunctionText.Logf(TEXT("\t\textern uint32 %s();\r\n"), *HashFuncName);
 		GeneratedStructRegisterFunctionText.Logf(TEXT("\t\tUPackage* Outer = %s;\r\n"), *GetPackageSingletonName(CastChecked<UPackage>(Struct->GetOuter())));
-		GeneratedStructRegisterFunctionText.Logf(TEXT("\t\tstatic UScriptStruct* ReturnStruct = FindExistingStructIfHotReloadOrDynamic(Outer, TEXT(\"%s\"), sizeof(%s), %s(), false);\r\n"), *ActualStructName, *NoExportStructNameCPP, *CRCFuncName);
+		GeneratedStructRegisterFunctionText.Logf(TEXT("\t\tstatic UScriptStruct* ReturnStruct = FindExistingStructIfHotReloadOrDynamic(Outer, TEXT(\"%s\"), sizeof(%s), %s(), false);\r\n"), *ActualStructName, *NoExportStructNameCPP, *HashFuncName);
 		GeneratedStructRegisterFunctionText.Log (TEXT("#else\r\n"));
 		GeneratedStructRegisterFunctionText.Logf(TEXT("\t\tstatic UScriptStruct* ReturnStruct = nullptr;\r\n"));
 		GeneratedStructRegisterFunctionText.Log (TEXT("#endif\r\n"));
 	}
 	else
 	{
-		GeneratedStructRegisterFunctionText.Logf(TEXT("\t\textern uint32 %s();\r\n"), *CRCFuncName);
+		GeneratedStructRegisterFunctionText.Logf(TEXT("\t\textern uint32 %s();\r\n"), *HashFuncName);
 		GeneratedStructRegisterFunctionText.Logf(TEXT("\t\tUPackage* Outer = FindOrConstructDynamicTypePackage(TEXT(\"%s\"));\r\n"), *FClass::GetTypePackageName(Struct));
-		GeneratedStructRegisterFunctionText.Logf(TEXT("\t\tUScriptStruct* ReturnStruct = FindExistingStructIfHotReloadOrDynamic(Outer, TEXT(\"%s\"), sizeof(%s), %s(), true);\r\n"), *ActualStructName, *NoExportStructNameCPP, *CRCFuncName);
+		GeneratedStructRegisterFunctionText.Logf(TEXT("\t\tUScriptStruct* ReturnStruct = FindExistingStructIfHotReloadOrDynamic(Outer, TEXT(\"%s\"), sizeof(%s), %s(), true);\r\n"), *ActualStructName, *NoExportStructNameCPP, *HashFuncName);
 	}
 	GeneratedStructRegisterFunctionText.Logf(TEXT("\t\tif (!ReturnStruct)\r\n"));
 	GeneratedStructRegisterFunctionText.Log (TEXT("\t\t{\r\n"));
@@ -3421,11 +3421,11 @@ void FNativeClassHeaderGenerator::ExportGeneratedStructBodyMacros(FOutputDevice&
 	GeneratedStructRegisterFunctionText.Logf(TEXT("\t\treturn ReturnStruct;\r\n"));
 	GeneratedStructRegisterFunctionText.Log (TEXT("\t}\r\n"));
 
-	uint32 StructCrc = GenerateTextCRC(*GeneratedStructRegisterFunctionText);
-	GGeneratedCodeCRCs.Add(Struct, StructCrc);
+	uint32 StructHash = GenerateTextHash(*GeneratedStructRegisterFunctionText);
+	GGeneratedCodeHashes.Add(Struct, StructHash);
 
 	Out.Log(GeneratedStructRegisterFunctionText);
-	Out.Logf(TEXT("\tuint32 %s() { return %uU; }\r\n"), *CRCFuncName, StructCrc);
+	Out.Logf(TEXT("\tuint32 %s() { return %uU; }\r\n"), *HashFuncName, StructHash);
 }
 
 void FNativeClassHeaderGenerator::ExportGeneratedEnumInitCode(FOutputDevice& Out, const FUnrealSourceFile& SourceFile, UEnum* Enum)
@@ -3489,7 +3489,7 @@ void FNativeClassHeaderGenerator::ExportGeneratedEnumInitCode(FOutputDevice& Out
 	);
 
 	const FString EnumSingletonName = GetSingletonName(Enum);
-	const FString CRCFuncName       = FString::Printf(TEXT("Get_%s_CRC"), *SingletonName);
+	const FString HashFuncName       = FString::Printf(TEXT("Get_%s_Hash"), *SingletonName);
 
 	FUHTStringBuilder GeneratedEnumRegisterFunctionText;
 
@@ -3503,7 +3503,7 @@ void FNativeClassHeaderGenerator::ExportGeneratedEnumInitCode(FOutputDevice& Out
 		OuterString = PackageSingletonName;
 		GeneratedEnumRegisterFunctionText.Logf(TEXT("#if WITH_HOT_RELOAD\r\n"));
 		GeneratedEnumRegisterFunctionText.Logf(TEXT("\t\tUPackage* Outer = %s;\r\n"), *OuterString);
-		GeneratedEnumRegisterFunctionText.Logf(TEXT("\t\tstatic UEnum* ReturnEnum = FindExistingEnumIfHotReloadOrDynamic(Outer, TEXT(\"%s\"), 0, %s(), false);\r\n"), *EnumNameCpp, *CRCFuncName);
+		GeneratedEnumRegisterFunctionText.Logf(TEXT("\t\tstatic UEnum* ReturnEnum = FindExistingEnumIfHotReloadOrDynamic(Outer, TEXT(\"%s\"), 0, %s(), false);\r\n"), *EnumNameCpp, *HashFuncName);
 		GeneratedEnumRegisterFunctionText.Logf(TEXT("#else\r\n"));
 		GeneratedEnumRegisterFunctionText.Logf(TEXT("\t\tstatic UEnum* ReturnEnum = nullptr;\r\n"));
 		GeneratedEnumRegisterFunctionText.Logf(TEXT("#endif // WITH_HOT_RELOAD\r\n"));
@@ -3512,7 +3512,7 @@ void FNativeClassHeaderGenerator::ExportGeneratedEnumInitCode(FOutputDevice& Out
 	{
 		OuterString = FString::Printf(TEXT("[](){ return (UObject*)FindOrConstructDynamicTypePackage(TEXT(\"%s\")); }()"), *PackageSingletonName);
 		GeneratedEnumRegisterFunctionText.Logf(TEXT("\t\tUPackage* Outer = FindOrConstructDynamicTypePackage(TEXT(\"%s\"));"), *PackageSingletonName);
-		GeneratedEnumRegisterFunctionText.Logf(TEXT("\t\tUEnum* ReturnEnum = FindExistingEnumIfHotReloadOrDynamic(Outer, TEXT(\"%s\"), 0, %s(), true);\r\n"), *OverriddenEnumNameCpp, *CRCFuncName);
+		GeneratedEnumRegisterFunctionText.Logf(TEXT("\t\tUEnum* ReturnEnum = FindExistingEnumIfHotReloadOrDynamic(Outer, TEXT(\"%s\"), 0, %s(), true);\r\n"), *OverriddenEnumNameCpp, *HashFuncName);
 	}
 	GeneratedEnumRegisterFunctionText.Logf(TEXT("\t\tif (!ReturnEnum)\r\n"));
 	GeneratedEnumRegisterFunctionText.Logf(TEXT("\t\t{\r\n"));
@@ -3558,9 +3558,9 @@ void FNativeClassHeaderGenerator::ExportGeneratedEnumInitCode(FOutputDevice& Out
 	GeneratedEnumRegisterFunctionText.Logf(TEXT("\t\treturn ReturnEnum;\r\n"));
 	GeneratedEnumRegisterFunctionText.Logf(TEXT("\t}\r\n"));
 
-	uint32 EnumCrc = GenerateTextCRC(*GeneratedEnumRegisterFunctionText);
-	GGeneratedCodeCRCs.Add(Enum, EnumCrc);
-	Out.Logf(TEXT("\tuint32 %s() { return %uU; }\r\n"), *CRCFuncName, EnumCrc);
+	uint32 EnumHash = GenerateTextHash(*GeneratedEnumRegisterFunctionText);
+	GGeneratedCodeHashes.Add(Enum, EnumHash);
+	Out.Logf(TEXT("\tuint32 %s() { return %uU; }\r\n"), *HashFuncName, EnumHash);
 	Out.Log(GeneratedEnumRegisterFunctionText);
 }
 
@@ -5184,23 +5184,23 @@ FNativeClassHeaderGenerator::FNativeClassHeaderGenerator(
 
 	if (GeneratedFunctionDeclarations.Len())
 	{
-		uint32 CombinedCRC = 0;
+		uint32 CombinedHash = 0;
 		for (const TPair<FUnrealSourceFile*, FGeneratedCPP>& GeneratedCPP : GeneratedCPPs)
 		{
-			uint32 SplitCRC = GenerateTextCRC(*GeneratedCPP.Value.GeneratedText);
-			if (CombinedCRC == 0)
+			uint32 SplitHash = GenerateTextHash(*GeneratedCPP.Value.GeneratedText);
+			if (CombinedHash == 0)
 			{
 				// Don't combine in the first case because it keeps GUID backwards compatibility
-				CombinedCRC = SplitCRC;
+				CombinedHash = SplitHash;
 			}
 			else
 			{
-				CombinedCRC = HashCombine(SplitCRC, CombinedCRC);
+				CombinedHash = HashCombine(SplitHash, CombinedHash);
 			}
 		}
 
 		FGeneratedCPP& GeneratedCPP = GeneratedCPPs.Emplace(nullptr, GeneratedIncludeDirectory / FString::Printf(TEXT("%s.init.gen.cpp"), *PackageName));
-		ExportGeneratedPackageInitCode(GeneratedCPP.GeneratedText, *GeneratedFunctionDeclarations, Package, CombinedCRC);
+		ExportGeneratedPackageInitCode(GeneratedCPP.GeneratedText, *GeneratedFunctionDeclarations, Package, CombinedHash);
 	}
 
 	const FManifestModule* ModuleInfo = GPackageToManifestModuleMap.FindChecked(Package);
diff --git a/Engine/Source/Programs/UnrealHeaderTool/Private/StringUtils.cpp b/Engine/Source/Programs/UnrealHeaderTool/Private/StringUtils.cpp
--- a/Engine/Source/Programs/UnrealHeaderTool/Private/StringUtils.cpp
+++ b/Engine/Source/Programs/UnrealHeaderTool/Private/StringUtils.cpp
@@ -117,7 +117,7 @@ FORCEINLINE uint64 GenerateTextHash64(const TCHAR* Str)
 	}
 }
 
-uint32 GenerateTextCRC(const TCHAR* Data)
+uint32 GenerateTextHash(const TCHAR* Data)
 {
 	uint64 Hash = GenerateTextHash64(Data);
 	return static_cast<uint32>(Hash) + static_cast<uint32>(Hash >> 32);
diff --git a/Engine/Source/Programs/UnrealHeaderTool/Private/StringUtils.h b/Engine/Source/Programs/UnrealHeaderTool/Private/StringUtils.h
--- a/Engine/Source/Programs/UnrealHeaderTool/Private/StringUtils.h
+++ b/Engine/Source/Programs/UnrealHeaderTool/Private/StringUtils.h
@@ -35,9 +35,9 @@ FString GetClassPrefix(const FString InClassName);
 FString GetClassPrefix(const FString InClassName, bool& bIsLabeledDeprecated);
 
 /**
-* Generates CRC value for text, ignoring characters such as \r and the terminating \0
+* Generates hash for text, ignoring characters such as \r and the terminating \0
 *
-* @param InText Text to generate CRC for 
-* @return CRC value generated for the specified text buffer
+* @param InText Text to digest
+* @return Hash generated for the specified text buffer
 */
-uint32 GenerateTextCRC(const TCHAR* InText);
+uint32 GenerateTextHash(const TCHAR* InText);
