UHT avoid needlessly creating an Array from Set

diff --git a/Engine/Source/Programs/UnrealHeaderTool/Private/CodeGenerator.cpp b/Engine/Source/Programs/UnrealHeaderTool/Private/CodeGenerator.cpp
--- a/Engine/Source/Programs/UnrealHeaderTool/Private/CodeGenerator.cpp
+++ b/Engine/Source/Programs/UnrealHeaderTool/Private/CodeGenerator.cpp
@@ -4868,7 +4868,7 @@ void GetSourceFilesInDependencyOrderRecursive(TArray<FUnrealSourceFile*>& OutTes
 	}
 }
 
-TArray<FUnrealSourceFile*> GetSourceFilesInDependencyOrder(const UPackage* Package, const TArray<FUnrealSourceFile*>& SourceFiles, const TSet<FUnrealSourceFile*>& Ignore)
+TArray<FUnrealSourceFile*> GetSourceFilesInDependencyOrder(const UPackage* Package, const TSet<FUnrealSourceFile*>& SourceFiles, const TSet<FUnrealSourceFile*>& Ignore)
 {
 	TArray<FUnrealSourceFile*> Result;
 	TSet<const FUnrealSourceFile*>	VisitedSet;
@@ -4929,7 +4929,7 @@ static void RecordPackageSingletons(
 // Constructor.
 FNativeClassHeaderGenerator::FNativeClassHeaderGenerator(
 	const UPackage* InPackage,
-	const TArray<FUnrealSourceFile*>& SourceFiles,
+	const TSet<FUnrealSourceFile*>& SourceFiles,
 	FClasses& AllClasses,
 	bool InAllowSaveExportedHeaders
 )
diff --git a/Engine/Source/Programs/UnrealHeaderTool/Private/HeaderParser.cpp b/Engine/Source/Programs/UnrealHeaderTool/Private/HeaderParser.cpp
--- a/Engine/Source/Programs/UnrealHeaderTool/Private/HeaderParser.cpp
+++ b/Engine/Source/Programs/UnrealHeaderTool/Private/HeaderParser.cpp
@@ -7708,7 +7708,7 @@ bool FHeaderParser::DependentClassNameFromHeader(const TCHAR* HeaderFilename, FS
  *
  * @returns Array of source files.
  */
-TArray<FUnrealSourceFile*> GetSourceFilesWithInheritanceOrdering(UPackage* CurrentPackage, FClasses& AllClasses)
+TSet<FUnrealSourceFile*> GetSourceFilesWithInheritanceOrdering(UPackage* CurrentPackage, FClasses& AllClasses)
 {
 	TSet<FUnrealSourceFile*> SourceFiles;
 
@@ -7745,7 +7745,7 @@ TArray<FUnrealSourceFile*> GetSourceFilesWithInheritanceOrdering(UPackage* Curre
 		}
 	}
 
-	return SourceFiles.Array();
+	return SourceFiles;
 }
 
 // Begins the process of exporting C++ class declarations for native classes in the specified package
@@ -7760,7 +7760,7 @@ void FHeaderParser::ExportNativeHeaders(
 	TArray<FString>	ClassHeaderFilenames;
 	new (ClassHeaderFilenames) FString();
 
-	TArray<FUnrealSourceFile*> SourceFiles = GetSourceFilesWithInheritanceOrdering(CurrentPackage, AllClasses);
+	TSet<FUnrealSourceFile*> SourceFiles = GetSourceFilesWithInheritanceOrdering(CurrentPackage, AllClasses);
 	if (SourceFiles.Num() > 0)
 	{
 		if ( CurrentPackage != NULL )
diff --git a/Engine/Source/Programs/UnrealHeaderTool/Private/NativeClassExporter.h b/Engine/Source/Programs/UnrealHeaderTool/Private/NativeClassExporter.h
--- a/Engine/Source/Programs/UnrealHeaderTool/Private/NativeClassExporter.h
+++ b/Engine/Source/Programs/UnrealHeaderTool/Private/NativeClassExporter.h
@@ -467,7 +467,7 @@ public:
 	// Constructor
 	FNativeClassHeaderGenerator(
 		const UPackage* InPackage,
-		const TArray<FUnrealSourceFile*>& SourceFiles,
+		const TSet<FUnrealSourceFile*>& SourceFiles,
 		FClasses& AllClasses,
 		bool InAllowSaveExportedHeaders
 	);
