Const-correction of some FHeaderParser member functions.

diff --git a/Engine/Source/Programs/UnrealHeaderTool/Private/HeaderParser.cpp b/Engine/Source/Programs/UnrealHeaderTool/Private/HeaderParser.cpp
--- a/Engine/Source/Programs/UnrealHeaderTool/Private/HeaderParser.cpp
+++ b/Engine/Source/Programs/UnrealHeaderTool/Private/HeaderParser.cpp
@@ -3135,7 +3135,7 @@ void FHeaderParser::GetVarType(
 	FScope*                         Scope,
 	FPropertyBase&                  VarProperty,
 	EPropertyFlags                  Disallow,
-	FToken*                         OuterPropertyType,
+	const FToken*                   OuterPropertyType,
 	EPropertyDeclarationStyle::Type PropertyDeclarationStyle,
 	EVariableCategory::Type         VariableCategory,
 	FIndexRange*                    ParsedVarIndexRange
@@ -7020,7 +7020,7 @@ bool FHeaderParser::IsBitfieldProperty()
 	return bIsBitfield;
 }
 
-void FHeaderParser::ValidatePropertyIsDeprecatedIfNecessary(FPropertyBase& VarProperty, FToken* OuterPropertyType)
+void FHeaderParser::ValidatePropertyIsDeprecatedIfNecessary(FPropertyBase& VarProperty, const FToken* OuterPropertyType)
 {
 	// check to see if we have a UClassProperty using a deprecated class
 	if ( VarProperty.MetaClass != NULL && VarProperty.MetaClass->HasAnyClassFlags(CLASS_Deprecated) && !(VarProperty.PropertyFlags & CPF_Deprecated) &&
diff --git a/Engine/Source/Programs/UnrealHeaderTool/Private/HeaderParser.h b/Engine/Source/Programs/UnrealHeaderTool/Private/HeaderParser.h
--- a/Engine/Source/Programs/UnrealHeaderTool/Private/HeaderParser.h
+++ b/Engine/Source/Programs/UnrealHeaderTool/Private/HeaderParser.h
@@ -576,7 +576,7 @@ protected:
 		FScope*							Scope,
 		FPropertyBase&                  VarProperty,
 		EPropertyFlags                  Disallow,
-		FToken*                         OuterPropertyType,
+		const FToken*                   OuterPropertyType,
 		EPropertyDeclarationStyle::Type PropertyDeclarationStyle,
 		EVariableCategory::Type         VariableCategory,
 		FIndexRange*                    ParsedVarIndexRange = nullptr);
@@ -702,7 +702,7 @@ protected:
 	 */
 	EFindName GetFindFlagForPropertyName(const TCHAR* PropertyName);
 
-	static void ValidatePropertyIsDeprecatedIfNecessary(FPropertyBase& VarProperty, FToken* OuterPropertyType);
+	static void ValidatePropertyIsDeprecatedIfNecessary(FPropertyBase& VarProperty, const FToken* OuterPropertyType);
 
 private:
 	// Source file currently parsed by UHT.
