Removal of confusing OriginalVarTypeFlags, which is always CPF_None.
Replacement of MakeShareable() with MakeShared().

diff --git a/Engine/Source/Programs/UnrealHeaderTool/Private/HeaderParser.cpp b/Engine/Source/Programs/UnrealHeaderTool/Private/HeaderParser.cpp
--- a/Engine/Source/Programs/UnrealHeaderTool/Private/HeaderParser.cpp
+++ b/Engine/Source/Programs/UnrealHeaderTool/Private/HeaderParser.cpp
@@ -3796,10 +3796,7 @@ void FHeaderParser::GetVarType(
 	{
 		RequireSymbol( TEXT("<"), TEXT("'tarray'") );
 
-		// GetVarType() clears the property flags of the array var, so use dummy 
-		// flags when getting the inner property
-		EPropertyFlags OriginalVarTypeFlags = VarType.PropertyFlags;
-		VarType.PropertyFlags |= Flags;
+		VarType.PropertyFlags = Flags;
 
 		GetVarType(AllClasses, Scope, VarProperty, Disallow, &VarType, EPropertyDeclarationStyle::None, VariableCategory);
 		if (VarProperty.IsContainer())
@@ -3812,8 +3809,7 @@ void FHeaderParser::GetVarType(
 			bNativeConstTemplateArg = true;
 		}
 
-		OriginalVarTypeFlags |= VarProperty.PropertyFlags & (CPF_ContainsInstancedReference | CPF_InstancedReference); // propagate these to the array, we will fix them later
-		VarType.PropertyFlags = OriginalVarTypeFlags;
+		VarType.PropertyFlags = VarProperty.PropertyFlags & (CPF_ContainsInstancedReference | CPF_InstancedReference); // propagate these to the array, we will fix them later
 		VarProperty.ArrayType = EArrayType::Dynamic;
 
 		FToken CloseTemplateToken;
@@ -3844,10 +3840,7 @@ void FHeaderParser::GetVarType(
 	{
 		RequireSymbol( TEXT("<"), TEXT("'tmap'") );
 
-		// GetVarType() clears the property flags of the array var, so use dummy 
-		// flags when getting the inner property
-		EPropertyFlags OriginalVarTypeFlags = VarType.PropertyFlags;
-		VarType.PropertyFlags |= Flags;
+		VarType.PropertyFlags = Flags;
 
 		FToken MapKeyType;
 		GetVarType(AllClasses, Scope, MapKeyType, Disallow, &VarType, EPropertyDeclarationStyle::None, VariableCategory);
@@ -3878,12 +3871,10 @@ void FHeaderParser::GetVarType(
 			FError::Throwf(TEXT("Nested containers are not supported.") );
 		}
 
-		OriginalVarTypeFlags |= VarProperty.PropertyFlags & (CPF_ContainsInstancedReference | CPF_InstancedReference); // propagate these to the map value, we will fix them later
-		OriginalVarTypeFlags |= MapKeyType .PropertyFlags & (CPF_ContainsInstancedReference | CPF_InstancedReference); // propagate these to the map key, we will fix them later
-		VarType.PropertyFlags = OriginalVarTypeFlags;
-		FToken* MapKeyProp = new FToken(MapKeyType);
-		VarProperty.MapKeyProp = MakeShareable<FToken>(MapKeyProp);
-		VarProperty.MapKeyProp->PropertyFlags = OriginalVarTypeFlags | (VarProperty.MapKeyProp->PropertyFlags & CPF_UObjectWrapper); // Make sure the 'UObjectWrapper' flag is maintained so that 'TMap<TSubclassOf<...>, ...>' works
+		EPropertyFlags InnerFlags = (MapKeyType.PropertyFlags | VarProperty.PropertyFlags) & (CPF_ContainsInstancedReference | CPF_InstancedReference); // propagate these to the map value, we will fix them later
+		VarType.PropertyFlags = InnerFlags;
+		VarProperty.MapKeyProp = MakeShared<FToken>(MapKeyType);
+		VarProperty.MapKeyProp->PropertyFlags = InnerFlags | (VarProperty.MapKeyProp->PropertyFlags & CPF_UObjectWrapper); // Make sure the 'UObjectWrapper' flag is maintained so that 'TMap<TSubclassOf<...>, ...>' works
 
 		FToken CloseTemplateToken;
 		if (!GetToken(CloseTemplateToken, /*bNoConsts=*/ true, ESymbolParseOption::CloseTemplateBracket))
@@ -3915,8 +3906,7 @@ void FHeaderParser::GetVarType(
 
 		// GetVarType() clears the property flags of the array var, so use dummy 
 		// flags when getting the inner property
-		EPropertyFlags OriginalVarTypeFlags = VarType.PropertyFlags;
-		VarType.PropertyFlags |= Flags;
+		VarType.PropertyFlags = Flags;
 
 		GetVarType(AllClasses, Scope, VarProperty, Disallow, &VarType, EPropertyDeclarationStyle::None, VariableCategory);
 		if (VarProperty.IsContainer())
@@ -3934,8 +3924,7 @@ void FHeaderParser::GetVarType(
 			FError::Throwf(TEXT("FText is not currently supported as an element type."));
 		}
 
-		OriginalVarTypeFlags |= VarProperty.PropertyFlags & (CPF_ContainsInstancedReference | CPF_InstancedReference); // propagate these to the set, we will fix them later
-		VarType.PropertyFlags = OriginalVarTypeFlags;
+		VarType.PropertyFlags = VarProperty.PropertyFlags & (CPF_ContainsInstancedReference | CPF_InstancedReference); // propagate these to the set, we will fix them later
 		VarProperty.ArrayType = EArrayType::Set;
 
 		FToken CloseTemplateToken;
