Missed some redundant comments from CL# 4457890.

diff --git a/Engine/Source/Programs/UnrealHeaderTool/Private/HeaderParser.cpp b/Engine/Source/Programs/UnrealHeaderTool/Private/HeaderParser.cpp
--- a/Engine/Source/Programs/UnrealHeaderTool/Private/HeaderParser.cpp
+++ b/Engine/Source/Programs/UnrealHeaderTool/Private/HeaderParser.cpp
@@ -3904,8 +3904,6 @@ void FHeaderParser::GetVarType(
 	{
 		RequireSymbol( TEXT("<"), TEXT("'tset'") );
 
-		// GetVarType() clears the property flags of the array var, so use dummy 
-		// flags when getting the inner property
 		VarType.PropertyFlags = Flags;
 
 		GetVarType(AllClasses, Scope, VarProperty, Disallow, &VarType, EPropertyDeclarationStyle::None, VariableCategory);
