Refactor of UProperty creation code to make it clearer and to fix instanced property propagation from TMap keys to the outer.
Test code.

See: https://udn.unrealengine.com/questions/456513/index.html

diff --git a/Engine/Source/Programs/UnrealHeaderTool/Private/HeaderParser.cpp b/Engine/Source/Programs/UnrealHeaderTool/Private/HeaderParser.cpp
--- a/Engine/Source/Programs/UnrealHeaderTool/Private/HeaderParser.cpp
+++ b/Engine/Source/Programs/UnrealHeaderTool/Private/HeaderParser.cpp
@@ -4788,115 +4788,103 @@ UProperty* FHeaderParser::GetVarNameAndDim
 	// create the FName for the property, splitting (ie Unnamed_3 -> Unnamed,3)
 	FName PropertyName(VarProperty.Identifier, FindFlag);
 
-	// Add property.
-	UProperty* NewProperty = nullptr;
-
+	UProperty* Prev = nullptr;
+	for (TFieldIterator<UProperty> It(Scope, EFieldIteratorFlags::ExcludeSuper); It; ++It)
 	{
-		UProperty* Prev = nullptr;
-	    for (TFieldIterator<UProperty> It(Scope, EFieldIteratorFlags::ExcludeSuper); It; ++It)
-		{
-			Prev = *It;
-		}
+		Prev = *It;
+	}
 
-		UArrayProperty* Array             = nullptr;
-		UMapProperty*   Map               = nullptr;
-		USetProperty*   Set               = nullptr; // TODO: Set Property
-		UProperty*      NewMapKeyProperty = nullptr;
-		UObject*        NewScope          = Scope;
-		int32           ArrayDim          = 1; // 1 = not a static array, 2 = static array
-		if (VarProperty.ArrayType == EArrayType::Dynamic)
-		{
-			Array       = new (EC_InternalUseOnlyConstructor, Scope, PropertyName, ObjectFlags) UArrayProperty(FObjectInitializer());
-			NewScope    = Array;
-			ObjectFlags = RF_Public;
-		}
-		else if (VarProperty.ArrayType == EArrayType::Static)
+	auto PropagateFlagsFromInnerAndHandlePersistentInstanceMetadata = [](EPropertyFlags& DestFlags, const TMap<FName, FString>& InMetaData, UProperty* Inner) {
+		// Copy some of the property flags to the container property.
+		if (Inner->PropertyFlags & (CPF_ContainsInstancedReference | CPF_InstancedReference))
 		{
-			ArrayDim = 2;
-		}
-		else if (VarProperty.ArrayType == EArrayType::Set)
-		{
-			Set               = new (EC_InternalUseOnlyConstructor, Scope, PropertyName, ObjectFlags) USetProperty(FObjectInitializer());
-			NewScope          = Set;
-			ObjectFlags       = RF_Public;
-		}
-		else if (VarProperty.MapKeyProp.IsValid())
-		{
-			Map               = new (EC_InternalUseOnlyConstructor, Scope, PropertyName, ObjectFlags) UMapProperty(FObjectInitializer());
-			NewScope          = Map;
-			ObjectFlags       = RF_Public;
-			NewMapKeyProperty = CreateVariableProperty(*VarProperty.MapKeyProp, NewScope, *(PropertyName.ToString() + TEXT("_Key")), ObjectFlags, VariableCategory, CurrentSrcFile);
+			DestFlags |= CPF_ContainsInstancedReference;
+			DestFlags &= ~(CPF_InstancedReference | CPF_PersistentInstance); //this was propagated to the inner
+
+			if (Inner->PropertyFlags & CPF_PersistentInstance)
+			{
+				TMap<FName, FString> MetaData;
+				AddEditInlineMetaData(MetaData);
+				AddMetaDataToClassData(Inner, InMetaData);
+			}
 		}
+	};
 
-		NewProperty = CreateVariableProperty(VarProperty, NewScope, PropertyName, ObjectFlags, VariableCategory, CurrentSrcFile);
+	UProperty* Result = nullptr;
+	if (VarProperty.ArrayType == EArrayType::Dynamic)
+	{
+		UArrayProperty* Array     = new (EC_InternalUseOnlyConstructor, Scope, PropertyName, ObjectFlags) UArrayProperty(FObjectInitializer());
+		UProperty*      InnerProp = CreateVariableProperty(VarProperty, Array, PropertyName, RF_Public, VariableCategory, CurrentSrcFile);
 
-		auto PropagateFlags = [](EPropertyFlags FlagsToPropagate, FPropertyBase& From, UProperty* To) {
-			// Copy some of the property flags to the inner property.
-			To->PropertyFlags |= (From.PropertyFlags & FlagsToPropagate);
+		Array->Inner         = InnerProp;
+		Array->PropertyFlags = VarProperty.PropertyFlags;
 
-			// Copy some of the property flags to the array property.
-			if (To->PropertyFlags & (CPF_ContainsInstancedReference | CPF_InstancedReference))
-			{
-				From.PropertyFlags |= CPF_ContainsInstancedReference;
-				From.PropertyFlags &= ~(CPF_InstancedReference | CPF_PersistentInstance); //this was propagated to the inner
+		// Propagate flags
+		InnerProp->PropertyFlags |= Array->PropertyFlags & CPF_PropagateToArrayInner;
 
-				if (To->PropertyFlags & CPF_PersistentInstance)
-				{
-					TMap<FName, FString> MetaData;
-					AddEditInlineMetaData(MetaData);
-					AddMetaDataToClassData(To, From.MetaData);
-				}
-			}
-		};
+		PropagateFlagsFromInnerAndHandlePersistentInstanceMetadata(Array->PropertyFlags, VarProperty.MetaData, InnerProp);
 
-		if( Array )
-		{
-			Array->Inner = NewProperty;
+		Result = Array;
+	}
+	else if (VarProperty.ArrayType == EArrayType::Set)
+	{
+		USetProperty* Set       = new (EC_InternalUseOnlyConstructor, Scope, PropertyName, ObjectFlags) USetProperty(FObjectInitializer());
+		UProperty*    InnerProp = CreateVariableProperty(VarProperty, Set, PropertyName, RF_Public, VariableCategory, CurrentSrcFile);
 
-			PropagateFlags(CPF_PropagateToArrayInner, VarProperty, NewProperty);
+		Set->ElementProp   = InnerProp;
+		Set->PropertyFlags = VarProperty.PropertyFlags;
 
-			NewProperty = Array;
-		}
+		// Propagate flags
+		InnerProp->PropertyFlags |= Set->PropertyFlags & CPF_PropagateToSetElement;
 
-		if (Map)
-		{
-			Map->KeyProp   = NewMapKeyProperty;
-			Map->ValueProp = NewProperty;
+		PropagateFlagsFromInnerAndHandlePersistentInstanceMetadata(Set->PropertyFlags, VarProperty.MetaData, InnerProp);
 
-			PropagateFlags(CPF_PropagateToMapKey,   *VarProperty.MapKeyProp, NewMapKeyProperty);
-			PropagateFlags(CPF_PropagateToMapValue, VarProperty,             NewProperty);
+		Result = Set;
+	}
+	else if (VarProperty.MapKeyProp.IsValid())
+	{
+		UMapProperty* Map       = new (EC_InternalUseOnlyConstructor, Scope, PropertyName, ObjectFlags) UMapProperty(FObjectInitializer());
+		UProperty*    KeyProp   = CreateVariableProperty(*VarProperty.MapKeyProp, Map, *(PropertyName.ToString() + TEXT("_Key")), RF_Public, VariableCategory, CurrentSrcFile);
+		UProperty*    ValueProp = CreateVariableProperty(VarProperty,             Map, PropertyName,                              RF_Public, VariableCategory, CurrentSrcFile);
 
-			NewProperty = Map;
-		}
+		Map->KeyProp       = KeyProp;
+		Map->ValueProp     = ValueProp;
+		Map->PropertyFlags = VarProperty.PropertyFlags;
 
-		if (Set)
-		{
-			Set->ElementProp = NewProperty;
+		// Propagate flags
+		KeyProp  ->PropertyFlags |= VarProperty.MapKeyProp->PropertyFlags & CPF_PropagateToMapKey;
+		ValueProp->PropertyFlags |= Map->PropertyFlags                    & CPF_PropagateToMapValue;
 
-			PropagateFlags(CPF_PropagateToSetElement, VarProperty, NewProperty);
+		PropagateFlagsFromInnerAndHandlePersistentInstanceMetadata(Map->PropertyFlags, VarProperty.MapKeyProp->MetaData, KeyProp);
+		PropagateFlagsFromInnerAndHandlePersistentInstanceMetadata(Map->PropertyFlags, VarProperty.MetaData,             ValueProp);
 
-			NewProperty = Set;
-		}
+		Result = Map;
+	}
+	else
+	{
+		Result = CreateVariableProperty(VarProperty, Scope, PropertyName, ObjectFlags, VariableCategory, CurrentSrcFile);
 
-		NewProperty->ArrayDim = ArrayDim;
-		if (ArrayDim == 2)
+		if (VarProperty.ArrayType == EArrayType::Static)
 		{
-			GArrayDimensions.Add(NewProperty, Dimensions.String);
-		}
-		NewProperty->PropertyFlags = VarProperty.PropertyFlags;
-		if (Prev != nullptr)
-		{
-			NewProperty->Next = Prev->Next;
-			Prev->Next = NewProperty;
-		}
-		else
-		{
-			NewProperty->Next = Scope->Children;
-			Scope->Children = NewProperty;
+			Result->ArrayDim = 2; // 2 = static array
+			GArrayDimensions.Add(Result, Dimensions.String);
 		}
+
+		Result->PropertyFlags = VarProperty.PropertyFlags;
 	}
 
-	VarProperty.TokenProperty = NewProperty;
+	if (Prev != nullptr)
+	{
+		Result->Next = Prev->Next;
+		Prev->Next = Result;
+	}
+	else
+	{
+		Result->Next = Scope->Children;
+		Scope->Children = Result;
+	}
+
+	VarProperty.TokenProperty = Result;
 	VarProperty.StartLine = InputLine;
 	VarProperty.StartPos = InputPos;
 	FClassMetaData* ScopeData = GScriptHelper.FindClassData(Scope);
@@ -4905,7 +4893,8 @@ UProperty* FHeaderParser::GetVarNameAndDim
 
 	// if we had any metadata, add it to the class
 	AddMetaDataToClassData(VarProperty.TokenProperty, VarProperty.MetaData);
-	return NewProperty;
+
+	return Result;
 }
 
 /*-----------------------------------------------------------------------------
diff --git a/Engine/Source/Programs/UnrealHeaderTool/Resources/UHTDebugging/TestObject.h b/Engine/Source/Programs/UnrealHeaderTool/Resources/UHTDebugging/TestObject.h
--- a/Engine/Source/Programs/UnrealHeaderTool/Resources/UHTDebugging/TestObject.h
+++ b/Engine/Source/Programs/UnrealHeaderTool/Resources/UHTDebugging/TestObject.h
@@ -14,6 +14,15 @@ struct ITestObject
 {
 };
 
+USTRUCT()
+struct FContainsInstancedProperty
+{
+	GENERATED_BODY()
+
+	UPROPERTY(Instanced)
+	UObject* Prop;
+};
+
 UCLASS()
 class UTestObject : public UObject, public ITestObject
 {
@@ -22,6 +31,21 @@ class UTestObject : public UObject, public ITestObject
 public:
 	UTestObject(const FObjectInitializer& ObjectInitializer = FObjectInitializer::Get());
 
+	UPROPERTY()
+	TArray<FContainsInstancedProperty> InstancedPropertyArray;
+
+	UPROPERTY()
+	TArray<TWeakObjectPtr<UObject>> ObjectWrapperArray;
+
+	UPROPERTY()
+	TSet<FContainsInstancedProperty> InstancedPropertySet;
+
+	UPROPERTY()
+	TMap<FContainsInstancedProperty, TWeakObjectPtr<UObject>> InstancedPropertyToObjectWrapperMap;
+
+	UPROPERTY()
+	TMap<TWeakObjectPtr<UObject>, FContainsInstancedProperty> ObjectWrapperToInstancedPropertyMap;
+
 	UFUNCTION(BlueprintCallable, Category="Random")
 	void TestForNullPtrDefaults(UObject* Obj1 = NULL, UObject* Obj2 = nullptr, UObject* Obj3 = 0);
 
