Support for alignas() parsing in UHT.

Modified from PR submission: https://github.com/EpicGames/UnrealEngine/pull/5214

diff --git a/Engine/Source/Programs/UnrealHeaderTool/Private/BaseParser.cpp b/Engine/Source/Programs/UnrealHeaderTool/Private/BaseParser.cpp
--- a/Engine/Source/Programs/UnrealHeaderTool/Private/BaseParser.cpp
+++ b/Engine/Source/Programs/UnrealHeaderTool/Private/BaseParser.cpp
@@ -843,6 +843,23 @@ bool FBaseParser::MatchConstInt( const TCHAR* Match )
 	return false;
 }
 
+bool FBaseParser::MatchAnyConstInt()
+{
+	FToken Token;
+	if (GetToken(Token))
+	{
+		if( Token.TokenType==TOKEN_Const && (Token.Type == CPT_Int || Token.Type == CPT_Int64) )
+		{
+			return true;
+		}
+		else
+		{
+			UngetToken(Token);
+		}
+	}
+
+	return false;
+}
 
 void FBaseParser::MatchSemi()
 {
@@ -921,6 +938,25 @@ void FBaseParser::RequireSymbol( const TCHAR* Match, const TCHAR* Tag, ESymbolPa
 	}
 }
 
+//
+// Require an integer.
+//
+void FBaseParser::RequireConstInt( const TCHAR* Match, const TCHAR* Tag )
+{
+	if (!MatchConstInt(Match))
+	{
+		FError::Throwf(TEXT("Missing integer '%s' in %s"), Match, Tag );
+	}
+}
+
+void FBaseParser::RequireAnyConstInt( const TCHAR* Tag )
+{
+	if (!MatchAnyConstInt())
+	{
+		FError::Throwf(TEXT("Missing integer in %s"), Tag );
+	}
+}
+
 //
 // Require an identifier.
 //
diff --git a/Engine/Source/Programs/UnrealHeaderTool/Private/BaseParser.h b/Engine/Source/Programs/UnrealHeaderTool/Private/BaseParser.h
--- a/Engine/Source/Programs/UnrealHeaderTool/Private/BaseParser.h
+++ b/Engine/Source/Programs/UnrealHeaderTool/Private/BaseParser.h
@@ -142,6 +142,7 @@ public:
 	bool MatchIdentifier( FName Match );
 	bool MatchIdentifier( const TCHAR* Match );
 	bool MatchConstInt( const TCHAR* Match );
+	bool MatchAnyConstInt();
 	bool PeekIdentifier( FName Match );
 	bool PeekIdentifier( const TCHAR* Match );
 	bool MatchSymbol( const TCHAR* Match, ESymbolParseOption bParseTemplateClosingBracket = ESymbolParseOption::Normal );
@@ -152,6 +153,8 @@ public:
 	void RequireIdentifier( FName Match, const TCHAR* Tag );
 	void RequireIdentifier( const TCHAR* Match, const TCHAR* Tag );
 	void RequireSymbol( const TCHAR* Match, const TCHAR* Tag, ESymbolParseOption bParseTemplateClosingBracket = ESymbolParseOption::Normal );
+	void RequireConstInt( const TCHAR* Match, const TCHAR* Tag );
+	void RequireAnyConstInt( const TCHAR* Tag );
 
 	/** Clears out the stored comment. */
 	void ClearComment();
diff --git a/Engine/Source/Programs/UnrealHeaderTool/Private/HeaderParser.cpp b/Engine/Source/Programs/UnrealHeaderTool/Private/HeaderParser.cpp
--- a/Engine/Source/Programs/UnrealHeaderTool/Private/HeaderParser.cpp
+++ b/Engine/Source/Programs/UnrealHeaderTool/Private/HeaderParser.cpp
@@ -1124,8 +1124,42 @@ namespace
 
 		return bSupportedType || (bIsSupportedMemberVariable && bMemberVariable);
 	}
+
+	void SkipAlignasIfNecessary(FBaseParser& Parser)
+	{
+		if (Parser.MatchIdentifier(TEXT("alignas")))
+		{
+			Parser.RequireSymbol(TEXT("("), TEXT("'alignas'"));
+			Parser.RequireAnyConstInt(TEXT("'alignas'"));
+			Parser.RequireSymbol(TEXT(")"), TEXT("'alignas'"));
+		}
+	}
+
+	void SkipDeprecatedMacroIfNecessary(FBaseParser& Parser)
+	{
+		if (!Parser.MatchIdentifier(TEXT("DEPRECATED")))
+		{
+			return;
+		}
+
+		FToken Token;
+		// DEPRECATED(Version, "Message")
+		Parser.RequireSymbol(TEXT("("), TEXT("DEPRECATED macro"));
+		if (Parser.GetToken(Token) && (Token.Type != CPT_Float || Token.TokenType != TOKEN_Const))
+		{
+			FError::Throwf(TEXT("Expected engine version in DEPRECATED macro"));
+		}
+
+		Parser.RequireSymbol(TEXT(","), TEXT("DEPRECATED macro"));
+		if (Parser.GetToken(Token) && (Token.Type != CPT_String || Token.TokenType != TOKEN_Const))
+		{
+			FError::Throwf(TEXT("Expected deprecation message in DEPRECATED macro"));
+		}
+
+		Parser.RequireSymbol(TEXT(")"), TEXT("DEPRECATED macro"));
+	}
 }
-	
+
 /////////////////////////////////////////////////////
 // FScriptLocation
 
@@ -1370,6 +1404,8 @@ UEnum* FHeaderParser::CompileEnum()
 	}
 	else if (EnumToken.Matches(TEXT("enum"), ESearchCase::CaseSensitive))
 	{
+		SkipAlignasIfNecessary(*this);
+
 		if (!GetIdentifier(EnumToken))
 		{
 			FError::Throwf(TEXT("Missing identifier after enum") );
@@ -1377,6 +1413,10 @@ UEnum* FHeaderParser::CompileEnum()
 
 		if (EnumToken.Matches(TEXT("class"), ESearchCase::CaseSensitive) || EnumToken.Matches(TEXT("struct"), ESearchCase::CaseSensitive))
 		{
+			// You can't actually have an alignas() before the class/struct keyword, but this
+			// makes the parsing easier and illegal syntax will be caught by the compiler anyway.
+			SkipAlignasIfNecessary(*this);
+
 			CppForm       = UEnum::ECppForm::EnumClass;
 			bReadEnumName = GetIdentifier(EnumToken);
 		}
@@ -1492,6 +1532,8 @@ UEnum* FHeaderParser::CompileEnum()
 			// Now handle the inner true enum portion
 			RequireIdentifier(TEXT("enum"), TEXT("'Enum'"));
 
+			SkipAlignasIfNecessary(*this);
+
 			FToken InnerEnumToken;
 			if (!GetIdentifier(InnerEnumToken))
 			{
@@ -2015,7 +2057,11 @@ UScriptStruct* FHeaderParser::CompileStructDeclaration(FClasses& AllClasses)
 	// The required API module for this struct, if any
 	FString RequiredAPIMacroIfPresent;
 
-	SkipDeprecatedMacroIfNecessary();
+	// alignas() can come before or after the deprecation macro.
+	// We can't have both, but the compiler will catch that anyway.
+	SkipAlignasIfNecessary(*this);
+	SkipDeprecatedMacroIfNecessary(*this);
+	SkipAlignasIfNecessary(*this);
 
 	// Read the struct name
 	ParseNameWithPotentialAPIMacroPrefix(/*out*/ StructNameInScript, /*out*/ RequiredAPIMacroIfPresent, TEXT("struct"));
@@ -5637,7 +5683,11 @@ UClass* FHeaderParser::CompileClassDeclaration(FClasses& AllClasses)
 	// New style files have the class name / extends afterwards
 	RequireIdentifier(TEXT("class"), TEXT("Class declaration"));
 
-	SkipDeprecatedMacroIfNecessary();
+	// alignas() can come before or after the deprecation macro.
+	// We can't have both, but the compiler will catch that anyway.
+	SkipAlignasIfNecessary(*this);
+	SkipDeprecatedMacroIfNecessary(*this);
+	SkipAlignasIfNecessary(*this);
 
 	FString DeclaredClassName;
 	FString RequiredAPIMacroIfPresent;
@@ -8629,6 +8679,12 @@ void FHeaderPreParser::ParseClassDeclaration(const TCHAR* Filename, const TCHAR*
 	// Require 'class'
 	RequireIdentifier(TEXT("class"), *ErrorMsg);
 
+	// alignas() can come before or after the deprecation macro.
+	// We can't have both, but the compiler will catch that anyway.
+	SkipAlignasIfNecessary(*this);
+	SkipDeprecatedMacroIfNecessary(*this);
+	SkipAlignasIfNecessary(*this);
+
 	// Read the class name
 	FString RequiredAPIMacroIfPresent;
 	ParseNameWithPotentialAPIMacroPrefix(/*out*/ out_ClassName, /*out*/ RequiredAPIMacroIfPresent, StartingMatchID);
@@ -9169,30 +9225,6 @@ bool FHeaderParser::TryToMatchConstructorParameterList(FToken Token)
 	return true;
 }
 
-void FHeaderParser::SkipDeprecatedMacroIfNecessary()
-{
-	if (!MatchIdentifier(TEXT("DEPRECATED")))
-	{
-		return;
-	}
-
-	FToken Token;
-	// DEPRECATED(Version, "Message")
-	RequireSymbol(TEXT("("), TEXT("DEPRECATED macro"));
-	if (GetToken(Token) && (Token.Type != CPT_Float || Token.TokenType != TOKEN_Const))
-	{
-		FError::Throwf(TEXT("Expected engine version in DEPRECATED macro"));
-	}
-
-	RequireSymbol(TEXT(","), TEXT("DEPRECATED macro"));
-	if (GetToken(Token) && (Token.Type != CPT_String || Token.TokenType != TOKEN_Const))
-	{
-		FError::Throwf(TEXT("Expected deprecation message in DEPRECATED macro"));
-	}
-
-	RequireSymbol(TEXT(")"), TEXT("DEPRECATED macro"));
-}
-
 void FHeaderParser::CompileVersionDeclaration(UStruct* Struct)
 {
 	FUnrealSourceFile* CurrentSourceFilePtr = GetCurrentSourceFile();
diff --git a/Engine/Source/Programs/UnrealHeaderTool/Private/HeaderParser.h b/Engine/Source/Programs/UnrealHeaderTool/Private/HeaderParser.h
--- a/Engine/Source/Programs/UnrealHeaderTool/Private/HeaderParser.h
+++ b/Engine/Source/Programs/UnrealHeaderTool/Private/HeaderParser.h
@@ -725,7 +725,6 @@ private:
 	 * @returns True if matched. False otherwise.
 	 */
 	bool TryToMatchConstructorParameterList(FToken Token);
-	void SkipDeprecatedMacroIfNecessary();
 
 	// Parses possible version declaration in generated code, e.g. GENERATED_BODY(<some_version>).
 	void CompileVersionDeclaration(UStruct* Struct);
diff --git a/Engine/Source/Programs/UnrealHeaderTool/Resources/UHTDebugging/EnumOnlyHeader.h b/Engine/Source/Programs/UnrealHeaderTool/Resources/UHTDebugging/EnumOnlyHeader.h
--- a/Engine/Source/Programs/UnrealHeaderTool/Resources/UHTDebugging/EnumOnlyHeader.h
+++ b/Engine/Source/Programs/UnrealHeaderTool/Resources/UHTDebugging/EnumOnlyHeader.h
@@ -5,7 +5,7 @@
 #include "EnumOnlyHeader.generated.h"
 
 UENUM()
-enum ESomeEnum
+enum EOldEnum
 {
 	One,
 	Two,
@@ -13,7 +13,7 @@ enum ESomeEnum
 };
 
 UENUM()
-namespace ESomeNamespacedEnum
+namespace ENamespacedEnum
 {
 	enum Type
 	{
@@ -30,3 +30,46 @@ enum class ECppEnum : uint8
 	Eight,
 	Nine
 };
+
+UENUM()
+enum struct ECppEnumStruct : uint8
+{
+	Ten,
+	Eleven,
+	Twelve
+};
+
+UENUM()
+enum alignas(8) EAlignedOldEnum
+{
+	Thirteen,
+	Fourteen,
+	Fifteen
+};
+
+UENUM()
+namespace EAlignedNamespacedEnum
+{
+	enum alignas(8) Type
+	{
+		Sixteen,
+		Seventeen,
+		Eighteen
+	};
+}
+
+UENUM()
+enum class alignas(8) EAlignedCppEnum : uint8
+{
+	Nineteen,
+	Twenty,
+	TwentyOne
+};
+
+UENUM()
+enum struct alignas(8) EAlignedCppEnumStruct : uint8
+{
+	TwentyTwo,
+	TwentyThree,
+	TwentyFour
+};
diff --git a/Engine/Source/Programs/UnrealHeaderTool/Resources/UHTDebugging/StructOnlyHeader.h b/Engine/Source/Programs/UnrealHeaderTool/Resources/UHTDebugging/StructOnlyHeader.h
--- a/Engine/Source/Programs/UnrealHeaderTool/Resources/UHTDebugging/StructOnlyHeader.h
+++ b/Engine/Source/Programs/UnrealHeaderTool/Resources/UHTDebugging/StructOnlyHeader.h
@@ -13,3 +13,9 @@ struct FSomeStruct
 
 #pragma endregion X
 };
+
+USTRUCT()
+struct alignas(8) FAlignedStruct
+{
+	GENERATED_BODY()
+};
diff --git a/Engine/Source/Programs/UnrealHeaderTool/Resources/UHTDebugging/TestObject.h b/Engine/Source/Programs/UnrealHeaderTool/Resources/UHTDebugging/TestObject.h
--- a/Engine/Source/Programs/UnrealHeaderTool/Resources/UHTDebugging/TestObject.h
+++ b/Engine/Source/Programs/UnrealHeaderTool/Resources/UHTDebugging/TestObject.h
@@ -23,6 +23,12 @@ struct FContainsInstancedProperty
 	UObject* Prop;
 };
 
+UCLASS()
+class alignas(8) UAlignedObject : public UObject
+{
+	GENERATED_BODY()
+};
+
 UCLASS()
 class UTestObject : public UObject, public ITestObject
 {
