GitHub 5032 : Added templated StaticClass, StaticStruct and StaticEnum free functions, which are specialized for every reflected type by the generated code that UHT produces

diff --git a/Engine/Plugins/2D/Paper2D/Source/Paper2DEditor/Private/SpriteEditor/SpriteDetailsCustomization.cpp b/Engine/Plugins/2D/Paper2D/Source/Paper2DEditor/Private/SpriteEditor/SpriteDetailsCustomization.cpp
--- a/Engine/Plugins/2D/Paper2D/Source/Paper2DEditor/Private/SpriteEditor/SpriteDetailsCustomization.cpp
+++ b/Engine/Plugins/2D/Paper2D/Source/Paper2DEditor/Private/SpriteEditor/SpriteDetailsCustomization.cpp
@@ -238,7 +238,7 @@ void FSpriteDetailsCustomization::BuildCollisionSection(IDetailCategoryBuilder&
 	{
 		// Restrict the diced value
 		TSharedPtr<FPropertyRestriction> PreventDicedRestriction = MakeShareable(new FPropertyRestriction(LOCTEXT("CollisionGeometryDoesNotSupportDiced", "Collision geometry can not be set to Diced")));
-		const UEnum* const SpritePolygonModeEnum = FindObject<UEnum>(ANY_PACKAGE, TEXT("ESpritePolygonMode"));		
+		const UEnum* const SpritePolygonModeEnum = StaticEnum<ESpritePolygonMode::Type>();
 		PreventDicedRestriction->AddDisabledValue(SpritePolygonModeEnum->GetNameStringByValue((uint8)ESpritePolygonMode::Diced));
 
 		// Find and add the property
diff --git a/Engine/Plugins/AI/HTNPlanner/Source/HTNPlanner/Private/Debug/HTNDebug.cpp b/Engine/Plugins/AI/HTNPlanner/Source/HTNPlanner/Private/Debug/HTNDebug.cpp
--- a/Engine/Plugins/AI/HTNPlanner/Source/HTNPlanner/Private/Debug/HTNDebug.cpp
+++ b/Engine/Plugins/AI/HTNPlanner/Source/HTNPlanner/Private/Debug/HTNDebug.cpp
@@ -9,13 +9,13 @@ namespace FHTNDebug
 {
 	FString HTNWorldStateCheckToString(const EHTNWorldStateCheck Value)
 	{
-		static const UEnum* EnumType = FindObject<UEnum>(ANY_PACKAGE, TEXT("EHTNWorldStateCheck"));
+		static const UEnum* EnumType = StaticEnum<EHTNWorldStateCheck>();
 		return EnumType->GetNameStringByIndex(static_cast<uint32>(Value));
 	}
 
 	FString HTNWorldStateOperationToString(const EHTNWorldStateOperation Value)
 	{
-		static const UEnum* EnumType = FindObject<UEnum>(ANY_PACKAGE, TEXT("EHTNWorldStateOperation"));
+		static const UEnum* EnumType = StaticEnum<EHTNWorldStateOperation>();
 		return EnumType->GetNameStringByIndex(static_cast<uint32>(Value));
 	}
 
diff --git a/Engine/Plugins/Experimental/AlembicImporter/Source/AlembicLibrary/Private/AbcImportSettingsCustomization.cpp b/Engine/Plugins/Experimental/AlembicImporter/Source/AlembicLibrary/Private/AbcImportSettingsCustomization.cpp
--- a/Engine/Plugins/Experimental/AlembicImporter/Source/AlembicLibrary/Private/AbcImportSettingsCustomization.cpp
+++ b/Engine/Plugins/Experimental/AlembicImporter/Source/AlembicLibrary/Private/AbcImportSettingsCustomization.cpp
@@ -38,7 +38,7 @@ void FAbcImportSettingsCustomization::CustomizeDetails(IDetailLayoutBuilder& Lay
 
 	if (CurrentSettings && CurrentSettings->bReimport)
 	{
-		UEnum* ImportTypeEnum = FindObject<UEnum>(ANY_PACKAGE, TEXT("EAlembicImportType"));		
+		UEnum* ImportTypeEnum = StaticEnum<EAlembicImportType>();		
 		static FText RestrictReason = NSLOCTEXT("AlembicImportFactory", "ReimportRestriction", "Unable to change type while reimporting");
 		TSharedPtr<FPropertyRestriction> EnumRestriction = MakeShareable(new FPropertyRestriction(RestrictReason));
 
diff --git a/Engine/Plugins/Experimental/PythonScriptPlugin/Source/PythonScriptPlugin/Private/PySlate.cpp b/Engine/Plugins/Experimental/PythonScriptPlugin/Source/PythonScriptPlugin/Private/PySlate.cpp
--- a/Engine/Plugins/Experimental/PythonScriptPlugin/Source/PythonScriptPlugin/Private/PySlate.cpp
+++ b/Engine/Plugins/Experimental/PythonScriptPlugin/Source/PythonScriptPlugin/Private/PySlate.cpp
@@ -140,7 +140,7 @@ PyObject* ParentExternalWindowToSlate(PyObject* InSelf, PyObject* InArgs)
 		return nullptr;
 	}
 
-	static const UEnum* ParentWindowSearchMethodEnum = FindObject<UEnum>(ANY_PACKAGE, TEXT("ESlateParentWindowSearchMethod"));
+	static const UEnum* ParentWindowSearchMethodEnum = StaticEnum<ESlateParentWindowSearchMethod>();
 	ESlateParentWindowSearchMethod ParentWindowSearchMethod = ESlateParentWindowSearchMethod::ActiveWindow;
 	if (PyParentWindowSearchMethod && !PyConversion::NativizeEnumEntry(PyParentWindowSearchMethod, ParentWindowSearchMethodEnum, ParentWindowSearchMethod))
 	{
diff --git a/Engine/Plugins/FX/Niagara/Source/Niagara/Private/NiagaraModule.cpp b/Engine/Plugins/FX/Niagara/Source/Niagara/Private/NiagaraModule.cpp
--- a/Engine/Plugins/FX/Niagara/Source/Niagara/Private/NiagaraModule.cpp
+++ b/Engine/Plugins/FX/Niagara/Source/Niagara/Private/NiagaraModule.cpp
@@ -520,8 +520,8 @@ void FNiagaraTypeDefinition::Init()
 	ScalarStructs.Add(IntStruct);
 	ScalarStructs.Add(FloatStruct);
 
-	ExecutionStateEnum = FindObject<UEnum>(ANY_PACKAGE, TEXT("ENiagaraExecutionState"), true);
-	ExecutionStateSourceEnum = FindObject<UEnum>(ANY_PACKAGE, TEXT("ENiagaraExecutionStateSource"), true);
+	ExecutionStateEnum = StaticEnum<ENiagaraExecutionState>();
+	ExecutionStateSourceEnum = StaticEnum<ENiagaraExecutionStateSource>();
 	
 	RecreateUserDefinedTypeRegistry();
 }
diff --git a/Engine/Plugins/FX/Niagara/Source/NiagaraEditor/Private/NiagaraCompiler.cpp b/Engine/Plugins/FX/Niagara/Source/NiagaraEditor/Private/NiagaraCompiler.cpp
--- a/Engine/Plugins/FX/Niagara/Source/NiagaraEditor/Private/NiagaraCompiler.cpp
+++ b/Engine/Plugins/FX/Niagara/Source/NiagaraEditor/Private/NiagaraCompiler.cpp
@@ -344,8 +344,8 @@ void FNiagaraCompileRequestData::DeepCopyGraphs(UNiagaraScriptSource* ScriptSour
 void FNiagaraCompileRequestData::FinishPrecompile(UNiagaraScriptSource* ScriptSource, const TArray<FNiagaraVariable>& EncounterableVariables, ENiagaraScriptUsage InUsage)
 {
 	{
-		ENiagaraScriptCompileStatusEnum = FindObject<UEnum>(ANY_PACKAGE, TEXT("ENiagaraScriptCompileStatus"), true);
-		ENiagaraScriptUsageEnum = FindObject<UEnum>(ANY_PACKAGE, TEXT("ENiagaraScriptUsage"), true);
+		ENiagaraScriptCompileStatusEnum = StaticEnum<ENiagaraScriptCompileStatus>();
+		ENiagaraScriptUsageEnum = StaticEnum<ENiagaraScriptUsage>();
 
 		PrecompiledHistories.Empty();
 
diff --git a/Engine/Plugins/FX/Niagara/Source/NiagaraEditor/Private/NiagaraEditorUtilities.cpp b/Engine/Plugins/FX/Niagara/Source/NiagaraEditor/Private/NiagaraEditorUtilities.cpp
--- a/Engine/Plugins/FX/Niagara/Source/NiagaraEditor/Private/NiagaraEditorUtilities.cpp
+++ b/Engine/Plugins/FX/Niagara/Source/NiagaraEditor/Private/NiagaraEditorUtilities.cpp
@@ -287,7 +287,7 @@ void FNiagaraEditorUtilities::GatherChangeIds(UNiagaraEmitter& Emitter, TMap<FGu
 				FString KeyString;
 				Id.AppendKeyString(KeyString);
 
-				UEnum* FoundEnum = FindObject<UEnum>(ANY_PACKAGE, TEXT("ENiagaraScriptUsage"), true);
+				UEnum* FoundEnum = StaticEnum<ENiagaraScriptUsage>();
 
 				FString ResultsEnum = TEXT("??");
 				if (FoundEnum)
diff --git a/Engine/Plugins/FX/Niagara/Source/NiagaraEditor/Private/NiagaraGraph.cpp b/Engine/Plugins/FX/Niagara/Source/NiagaraEditor/Private/NiagaraGraph.cpp
--- a/Engine/Plugins/FX/Niagara/Source/NiagaraEditor/Private/NiagaraGraph.cpp
+++ b/Engine/Plugins/FX/Niagara/Source/NiagaraEditor/Private/NiagaraGraph.cpp
@@ -801,7 +801,7 @@ void UNiagaraGraph::RebuildCachedData(bool bForce)
 
 			if (FoundEnum == nullptr)
 			{
-				FoundEnum = FindObject<UEnum>(ANY_PACKAGE, TEXT("ENiagaraScriptUsage"), true);
+				FoundEnum = StaticEnum<ENiagaraScriptUsage>();
 			}
 
 			FString ResultsEnum = TEXT("??");
@@ -908,7 +908,7 @@ void UNiagaraGraph::SynchronizeInternalCacheWithGraph(UNiagaraGraph* Other)
 				{
 					if (FoundEnum == nullptr)
 					{
-						FoundEnum = FindObject<UEnum>(ANY_PACKAGE, TEXT("ENiagaraScriptUsage"), true);
+						FoundEnum = StaticEnum<ENiagaraScriptUsage>();
 					}
 
 					FString ResultsEnum = TEXT("??");
diff --git a/Engine/Plugins/FX/Niagara/Source/NiagaraEditor/Private/NiagaraNodeUsageSelector.cpp b/Engine/Plugins/FX/Niagara/Source/NiagaraEditor/Private/NiagaraNodeUsageSelector.cpp
--- a/Engine/Plugins/FX/Niagara/Source/NiagaraEditor/Private/NiagaraNodeUsageSelector.cpp
+++ b/Engine/Plugins/FX/Niagara/Source/NiagaraEditor/Private/NiagaraNodeUsageSelector.cpp
@@ -36,7 +36,7 @@ void UNiagaraNodeUsageSelector::AllocateDefaultPins()
 {
 	const UEdGraphSchema_Niagara* Schema = GetDefault<UEdGraphSchema_Niagara>();
 
-	UEnum* ENiagaraScriptGroupEnum = FindObject<UEnum>(ANY_PACKAGE, TEXT("ENiagaraScriptGroup"), true);
+	UEnum* ENiagaraScriptGroupEnum = StaticEnum<ENiagaraScriptGroup>();
 
 	//Create the inputs for each path.
 	for (int64 i = 0; i < (int64)ENiagaraScriptGroup::Max; i++)
diff --git a/Engine/Plugins/FX/Niagara/Source/NiagaraEditor/Private/ViewModels/NiagaraEmitterViewModel.cpp b/Engine/Plugins/FX/Niagara/Source/NiagaraEditor/Private/ViewModels/NiagaraEmitterViewModel.cpp
--- a/Engine/Plugins/FX/Niagara/Source/NiagaraEditor/Private/ViewModels/NiagaraEmitterViewModel.cpp
+++ b/Engine/Plugins/FX/Niagara/Source/NiagaraEditor/Private/ViewModels/NiagaraEmitterViewModel.cpp
@@ -36,7 +36,7 @@ FNiagaraEmitterViewModel::FNiagaraEmitterViewModel(UNiagaraEmitter* InEmitter, T
 	, Simulation(InSimulation)
 	, SharedScriptViewModel(MakeShareable(new FNiagaraScriptViewModel(InEmitter, LOCTEXT("SharedDisplayName", "Graph"), ENiagaraParameterEditMode::EditAll)))
 	, bUpdatingSelectionInternally(false)
-	, ExecutionStateEnum(FindObject<UEnum>(ANY_PACKAGE, TEXT("ENiagaraExecutionState")))
+	, ExecutionStateEnum(StaticEnum<ENiagaraExecutionState>())
 {	
 	SetEmitter(InEmitter);
 }
diff --git a/Engine/Plugins/FX/Niagara/Source/NiagaraEditor/Private/Widgets/SNiagaraGeneratedCodeView.cpp b/Engine/Plugins/FX/Niagara/Source/NiagaraEditor/Private/Widgets/SNiagaraGeneratedCodeView.cpp
--- a/Engine/Plugins/FX/Niagara/Source/NiagaraEditor/Private/Widgets/SNiagaraGeneratedCodeView.cpp
+++ b/Engine/Plugins/FX/Niagara/Source/NiagaraEditor/Private/Widgets/SNiagaraGeneratedCodeView.cpp
@@ -29,7 +29,7 @@
 void SNiagaraGeneratedCodeView::Construct(const FArguments& InArgs, TSharedRef<FNiagaraSystemViewModel> InSystemViewModel)
 {
 	TabState = 0;
-	ScriptEnum = FindObject<UEnum>(ANY_PACKAGE, TEXT("ENiagaraScriptUsage"));
+	ScriptEnum = StaticEnum<ENiagaraScriptUsage>();
 	ensure(ScriptEnum);
 
 	SystemViewModel = InSystemViewModel;
diff --git a/Engine/Plugins/FX/Niagara/Source/NiagaraEditor/Private/Widgets/SNiagaraSpreadsheetView.cpp b/Engine/Plugins/FX/Niagara/Source/NiagaraEditor/Private/Widgets/SNiagaraSpreadsheetView.cpp
--- a/Engine/Plugins/FX/Niagara/Source/NiagaraEditor/Private/Widgets/SNiagaraSpreadsheetView.cpp
+++ b/Engine/Plugins/FX/Niagara/Source/NiagaraEditor/Private/Widgets/SNiagaraSpreadsheetView.cpp
@@ -213,7 +213,7 @@ TSharedRef< SWidget > SNiagaraSpreadsheetRow::GenerateWidgetForColumn(const FNam
 void SNiagaraSpreadsheetView::Construct(const FArguments& InArgs, TSharedRef<FNiagaraSystemViewModel> InSystemViewModel)
 {
 	TabState = UIPerParticleUpdate;
-	ScriptEnum = FindObject<UEnum>(ANY_PACKAGE, TEXT("ENiagaraScriptUsage"));
+	ScriptEnum = StaticEnum<ENiagaraScriptUsage>();
 	TargetComponent = InSystemViewModel->GetPreviewComponent();
 	ensure(ScriptEnum);
 	
diff --git a/Engine/Plugins/Lumin/MagicLeap/Source/MagicLeap/Private/MagicLeapHMD.cpp b/Engine/Plugins/Lumin/MagicLeap/Source/MagicLeap/Private/MagicLeapHMD.cpp
--- a/Engine/Plugins/Lumin/MagicLeap/Source/MagicLeap/Private/MagicLeapHMD.cpp
+++ b/Engine/Plugins/Lumin/MagicLeap/Source/MagicLeap/Private/MagicLeapHMD.cpp
@@ -697,7 +697,7 @@ void FMagicLeapHMD::DisplayWarningIfVDZINotEnabled()
 #if PLATFORM_LUMIN
 void FMagicLeapHMD::SetFrameTimingHint(ELuminFrameTimingHint InFrameTimingHint)
 {
-	const static UEnum* FrameTimingEnum = FindObject<UEnum>(ANY_PACKAGE, TEXT("ELuminFrameTimingHint"));
+	const static UEnum* FrameTimingEnum = StaticEnum<ELuminFrameTimingHint>();
 	check(FrameTimingEnum != nullptr);
 
 	if (InFrameTimingHint != CurrentFrameTimingHint)
@@ -1101,7 +1101,7 @@ void FMagicLeapHMD::LoadFromIni()
 
 #if PLATFORM_LUMIN
 	ELuminFrameTimingHint ConfigFrameTimingHint = ELuminFrameTimingHint::Unspecified;
-	const static UEnum* FrameTimingEnum = FindObject<UEnum>(ANY_PACKAGE, TEXT("ELuminFrameTimingHint"));
+	const static UEnum* FrameTimingEnum = StaticEnum<ELuminFrameTimingHint>();
 
 	FString EnumVal;
 	GConfig->GetString(TEXT("/Script/LuminRuntimeSettings.LuminRuntimeSettings"), TEXT("FrameTimingHint"), EnumVal, GEngineIni);
diff --git a/Engine/Plugins/Lumin/MagicLeap/Source/MagicLeapController/Private/MagicLeapController.cpp b/Engine/Plugins/Lumin/MagicLeap/Source/MagicLeapController/Private/MagicLeapController.cpp
--- a/Engine/Plugins/Lumin/MagicLeap/Source/MagicLeapController/Private/MagicLeapController.cpp
+++ b/Engine/Plugins/Lumin/MagicLeap/Source/MagicLeapController/Private/MagicLeapController.cpp
@@ -564,7 +564,7 @@ void FMagicLeapController::Enable()
 	TrackingMode = EMLControllerTrackingMode::CoordinateFrameUID;
 
 	// Pull preference from config file
-	const static UEnum* TrackingModeEnum = FindObject<UEnum>(ANY_PACKAGE, TEXT("EMLControllerTrackingMode"));
+	const static UEnum* TrackingModeEnum = StaticEnum<EMLControllerTrackingMode>();
 
 	FString EnumVal;
 	GConfig->GetString(TEXT("/Script/LuminRuntimeSettings.LuminRuntimeSettings"),
diff --git a/Engine/Plugins/NetcodeUnitTest/NetcodeUnitTest/Source/NetcodeUnitTest/Private/UnitTestManager.cpp b/Engine/Plugins/NetcodeUnitTest/NetcodeUnitTest/Source/NetcodeUnitTest/Private/UnitTestManager.cpp
--- a/Engine/Plugins/NetcodeUnitTest/NetcodeUnitTest/Source/NetcodeUnitTest/Private/UnitTestManager.cpp
+++ b/Engine/Plugins/NetcodeUnitTest/NetcodeUnitTest/Source/NetcodeUnitTest/Private/UnitTestManager.cpp
@@ -902,7 +902,7 @@ void UUnitTestManager::DumpStatus(bool bForce/*=false*/)
 #endif
 void UUnitTestManager::PrintUnitTestResult(UUnitTest* InUnitTest, bool bFinalSummary/*=false*/, bool bUnfinished/*=false*/)
 {
-	static const UEnum* VerificationStateEnum = FindObject<UEnum>(ANY_PACKAGE, TEXT("EUnitTestVerification"));
+	static const UEnum* VerificationStateEnum = StaticEnum<EUnitTestVerification>();
 
 	EUnitTestVerification UnitTestResult = InUnitTest->VerificationState;
 
diff --git a/Engine/Plugins/Online/OnlineFramework/Source/Hotfix/Private/UpdateManager.cpp b/Engine/Plugins/Online/OnlineFramework/Source/Hotfix/Private/UpdateManager.cpp
--- a/Engine/Plugins/Online/OnlineFramework/Source/Hotfix/Private/UpdateManager.cpp
+++ b/Engine/Plugins/Online/OnlineFramework/Source/Hotfix/Private/UpdateManager.cpp
@@ -76,8 +76,8 @@ UUpdateManager::UUpdateManager()
 	LastCompletionResult[0] = LastCompletionResult[1] = EUpdateCompletionStatus::UpdateUnknown;
 	if (!HasAnyFlags(RF_ClassDefaultObject))
 	{
-		UpdateStateEnum = FindObject<UEnum>(ANY_PACKAGE, TEXT("EUpdateState"));
-		UpdateCompletionEnum = FindObject<UEnum>(ANY_PACKAGE, TEXT("EUpdateCompletionStatus"));
+		UpdateStateEnum = StaticEnum<EUpdateState>();
+		UpdateCompletionEnum = StaticEnum<EUpdateCompletionStatus>();
 
 		RegisterDelegates();
 	}
diff --git a/Engine/Plugins/Runtime/AppleARKitFaceSupport/Source/AppleARKitFaceSupport/Private/AppleARKitLiveLinkSource.cpp b/Engine/Plugins/Runtime/AppleARKitFaceSupport/Source/AppleARKitFaceSupport/Private/AppleARKitLiveLinkSource.cpp
--- a/Engine/Plugins/Runtime/AppleARKitFaceSupport/Source/AppleARKitFaceSupport/Private/AppleARKitLiveLinkSource.cpp
+++ b/Engine/Plugins/Runtime/AppleARKitFaceSupport/Source/AppleARKitFaceSupport/Private/AppleARKitLiveLinkSource.cpp
@@ -142,7 +142,7 @@ void FAppleARKitLiveLinkSource::PublishBlendShapes(FName SubjectName, double Tim
 	}
 	LastSubjectName = SubjectName;
 
-	const UEnum* EnumPtr = FindObject<UEnum>(ANY_PACKAGE, TEXT("EARFaceBlendShape"), true);
+	const UEnum* EnumPtr = StaticEnum<EARFaceBlendShape>();
 	if (EnumPtr != nullptr)
 	{
 		static FLiveLinkFrameData LiveLinkFrame;
diff --git a/Engine/Plugins/Runtime/GameplayAbilities/Source/GameplayAbilities/Private/GameplayEffectTypes.cpp b/Engine/Plugins/Runtime/GameplayAbilities/Source/GameplayAbilities/Private/GameplayEffectTypes.cpp
--- a/Engine/Plugins/Runtime/GameplayAbilities/Source/GameplayAbilities/Private/GameplayEffectTypes.cpp
+++ b/Engine/Plugins/Runtime/GameplayAbilities/Source/GameplayAbilities/Private/GameplayEffectTypes.cpp
@@ -25,7 +25,7 @@ FGameplayModEvaluationChannelSettings::FGameplayModEvaluationChannelSettings()
 	// static default channel used to initialize this struct
 	if (!EvalChannelEnum)
 	{
-		EvalChannelEnum = FindObject<UEnum>(ANY_PACKAGE, TEXT("EGameplayModEvaluationChannel"));
+		EvalChannelEnum = StaticEnum<EGameplayModEvaluationChannel>();
 		if (ensure(EvalChannelEnum) && ensure(GConfig))
 		{
 			const FString INISection(TEXT("/Script/GameplayAbilities.AbilitySystemGlobals"));
@@ -352,19 +352,25 @@ bool FGameplayEffectContextHandle::NetSerialize(FArchive& Ar, class UPackageMap*
 
 FString EGameplayModOpToString(int32 Type)
 {
-	static UEnum *e = FindObject<UEnum>(ANY_PACKAGE, TEXT("EGameplayModOp"));
+	static UEnum *e = StaticEnum<EGameplayModOp::Type>();
 	return e->GetNameStringByValue(Type);
 }
 
+FString EGameplayModToString(int32 Type)
+{
+	// Enum no longer exists.
+	return FString();
+}
+
 FString EGameplayModEffectToString(int32 Type)
 {
-	static UEnum *e = FindObject<UEnum>(ANY_PACKAGE, TEXT("EGameplayModEffect"));
-	return e->GetNameStringByValue(Type);
+	// Enum no longer exists.
+	return FString();
 }
 
 FString EGameplayCueEventToString(int32 Type)
 {
-	static UEnum *e = FindObject<UEnum>(ANY_PACKAGE, TEXT("EGameplayCueEvent"));
+	static UEnum *e = StaticEnum<EGameplayCueEvent::Type>();
 	return e->GetNameStringByValue(Type);
 }
 
diff --git a/Engine/Plugins/Runtime/GameplayAbilities/Source/GameplayAbilitiesEditor/Private/GameplayAbilitiesEditorModule.cpp b/Engine/Plugins/Runtime/GameplayAbilities/Source/GameplayAbilitiesEditor/Private/GameplayAbilitiesEditorModule.cpp
--- a/Engine/Plugins/Runtime/GameplayAbilities/Source/GameplayAbilitiesEditor/Private/GameplayAbilitiesEditorModule.cpp
+++ b/Engine/Plugins/Runtime/GameplayAbilities/Source/GameplayAbilitiesEditor/Private/GameplayAbilitiesEditorModule.cpp
@@ -363,7 +363,7 @@ FAutoConsoleCommand RecompileGameplayAbilitiesEditorCommand(
 void FGameplayAbilitiesEditorModule::ApplyGameplayModEvaluationChannelAliasesToEnumMetadata()
 {
 	UAbilitySystemGlobals* AbilitySystemGlobalsCDO = UAbilitySystemGlobals::StaticClass()->GetDefaultObject<UAbilitySystemGlobals>();
-	const UEnum* EvalChannelEnum = FindObject<UEnum>(ANY_PACKAGE, TEXT("EGameplayModEvaluationChannel"));
+	const UEnum* EvalChannelEnum = StaticEnum<EGameplayModEvaluationChannel>();
 	if (ensure(EvalChannelEnum) && ensure(AbilitySystemGlobalsCDO))
 	{
 		const TCHAR* DisplayNameMeta = TEXT("DisplayName");
@@ -396,7 +396,7 @@ void FGameplayAbilitiesEditorModule::ApplyGameplayModEvaluationChannelAliasesToE
 		else
 		{
 			// If not allowed to use channels, also hide the "Evaluate up to channel" option 
-			const UEnum* AttributeBasedFloatCalculationTypeEnum = FindObject<UEnum>(ANY_PACKAGE, TEXT("EAttributeBasedFloatCalculationType"));
+			const UEnum* AttributeBasedFloatCalculationTypeEnum = StaticEnum<EAttributeBasedFloatCalculationType>();
 			if (ensure(AttributeBasedFloatCalculationTypeEnum))
 			{
 				const int32 ChannelBasedCalcIdx = AttributeBasedFloatCalculationTypeEnum->GetIndexByValue(static_cast<int64>(EAttributeBasedFloatCalculationType::AttributeMagnitudeEvaluatedUpToChannel));
diff --git a/Engine/Plugins/Runtime/GoogleVR/GoogleVRController/Source/GoogleVRController/Private/GoogleVRController.cpp b/Engine/Plugins/Runtime/GoogleVR/GoogleVRController/Source/GoogleVRController/Private/GoogleVRController.cpp
--- a/Engine/Plugins/Runtime/GoogleVR/GoogleVRController/Source/GoogleVRController/Private/GoogleVRController.cpp
+++ b/Engine/Plugins/Runtime/GoogleVR/GoogleVRController/Source/GoogleVRController/Private/GoogleVRController.cpp
@@ -148,7 +148,7 @@ FGoogleVRController::FGoogleVRController(const TSharedRef< FGenericApplicationMe
 		GoogleVRCaps = EGoogleVRCaps::Cardboard;
 		FString ValueString;
 		GConfig->GetString(TEXT("/Script/AndroidRuntimeSettings.AndroidRuntimeSettings"), TEXT("GoogleVRCaps"), ValueString, GEngineIni);
-		UEnum* Enum = FindObject<UEnum>((UObject*)ANY_PACKAGE, TEXT("EGoogleVRCaps"), true);
+		UEnum* Enum = StaticEnum<EGoogleVRCaps::Type>();
 		if (Enum)
 		{
 			int64 Value = Enum->GetValueByName(FName(*ValueString));
diff --git a/Engine/Plugins/Runtime/GoogleVR/GoogleVRHMD/Source/GoogleVRHMD/Private/GoogleVRHMD.cpp b/Engine/Plugins/Runtime/GoogleVR/GoogleVRHMD/Source/GoogleVRHMD/Private/GoogleVRHMD.cpp
--- a/Engine/Plugins/Runtime/GoogleVR/GoogleVRHMD/Source/GoogleVRHMD/Private/GoogleVRHMD.cpp
+++ b/Engine/Plugins/Runtime/GoogleVR/GoogleVRHMD/Source/GoogleVRHMD/Private/GoogleVRHMD.cpp
@@ -1826,7 +1826,7 @@ bool FGoogleVRHMD::Exec( UWorld* InWorld, const TCHAR* Cmd, FOutputDevice& Ar )
 	// Tune the distortion mesh vert count when use Unreal's PostProcessing Distortion
 	else if (FParse::Command(&Cmd, TEXT("DISTORTMESH")))
 	{
-		const static UEnum* MeshSizeEnum = FindObject<UEnum>(ANY_PACKAGE, TEXT("EDistortionMeshSizeEnum"));
+		const static UEnum* MeshSizeEnum = StaticEnum<EDistortionMeshSizeEnum>();
 		FString Value = FParse::Token(Cmd, 0);
 
 		if (!Value.IsEmpty() && MeshSizeEnum->GetIndexByName(*FString::Printf(TEXT("DMS_%s"), *Value.ToUpper())) != INDEX_NONE)
@@ -1916,7 +1916,7 @@ void FGoogleVRHMD::NeckModelScaleCommandHandler(const TArray<FString>& Args, UWo
 #if GOOGLEVRHMD_SUPPORTED_PLATFORMS
 void FGoogleVRHMD::DistortMeshSizeCommandHandler(const TArray<FString>& Args, UWorld* World, FOutputDevice& Ar)
 {
-	const static UEnum* MeshSizeEnum = FindObject<UEnum>(ANY_PACKAGE, TEXT("EDistortionMeshSizeEnum"));
+	const static UEnum* MeshSizeEnum = StaticEnum<EDistortionMeshSizeEnum>();
 	int EnumIndex = INDEX_NONE;
 
 	if (Args.Num())
diff --git a/Engine/Plugins/Runtime/Oculus/OculusVR/Source/OculusEditor/Private/OculusToolWidget.cpp b/Engine/Plugins/Runtime/Oculus/OculusVR/Source/OculusEditor/Private/OculusToolWidget.cpp
--- a/Engine/Plugins/Runtime/Oculus/OculusVR/Source/OculusEditor/Private/OculusToolWidget.cpp
+++ b/Engine/Plugins/Runtime/Oculus/OculusVR/Source/OculusEditor/Private/OculusToolWidget.cpp
@@ -284,7 +284,7 @@ void SOculusToolWidget::RebuildLayout()
 void SOculusToolWidget::Construct(const FArguments& InArgs)
 {
 	pendingRestart = false;
-	PlatformEnum = FindObject<UEnum>(ANY_PACKAGE, TEXT("EOculusPlatform"), true);
+	PlatformEnum = StaticEnum<EOculusPlatform>();
 	Platforms.Reset(2);
 
 	UOculusEditorSettings* EditorSettings = GetMutableDefault<UOculusEditorSettings>();
diff --git a/Engine/Source/Developer/FunctionalTesting/Private/FunctionalTest.cpp b/Engine/Source/Developer/FunctionalTesting/Private/FunctionalTest.cpp
--- a/Engine/Source/Developer/FunctionalTesting/Private/FunctionalTest.cpp
+++ b/Engine/Source/Developer/FunctionalTesting/Private/FunctionalTest.cpp
@@ -57,7 +57,7 @@ namespace
 
 	FString GetComparisonAsString(EComparisonMethod comparison)
 	{
-		UEnum* Enum = FindObject<UEnum>(ANY_PACKAGE, TEXT("EComparisonMethod"), true);
+		UEnum* Enum = StaticEnum<EComparisonMethod>();
 		return Enum->GetNameStringByValue((uint8)comparison).ToLower().Replace(TEXT("_"), TEXT(" "));
 	}
 
@@ -302,7 +302,7 @@ bool AFunctionalTest::IsReady_Implementation()
 
 void AFunctionalTest::FinishTest(EFunctionalTestResult TestResult, const FString& Message)
 {
-	const static UEnum* FTestResultTypeEnum = FindObject<UEnum>( nullptr, TEXT("FunctionalTesting.EFunctionalTestResult") );
+	const static UEnum* FTestResultTypeEnum = StaticEnum<EFunctionalTestResult>();
 	
 	if (bIsRunning == false)
 	{
diff --git a/Engine/Source/Developer/MaterialBaking/Private/MaterialBakingModule.cpp b/Engine/Source/Developer/MaterialBaking/Private/MaterialBakingModule.cpp
--- a/Engine/Source/Developer/MaterialBaking/Private/MaterialBakingModule.cpp
+++ b/Engine/Source/Developer/MaterialBaking/Private/MaterialBakingModule.cpp
@@ -209,7 +209,7 @@ void FMaterialBakingModule::BakeMaterials(const TArray<FMaterialData*>& Material
 					// If saving intermediates is turned on
 					if (bSaveIntermediateTextures)
 					{
-						const UEnum* PropertyEnum = FindObject<UEnum>(ANY_PACKAGE, TEXT("EMaterialProperty"));
+						const UEnum* PropertyEnum = StaticEnum<EMaterialProperty>();
 						FName PropertyName = PropertyEnum->GetNameByValue(Property);
 						FString TrimmedPropertyName = PropertyName.ToString();
 						TrimmedPropertyName.RemoveFromStart(TEXT("MP_"));
diff --git a/Engine/Source/Developer/MaterialBaking/Private/MaterialOptionsCustomization.cpp b/Engine/Source/Developer/MaterialBaking/Private/MaterialOptionsCustomization.cpp
--- a/Engine/Source/Developer/MaterialBaking/Private/MaterialOptionsCustomization.cpp
+++ b/Engine/Source/Developer/MaterialBaking/Private/MaterialOptionsCustomization.cpp
@@ -73,7 +73,7 @@ void FPropertyEntryCustomization::CustomizeChildren(TSharedRef<IPropertyHandle>
 void FPropertyEntryCustomization::UpdateRestrictions(const int32 EntryIndex)
 {	
 	PropertyRestriction->RemoveAll();
-	const UEnum* PropertyEnum = FindObject<UEnum>(ANY_PACKAGE, TEXT("EMaterialProperty"));
+	const UEnum* PropertyEnum = StaticEnum<EMaterialProperty>();
 	if (CurrentOptions)
 	{	
 		// Add all previously set material properties to be disabled
diff --git a/Engine/Source/Developer/MaterialUtilities/Private/MaterialUtilities.cpp b/Engine/Source/Developer/MaterialUtilities/Private/MaterialUtilities.cpp
--- a/Engine/Source/Developer/MaterialUtilities/Private/MaterialUtilities.cpp
+++ b/Engine/Source/Developer/MaterialUtilities/Private/MaterialUtilities.cpp
@@ -106,7 +106,7 @@ UMaterialInterface* FMaterialUtilities::CreateProxyMaterialAndTextures(UPackage*
 		const TArray<FColor>& ColorData = BakeOutput.PropertyData.FindChecked(Property);
 
 		// Look up the property name string 
-		const UEnum* PropertyEnum = FindObject<UEnum>(ANY_PACKAGE, TEXT("EMaterialProperty"));
+		const UEnum* PropertyEnum = StaticEnum<EMaterialProperty>();
 		FName PropertyName = PropertyEnum->GetNameByValue(Property);
 		FString TrimmedPropertyName = PropertyName.ToString();
 		TrimmedPropertyName.RemoveFromStart(TEXT("MP_"));
diff --git a/Engine/Source/Editor/BlueprintGraph/Private/K2Node_EaseFunction.cpp b/Engine/Source/Editor/BlueprintGraph/Private/K2Node_EaseFunction.cpp
--- a/Engine/Source/Editor/BlueprintGraph/Private/K2Node_EaseFunction.cpp
+++ b/Engine/Source/Editor/BlueprintGraph/Private/K2Node_EaseFunction.cpp
@@ -93,7 +93,7 @@ void UK2Node_EaseFunction::AllocateDefaultPins()
 	const UEdGraphSchema_K2* K2Schema = GetDefault<UEdGraphSchema_K2>();
 
 	// Add the first pin representing all available easing functions. If EEasingFunc changes its name this will fail a runtime check!
-	UEnum* EaseFuncEnum = FindObject<UEnum>(ANY_PACKAGE, TEXT("EEasingFunc"), true);
+	UEnum* EaseFuncEnum = StaticEnum<EEasingFunc::Type>();
 	check(EaseFuncEnum);
 	CachedEaseFuncPin = CreatePin(EGPD_Input, UEdGraphSchema_K2::PC_Byte, EaseFuncEnum, FEaseFunctionNodeHelper::GetEaseFuncPinName());
 	SetPinToolTip(*CachedEaseFuncPin, LOCTEXT("EaseFunsPinDescription", "Specifies the desired ease function to be applied. If connected no customization is possible."));
@@ -456,7 +456,7 @@ UEdGraphPin* UK2Node_EaseFunction::GetEaseFuncPin() const
 void UK2Node_EaseFunction::RefreshPinVisibility()
 {
 	const auto EaseFuncPin = GetEaseFuncPin();
-	UEnum * Enum = FindObject<UEnum>(ANY_PACKAGE, TEXT("EEasingFunc"), true);
+	UEnum * Enum = StaticEnum<EEasingFunc::Type>();
 	check(Enum != NULL);
 	const int32 NewEasingFunc = CanCustomizeCurve() ? Enum->GetValueByName(*EaseFuncPin->DefaultValue) : -1;
 	
diff --git a/Engine/Source/Editor/ClothPainter/Private/SClothAssetSelector.cpp b/Engine/Source/Editor/ClothPainter/Private/SClothAssetSelector.cpp
--- a/Engine/Source/Editor/ClothPainter/Private/SClothAssetSelector.cpp
+++ b/Engine/Source/Editor/ClothPainter/Private/SClothAssetSelector.cpp
@@ -394,7 +394,7 @@ public:
 		if(InColumnName == Column_CurrentTarget)
 		{
 			FClothParameterMask_PhysMesh* Mask = Item->GetMask();
-			UEnum* Enum = FindObject<UEnum>(ANY_PACKAGE, TEXT(PREPROCESSOR_TO_STRING(MaskTarget_PhysMesh)), true);
+			UEnum* Enum = StaticEnum<MaskTarget_PhysMesh>();
 			if(Enum && Mask)
 			{
 				return SNew(STextBlock).Text(Enum->GetDisplayNameTextByIndex((int32)Mask->CurrentTarget));
@@ -557,7 +557,7 @@ private:
 	{
 		Builder.BeginSection(NAME_None, LOCTEXT("MaskTargets_SectionName", "Targets"));
 		{
-			UEnum* Enum = FindObject<UEnum>(ANY_PACKAGE, TEXT(PREPROCESSOR_TO_STRING(MaskTarget_PhysMesh)), true);
+			UEnum* Enum = StaticEnum<MaskTarget_PhysMesh>();
 			if(Enum)
 			{
 				const int32 NumEntries = Enum->NumEnums();
diff --git a/Engine/Source/Editor/ContentBrowser/Private/AssetContextMenu.cpp b/Engine/Source/Editor/ContentBrowser/Private/AssetContextMenu.cpp
--- a/Engine/Source/Editor/ContentBrowser/Private/AssetContextMenu.cpp
+++ b/Engine/Source/Editor/ContentBrowser/Private/AssetContextMenu.cpp
@@ -1051,7 +1051,7 @@ bool FAssetContextMenu::AddDocumentationMenuOptions(FMenuBuilder& MenuBuilder)
 								);
 						}
 
-						UEnum* BlueprintTypeEnum = FindObject<UEnum>(ANY_PACKAGE, TEXT("EBlueprintType"), true);
+						UEnum* BlueprintTypeEnum = StaticEnum<EBlueprintType>();
 						const FString EnumString = SelectedAssets[0].GetTagValueRef<FString>(GET_MEMBER_NAME_CHECKED(UBlueprint,BlueprintType));
 						EBlueprintType BlueprintType = (!EnumString.IsEmpty() ? (EBlueprintType)BlueprintTypeEnum->GetValueByName(*EnumString) : BPTYPE_Normal);
 
diff --git a/Engine/Source/Editor/DetailCustomizations/Private/AttenuationSettingsCustomizations.cpp b/Engine/Source/Editor/DetailCustomizations/Private/AttenuationSettingsCustomizations.cpp
--- a/Engine/Source/Editor/DetailCustomizations/Private/AttenuationSettingsCustomizations.cpp
+++ b/Engine/Source/Editor/DetailCustomizations/Private/AttenuationSettingsCustomizations.cpp
@@ -494,7 +494,7 @@ void FForceFeedbackAttenuationSettingsCustomization::CustomizeChildren(TSharedRe
 	}
 
 	TSharedPtr<FPropertyRestriction> EnumRestriction = MakeShareable(new FPropertyRestriction(NSLOCTEXT("AttenuationSettings", "NoNaturalSound", "Natural Sound is only available for Sound Attenuation")));
-	const UEnum* const AttenuationDistanceModelEnum = FindObject<UEnum>(ANY_PACKAGE, TEXT("EAttenuationDistanceModel"));		
+	const UEnum* const AttenuationDistanceModelEnum = StaticEnum<EAttenuationDistanceModel>();		
 	EnumRestriction->AddHiddenValue(AttenuationDistanceModelEnum->GetNameStringByValue((uint8)EAttenuationDistanceModel::NaturalSound));
 	DistanceAlgorithmHandle->AddRestriction(EnumRestriction.ToSharedRef());
 
diff --git a/Engine/Source/Editor/DetailCustomizations/Private/BodyInstanceCustomization.cpp b/Engine/Source/Editor/DetailCustomizations/Private/BodyInstanceCustomization.cpp
--- a/Engine/Source/Editor/DetailCustomizations/Private/BodyInstanceCustomization.cpp
+++ b/Engine/Source/Editor/DetailCustomizations/Private/BodyInstanceCustomization.cpp
@@ -247,7 +247,7 @@ int32 FBodyInstanceCustomization::InitializeObjectTypeComboList()
 	ObjectTypeComboList.Empty();
 	ObjectTypeValues.Empty();
 
-	UEnum * Enum = FindObject<UEnum>(ANY_PACKAGE, TEXT("ECollisionChannel"), true);
+	UEnum * Enum = StaticEnum<ECollisionChannel>();
 	const FString KeyName = TEXT("DisplayName");
 	const FString QueryType = TEXT("TraceQuery");
 
@@ -335,7 +335,7 @@ TSharedPtr<FString> FBodyInstanceCustomization::GetProfileString(FName ProfileNa
 void FBodyInstanceCustomization::UpdateValidCollisionChannels()
 {
 	// find the enum
-	UEnum * Enum = FindObject<UEnum>(ANY_PACKAGE, TEXT("ECollisionChannel"), true);
+	UEnum * Enum = StaticEnum<ECollisionChannel>();
 	// we need this Enum
 	check (Enum);
 	const FString KeyName = TEXT("DisplayName");
diff --git a/Engine/Source/Editor/DetailCustomizations/Private/CollisionProfileDetails.cpp b/Engine/Source/Editor/DetailCustomizations/Private/CollisionProfileDetails.cpp
--- a/Engine/Source/Editor/DetailCustomizations/Private/CollisionProfileDetails.cpp
+++ b/Engine/Source/Editor/DetailCustomizations/Private/CollisionProfileDetails.cpp
@@ -834,7 +834,7 @@ void SProfileEditDialog::AddCollisionResponse()
 
 	// find the enum
 	TArray<FCollisionChannelInfo>	ValidCollisionChannels;
-	UEnum * Enum = FindObject<UEnum>(ANY_PACKAGE, TEXT("ECollisionChannel"), true);
+	UEnum * Enum = StaticEnum<ECollisionChannel>();
 	// we need this Enum
 	check(Enum);
 	const FString KeyName = TEXT("DisplayName");
diff --git a/Engine/Source/Editor/DetailCustomizations/Private/MaterialProxySettingsCustomizations.cpp b/Engine/Source/Editor/DetailCustomizations/Private/MaterialProxySettingsCustomizations.cpp
--- a/Engine/Source/Editor/DetailCustomizations/Private/MaterialProxySettingsCustomizations.cpp
+++ b/Engine/Source/Editor/DetailCustomizations/Private/MaterialProxySettingsCustomizations.cpp
@@ -108,7 +108,7 @@ void FMaterialProxySettingsCustomizations::CustomizeChildren(TSharedRef<IPropert
 			if (bUseNativeTool)
 			{
 				TSharedPtr<FPropertyRestriction> EnumRestriction = MakeShareable(new FPropertyRestriction(LOCTEXT("NoSupport", "Unable to support this option in Merge Actor")));
-				const UEnum* const TextureSizingTypeEnum = FindObject<UEnum>(ANY_PACKAGE, TEXT("ETextureSizingType"));
+				const UEnum* const TextureSizingTypeEnum = StaticEnum<ETextureSizingType>();
 				EnumRestriction->AddHiddenValue(TextureSizingTypeEnum->GetNameStringByValue((uint8)ETextureSizingType::TextureSizingType_UseSimplygonAutomaticSizing));
 				EnumHandle->AddRestriction(EnumRestriction.ToSharedRef());
 			}
diff --git a/Engine/Source/Editor/DetailCustomizations/Private/MeshMergingSettingsCustomization.cpp b/Engine/Source/Editor/DetailCustomizations/Private/MeshMergingSettingsCustomization.cpp
--- a/Engine/Source/Editor/DetailCustomizations/Private/MeshMergingSettingsCustomization.cpp
+++ b/Engine/Source/Editor/DetailCustomizations/Private/MeshMergingSettingsCustomization.cpp
@@ -46,7 +46,7 @@ void FMeshMergingSettingsObjectCustomization::CustomizeDetails(IDetailLayoutBuil
 			EnumProperty = Property;
 						
 			TSharedPtr<FPropertyRestriction> EnumRestriction = MakeShareable(new FPropertyRestriction(LOCTEXT("NoSupport","Unable to support this option in Merge Actor")));
-			const UEnum* const MeshLODSelectionTypeEnum = FindObject<UEnum>(ANY_PACKAGE, TEXT("EMeshLODSelectionType"));		
+			const UEnum* const MeshLODSelectionTypeEnum = StaticEnum<EMeshLODSelectionType>();		
 			EnumRestriction->AddDisabledValue(MeshLODSelectionTypeEnum->GetNameStringByValue((uint8)EMeshLODSelectionType::CalculateLOD));
 			EnumProperty->AddRestriction(EnumRestriction.ToSharedRef());
 		}
diff --git a/Engine/Source/Editor/DetailCustomizations/Private/ParticleModuleDetails.cpp b/Engine/Source/Editor/DetailCustomizations/Private/ParticleModuleDetails.cpp
--- a/Engine/Source/Editor/DetailCustomizations/Private/ParticleModuleDetails.cpp
+++ b/Engine/Source/Editor/DetailCustomizations/Private/ParticleModuleDetails.cpp
@@ -82,7 +82,7 @@ void FParticleModuleRequiredDetails::CustomizeDetails( IDetailLayoutBuilder& Det
 	if( !RandomRandomBlendRestriction.IsValid() )
 	{
 		RandomRandomBlendRestriction = MakeShareable( new FPropertyRestriction(GetNotAllowedOnGPUEmitterText()) );
-		const UEnum* const ParticleSubUVInterpMethodEnum = FindObject<UEnum>(ANY_PACKAGE, TEXT("EParticleSubUVInterpMethod"));		
+		const UEnum* const ParticleSubUVInterpMethodEnum = StaticEnum<EParticleSubUVInterpMethod>();		
 		RandomRandomBlendRestriction->AddDisabledValue(ParticleSubUVInterpMethodEnum->GetNameStringByValue((uint8)EParticleSubUVInterpMethod::PSUVIM_Random));
 		RandomRandomBlendRestriction->AddDisabledValue(ParticleSubUVInterpMethodEnum->GetNameStringByValue((uint8)EParticleSubUVInterpMethod::PSUVIM_Random_Blend));
 	}
diff --git a/Engine/Source/Editor/DetailCustomizations/Private/PhysicsSettingsDetails.cpp b/Engine/Source/Editor/DetailCustomizations/Private/PhysicsSettingsDetails.cpp
--- a/Engine/Source/Editor/DetailCustomizations/Private/PhysicsSettingsDetails.cpp
+++ b/Engine/Source/Editor/DetailCustomizations/Private/PhysicsSettingsDetails.cpp
@@ -283,7 +283,7 @@ void FPhysicsSettingsDetails::CustomizeDetails( IDetailLayoutBuilder& DetailBuil
 	PhysicsSettings= UPhysicsSettings::Get();
 	check(PhysicsSettings);
 
-	PhysicalSurfaceEnum = FindObject<UEnum>(ANY_PACKAGE, TEXT("EPhysicalSurface"), true);
+	PhysicalSurfaceEnum = StaticEnum<EPhysicalSurface>();
 	check(PhysicalSurfaceEnum);
 
 	TSharedPtr<IPropertyHandle> PhysicalSurfacesProperty = DetailBuilder.GetProperty(GET_MEMBER_NAME_CHECKED(UPhysicsSettings, PhysicalSurfaces));
diff --git a/Engine/Source/Editor/DetailCustomizations/Private/SceneComponentDetails.cpp b/Engine/Source/Editor/DetailCustomizations/Private/SceneComponentDetails.cpp
--- a/Engine/Source/Editor/DetailCustomizations/Private/SceneComponentDetails.cpp
+++ b/Engine/Source/Editor/DetailCustomizations/Private/SceneComponentDetails.cpp
@@ -182,7 +182,7 @@ void FSceneComponentDetails::CustomizeDetails( IDetailLayoutBuilder& DetailBuild
 			if (IsMobilitySettingProhibited(EComponentMobility::Static, SceneComponent, RestrictReason))
 			{
 				 TSharedPtr<FPropertyRestriction> StaticRestriction = MakeShareable(new FPropertyRestriction(MoveTemp(RestrictReason)));
-				 const UEnum* const ComponentMobilityEnum = FindObject<UEnum>(ANY_PACKAGE, TEXT("EComponentMobility"));		
+				 const UEnum* const ComponentMobilityEnum = StaticEnum<EComponentMobility::Type>();		
 				 StaticRestriction->AddDisabledValue(ComponentMobilityEnum->GetNameStringByValue((uint8)EComponentMobility::Static));
 				 MobilityHandle->AddRestriction(StaticRestriction.ToSharedRef());
 
@@ -197,7 +197,7 @@ void FSceneComponentDetails::CustomizeDetails( IDetailLayoutBuilder& DetailBuild
 			if (IsMobilitySettingProhibited(EComponentMobility::Stationary, SceneComponent, RestrictReason))
 			{
 				TSharedPtr<FPropertyRestriction> StationaryRestriction = MakeShareable(new FPropertyRestriction(MoveTemp(RestrictReason)));
-				const UEnum* const ComponentMobilityEnum = FindObject<UEnum>(ANY_PACKAGE, TEXT("EComponentMobility"));		
+				const UEnum* const ComponentMobilityEnum = StaticEnum<EComponentMobility::Type>();
 				StationaryRestriction->AddDisabledValue(ComponentMobilityEnum->GetNameStringByValue((uint8)EComponentMobility::Stationary));
 				MobilityHandle->AddRestriction(StationaryRestriction.ToSharedRef());
 
diff --git a/Engine/Source/Editor/DetailCustomizations/Private/SplineComponentDetails.cpp b/Engine/Source/Editor/DetailCustomizations/Private/SplineComponentDetails.cpp
--- a/Engine/Source/Editor/DetailCustomizations/Private/SplineComponentDetails.cpp
+++ b/Engine/Source/Editor/DetailCustomizations/Private/SplineComponentDetails.cpp
@@ -196,7 +196,7 @@ FSplinePointDetails::FSplinePointDetails()
 
 	SplineCurvesProperty = FindField<UProperty>(USplineComponent::StaticClass(), GET_MEMBER_NAME_CHECKED(USplineComponent, SplineCurves));
 
-	UEnum* SplinePointTypeEnum = FindObject<UEnum>(ANY_PACKAGE, TEXT("ESplinePointType"));
+	UEnum* SplinePointTypeEnum = StaticEnum<ESplinePointType::Type>();
 	check(SplinePointTypeEnum);
 	for (int32 EnumIndex = 0; EnumIndex < SplinePointTypeEnum->NumEnums() - 1; ++EnumIndex)
 	{
diff --git a/Engine/Source/Editor/DetailCustomizations/Private/TextureLODSettingsDetails.cpp b/Engine/Source/Editor/DetailCustomizations/Private/TextureLODSettingsDetails.cpp
--- a/Engine/Source/Editor/DetailCustomizations/Private/TextureLODSettingsDetails.cpp
+++ b/Engine/Source/Editor/DetailCustomizations/Private/TextureLODSettingsDetails.cpp
@@ -81,7 +81,7 @@ void FTextureLODGroupLayout::AddToAvailableMipGenSettings(TextureMipGenSettings
 
 void FTextureLODGroupLayout::GenerateHeaderRowContent(FDetailWidgetRow& NodeRow)
 {
-	UEnum* TextureGroupEnum = FindObject<UEnum>(NULL, TEXT("/Script/Engine.TextureGroup"));
+	UEnum* TextureGroupEnum = StaticEnum<TextureGroup>();
 	const FString& LODGroupName = TextureGroupEnum->GetMetaData(TEXT("DisplayName"), LodGroup->Group);
 
 	NodeRow.NameContent()
@@ -366,7 +366,7 @@ FText FTextureLODGroupLayout::GetMipFilterComboBoxContent() const
 
 TSharedRef<SWidget> FTextureLODGroupLayout::MakeMipGenSettingsComboWidget(TSharedPtr<TextureMipGenSettings> InItem)
 {
-	UEnum* TextureGroupEnum = FindObject<UEnum>(NULL, TEXT("/Script/Engine.TextureMipGenSettings"));
+	UEnum* TextureGroupEnum = StaticEnum<TextureMipGenSettings>();
 	const FString& MipGenSettingsName = TextureGroupEnum->GetMetaData(TEXT("DisplayName"), *InItem.Get());
 
 	return SNew(STextBlock).Text(FText::FromString(MipGenSettingsName)).Font(IDetailLayoutBuilder::GetDetailFont());
@@ -389,7 +389,7 @@ FText FTextureLODGroupLayout::GetMipGenSettingsComboBoxToolTip() const
 
 FText FTextureLODGroupLayout::GetMipGenSettingsComboBoxContent() const
 {
-	UEnum* TextureGroupEnum = FindObject<UEnum>(NULL, TEXT("/Script/Engine.TextureMipGenSettings"));
+	UEnum* TextureGroupEnum = StaticEnum<TextureMipGenSettings>();
 	const FString& MipGenSettingsName = TextureGroupEnum->GetMetaData(TEXT("DisplayName"), LodGroup->MipGenSettings);
 	
 	return FText::FromString(MipGenSettingsName);
diff --git a/Engine/Source/Editor/EnvironmentQueryEditor/Private/DetailCustomizations/EnvQueryTestDetails.cpp b/Engine/Source/Editor/EnvironmentQueryEditor/Private/DetailCustomizations/EnvQueryTestDetails.cpp
--- a/Engine/Source/Editor/EnvironmentQueryEditor/Private/DetailCustomizations/EnvQueryTestDetails.cpp
+++ b/Engine/Source/Editor/EnvironmentQueryEditor/Private/DetailCustomizations/EnvQueryTestDetails.cpp
@@ -271,7 +271,7 @@ END_SLATE_FUNCTION_BUILD_OPTIMIZATION
 
 void FEnvQueryTestDetails::BuildFilterTestValues()
 {
-	UEnum* TestConditionEnum = FindObject<UEnum>(ANY_PACKAGE, TEXT("EEnvTestFilterType"));
+	UEnum* TestConditionEnum = StaticEnum<EEnvTestFilterType::Type>();
 	check(TestConditionEnum);
 
 	FilterTestValues.Reset();
@@ -294,7 +294,7 @@ void FEnvQueryTestDetails::BuildFilterTestValues()
 
 void FEnvQueryTestDetails::BuildScoreEquationValues()
 {
-	UEnum* TestScoreEquationEnum = FindObject<UEnum>(ANY_PACKAGE, TEXT("EEnvTestScoreEquation"));
+	UEnum* TestScoreEquationEnum = StaticEnum<EEnvTestScoreEquation::Type>();
 	check(TestScoreEquationEnum);
 
 	ScoreEquationValues.Reset();
@@ -329,7 +329,7 @@ void FEnvQueryTestDetails::OnScoreEquationChange(int32 Index)
 
 void FEnvQueryTestDetails::BuildScoreClampingTypeValues(bool bBuildMinValues, TArray<FTextIntPair>& ClampTypeValues) const
 {
-	UEnum* ScoringNormalizationEnum = FindObject<UEnum>(ANY_PACKAGE, TEXT("EEnvQueryTestClamping"));
+	UEnum* ScoringNormalizationEnum = StaticEnum<EEnvQueryTestClamping::Type>();
 	check(ScoringNormalizationEnum);
 
 	ClampTypeValues.Reset();
diff --git a/Engine/Source/Editor/EnvironmentQueryEditor/Private/DetailCustomizations/EnvTraceDataCustomization.cpp b/Engine/Source/Editor/EnvironmentQueryEditor/Private/DetailCustomizations/EnvTraceDataCustomization.cpp
--- a/Engine/Source/Editor/EnvironmentQueryEditor/Private/DetailCustomizations/EnvTraceDataCustomization.cpp
+++ b/Engine/Source/Editor/EnvironmentQueryEditor/Private/DetailCustomizations/EnvTraceDataCustomization.cpp
@@ -134,7 +134,7 @@ void FEnvTraceDataCustomization::CacheTraceModes(TSharedRef<class IPropertyHandl
 	PropCanDisable->GetValue(bCanDisable);
 	PropCanProject->GetValue(bCanShowProjection);
 
-	static UEnum* TraceModeEnum = FindObject<UEnum>(ANY_PACKAGE, TEXT("EEnvQueryTrace"));
+	static UEnum* TraceModeEnum = StaticEnum<EEnvQueryTrace::Type>();
 	check(TraceModeEnum);
 
 	TraceModes.Reset();
diff --git a/Engine/Source/Editor/GameProjectGeneration/Private/GameProjectUtils.cpp b/Engine/Source/Editor/GameProjectGeneration/Private/GameProjectUtils.cpp
--- a/Engine/Source/Editor/GameProjectGeneration/Private/GameProjectUtils.cpp
+++ b/Engine/Source/Editor/GameProjectGeneration/Private/GameProjectUtils.cpp
@@ -721,9 +721,9 @@ bool GameProjectUtils::CreateProject(const FProjectInformation& InProjectInfo, F
 		EventAttributes.Add(FAnalyticsEventAttribute(TEXT("ProjectType"), InProjectInfo.bShouldGenerateCode ? TEXT("C++ Code") : TEXT("Content Only")));
 		EventAttributes.Add(FAnalyticsEventAttribute(TEXT("Outcome"), bProjectCreationSuccessful ? TEXT("Successful") : TEXT("Failed")));
 
-		UEnum* Enum = FindObject<UEnum>(ANY_PACKAGE, TEXT("EHardwareClass"), true);
+		UEnum* Enum = StaticEnum<EHardwareClass::Type>();
 		EventAttributes.Add(FAnalyticsEventAttribute(TEXT("HardwareClass"), Enum ? Enum->GetNameStringByValue(InProjectInfo.TargetedHardware) : FString()));
-		Enum = FindObject<UEnum>(ANY_PACKAGE, TEXT("EGraphicsPreset"), true);
+		Enum = StaticEnum<EGraphicsPreset::Type>();
 		EventAttributes.Add(FAnalyticsEventAttribute(TEXT("GraphicsPreset"), Enum ? Enum->GetNameStringByValue(InProjectInfo.DefaultGraphicsPerformance) : FString()));
 		EventAttributes.Add(FAnalyticsEventAttribute(TEXT("StarterContent"), InProjectInfo.bCopyStarterContent ? TEXT("Yes") : TEXT("No")));
 		EventAttributes.Emplace(TEXT("Enterprise"), InProjectInfo.bIsEnterpriseProject);
diff --git a/Engine/Source/Editor/GameProjectGeneration/Private/Tests/GameProjectAutomationTests.cpp b/Engine/Source/Editor/GameProjectGeneration/Private/Tests/GameProjectAutomationTests.cpp
--- a/Engine/Source/Editor/GameProjectGeneration/Private/Tests/GameProjectAutomationTests.cpp
+++ b/Engine/Source/Editor/GameProjectGeneration/Private/Tests/GameProjectAutomationTests.cpp
@@ -87,7 +87,7 @@ namespace GameProjectAutomationUtils
 		bool bRemoveCreatedProjects = true;		
 		OutCreatedProjects = 0;
 		OutMatchedProjects = 0;
-		UEnum* SourceCategoryEnum = FindObject<UEnum>(ANY_PACKAGE, TEXT("EContentSourceCategory"));
+		UEnum* SourceCategoryEnum = StaticEnum<EContentSourceCategory>();
 
 		// The category name in the FTemplateItem is not the same as the enum definition EContentSourceCategory - convert it
 		FName CategoryName;
diff --git a/Engine/Source/Editor/IntroTutorials/Private/IntroTutorials.cpp b/Engine/Source/Editor/IntroTutorials/Private/IntroTutorials.cpp
--- a/Engine/Source/Editor/IntroTutorials/Private/IntroTutorials.cpp
+++ b/Engine/Source/Editor/IntroTutorials/Private/IntroTutorials.cpp
@@ -243,7 +243,7 @@ void FIntroTutorials::AddSummonBlueprintTutorialsMenuExtension(FMenuBuilder& Men
 		UBlueprint* BP = Cast<UBlueprint>(PrimaryObject);
 		if(BP != nullptr)
 		{
-			UEnum* Enum = FindObject<UEnum>(ANY_PACKAGE, TEXT("EBlueprintType"), true);
+			UEnum* Enum = StaticEnum<EBlueprintType>();
 			check(Enum);
 			MenuBuilder.AddMenuEntry(
 				FText::Format(LOCTEXT("BlueprintTutorialsMenuEntryTitle", "{0} Tutorial"), Enum->GetDisplayNameTextByValue(BP->BlueprintType)),
diff --git a/Engine/Source/Editor/Kismet/Private/BlueprintDetailsCustomization.cpp b/Engine/Source/Editor/Kismet/Private/BlueprintDetailsCustomization.cpp
--- a/Engine/Source/Editor/Kismet/Private/BlueprintDetailsCustomization.cpp
+++ b/Engine/Source/Editor/Kismet/Private/BlueprintDetailsCustomization.cpp
@@ -706,7 +706,7 @@ void FBlueprintVarActionDetails::CustomizeDetails( IDetailLayoutBuilder& DetailL
 	];
 
 	ReplicationConditionEnumTypeNames.Empty();
-	UEnum* Enum = FindObject<UEnum>(ANY_PACKAGE, TEXT("ELifetimeCondition"), true);
+	UEnum* Enum = StaticEnum<ELifetimeCondition>();
 	check(Enum);
 	
 	for (int32 i = 0; i < Enum->NumEnums(); i++)
diff --git a/Engine/Source/Editor/Kismet/Private/SBlueprintContextTargetMenu.cpp b/Engine/Source/Editor/Kismet/Private/SBlueprintContextTargetMenu.cpp
--- a/Engine/Source/Editor/Kismet/Private/SBlueprintContextTargetMenu.cpp
+++ b/Engine/Source/Editor/Kismet/Private/SBlueprintContextTargetMenu.cpp
@@ -1,4 +1,4 @@
-﻿// Copyright 1998-2018 Epic Games, Inc. All Rights Reserved.
+// Copyright 1998-2018 Epic Games, Inc. All Rights Reserved.
 
 #include "SBlueprintContextTargetMenu.h"
 #include "Engine/Blueprint.h"
@@ -538,7 +538,7 @@ void SBlueprintContextTargetMenu::Construct(const FArguments& InArgs, const FBlu
 		];
 	}
 
-	UEnum* const TargetEnum = FindObject<UEnum>(ANY_PACKAGE, TEXT("EContextTargetFlags"));
+	UEnum* const TargetEnum = StaticEnum<EContextTargetFlags::Type>();
 	uint32 const GreatestFlag = (EContextTargetFlags::ContextTargetFlagsEnd & ~1);
 
 	int32  ColIndex = 0;
diff --git a/Engine/Source/Editor/MainFrame/Private/Frame/MainFrameActions.cpp b/Engine/Source/Editor/MainFrame/Private/Frame/MainFrameActions.cpp
--- a/Engine/Source/Editor/MainFrame/Private/Frame/MainFrameActions.cpp
+++ b/Engine/Source/Editor/MainFrame/Private/Frame/MainFrameActions.cpp
@@ -723,7 +723,7 @@ void FMainFrameActionCallbacks::PackageProject( const FName InPlatformInfoName )
 		OptionalParams += FString::Printf(TEXT(" -NumCookersToSpawn=%d"), NumCookers); 
 	}
 
-	FString Configuration = FindObject<UEnum>(ANY_PACKAGE, TEXT("EProjectPackagingBuildConfigurations"))->GetNameStringByValue(PackagingSettings->BuildConfiguration);
+	FString Configuration = StaticEnum<EProjectPackagingBuildConfigurations>()->GetNameStringByValue(PackagingSettings->BuildConfiguration);
 	Configuration = Configuration.Replace(TEXT("PPBC_"), TEXT(""));
 	if (Configuration.Right(6) == TEXT("Client"))
 	{
diff --git a/Engine/Source/Editor/MaterialEditor/Private/MaterialInstanceEditor.cpp b/Engine/Source/Editor/MaterialEditor/Private/MaterialInstanceEditor.cpp
--- a/Engine/Source/Editor/MaterialEditor/Private/MaterialInstanceEditor.cpp
+++ b/Engine/Source/Editor/MaterialEditor/Private/MaterialInstanceEditor.cpp
@@ -1232,7 +1232,7 @@ void FMaterialInstanceEditor::DrawSamplerWarningStrings(FCanvas* Canvas, int32&
 		{
 			UFont* FontToUse = GEngine->GetTinyFont();
 			const int32 SpacingBetweenLines = 13;
-			UEnum* SamplerTypeEnum = FindObject<UEnum>( NULL, TEXT("/Script/Engine.EMaterialSamplerType") );
+			UEnum* SamplerTypeEnum = StaticEnum<EMaterialSamplerType>();
 			check( SamplerTypeEnum );
 
 			const int32 GroupCount = MaterialEditorInstance->ParameterGroups.Num();
diff --git a/Engine/Source/Editor/MaterialEditor/Private/MaterialPropertyHelpers.cpp b/Engine/Source/Editor/MaterialEditor/Private/MaterialPropertyHelpers.cpp
--- a/Engine/Source/Editor/MaterialEditor/Private/MaterialPropertyHelpers.cpp
+++ b/Engine/Source/Editor/MaterialEditor/Private/MaterialPropertyHelpers.cpp
@@ -1024,7 +1024,7 @@ FEditorParameterGroup&  FMaterialPropertyHelpers::GetParameterGroup(UMaterial* I
 
 void FMaterialPropertyHelpers::GetVectorChannelMaskComboBoxStrings(TArray<TSharedPtr<FString>>& OutComboBoxStrings, TArray<TSharedPtr<SToolTip>>& OutToolTips, TArray<bool>& OutRestrictedItems)
 {
-	const UEnum* ChannelEnum = FindObject<UEnum>(ANY_PACKAGE, TEXT("EChannelMaskParameterColor"));
+	const UEnum* ChannelEnum = StaticEnum<EChannelMaskParameterColor::Type>();
 	check(ChannelEnum);
 
 	// Add RGBA string options (Note: Exclude the "::Max" entry)
@@ -1044,7 +1044,7 @@ FString FMaterialPropertyHelpers::GetVectorChannelMaskValue(UDEditorParameterVal
 	UDEditorVectorParameterValue* VectorParam = Cast<UDEditorVectorParameterValue>(InParameter);
 	check(VectorParam && VectorParam->bIsUsedAsChannelMask);
 
-	const UEnum* ChannelEnum = FindObject<UEnum>(ANY_PACKAGE, TEXT("EChannelMaskParameterColor"));
+	const UEnum* ChannelEnum = StaticEnum<EChannelMaskParameterColor::Type>();
 	check(ChannelEnum);
 
 	// Convert from vector to RGBA string
@@ -1075,7 +1075,7 @@ void FMaterialPropertyHelpers::SetVectorChannelMaskValue(const FString& StringVa
 	UDEditorVectorParameterValue* VectorParam = Cast<UDEditorVectorParameterValue>(InParameter);
 	check(VectorParam && VectorParam->bIsUsedAsChannelMask);
 
-	const UEnum* ChannelEnum = FindObject<UEnum>(ANY_PACKAGE, TEXT("EChannelMaskParameterColor"));
+	const UEnum* ChannelEnum = StaticEnum<EChannelMaskParameterColor::Type>();
 	check(ChannelEnum);
 
 	// Convert from RGBA string to vector
diff --git a/Engine/Source/Editor/MeshPaintMode/Private/PaintModeSettingsCustomization.cpp b/Engine/Source/Editor/MeshPaintMode/Private/PaintModeSettingsCustomization.cpp
--- a/Engine/Source/Editor/MeshPaintMode/Private/PaintModeSettingsCustomization.cpp
+++ b/Engine/Source/Editor/MeshPaintMode/Private/PaintModeSettingsCustomization.cpp
@@ -355,7 +355,7 @@ EVisibility FVertexPaintSettingsCustomization::ArePropertiesVisible(const int32
 
 void FVertexPaintSettingsCustomization::OnTextureWeightTypeChanged(TSharedRef<IPropertyHandle> WeightTypeProperty, TSharedRef<IPropertyHandle> PaintWeightProperty, TSharedRef<IPropertyHandle> EraseWeightProperty)
 {
-	UEnum* ImportTypeEnum = FindObject<UEnum>(ANY_PACKAGE, TEXT("ETexturePaintIndex"));
+	UEnum* ImportTypeEnum = StaticEnum<ETexturePaintIndex>();
 	uint8 EnumValue = 0;
 	WeightTypeProperty->GetValue(EnumValue);	
 
diff --git a/Engine/Source/Editor/MovieSceneTools/Private/MovieSceneBuiltInEasingFunctionCustomization.cpp b/Engine/Source/Editor/MovieSceneTools/Private/MovieSceneBuiltInEasingFunctionCustomization.cpp
--- a/Engine/Source/Editor/MovieSceneTools/Private/MovieSceneBuiltInEasingFunctionCustomization.cpp
+++ b/Engine/Source/Editor/MovieSceneTools/Private/MovieSceneBuiltInEasingFunctionCustomization.cpp
@@ -153,7 +153,7 @@ void FMovieSceneBuiltInEasingFunctionCustomization::CustomizeDetails(IDetailLayo
 {
 	TypeProperty = DetailBuilder.GetProperty(GET_MEMBER_NAME_CHECKED(UMovieSceneBuiltInEasingFunction, Type));
 
-	const UEnum* EasingEnum = FindObject<UEnum>(ANY_PACKAGE, TEXT("EMovieSceneBuiltInEasing"), false);
+	const UEnum* EasingEnum = StaticEnum<EMovieSceneBuiltInEasing>();
 	check(EasingEnum)
 
 	TArray<FGroupedEasing> Groups;
diff --git a/Engine/Source/Editor/MovieSceneTools/Private/TrackEditors/SkeletalAnimationTrackEditor.cpp b/Engine/Source/Editor/MovieSceneTools/Private/TrackEditors/SkeletalAnimationTrackEditor.cpp
--- a/Engine/Source/Editor/MovieSceneTools/Private/TrackEditors/SkeletalAnimationTrackEditor.cpp
+++ b/Engine/Source/Editor/MovieSceneTools/Private/TrackEditors/SkeletalAnimationTrackEditor.cpp
@@ -503,7 +503,7 @@ bool FSkeletalAnimationTrackEditor::ShouldFilterAsset(const FAssetData& AssetDat
 		return false;
 	}
 
-	UEnum* AdditiveTypeEnum = FindObject<UEnum>(ANY_PACKAGE, TEXT("EAdditiveAnimationType"), true);
+	UEnum* AdditiveTypeEnum = StaticEnum<EAdditiveAnimationType>();
 	return ((EAdditiveAnimationType)AdditiveTypeEnum->GetValueByName(*EnumString) == AAT_RotationOffsetMeshSpace);
 }
 
diff --git a/Engine/Source/Editor/Persona/Private/Customization/BlendSampleDetails.cpp b/Engine/Source/Editor/Persona/Private/Customization/BlendSampleDetails.cpp
--- a/Engine/Source/Editor/Persona/Private/Customization/BlendSampleDetails.cpp
+++ b/Engine/Source/Editor/Persona/Private/Customization/BlendSampleDetails.cpp
@@ -31,7 +31,7 @@ FBlendSampleDetails::FBlendSampleDetails(const UBlendSpaceBase* InBlendSpace, cl
 	, GridWidget(InGridWidget)
 {
 	// Retrieve the additive animation type enum
-	const UEnum* AdditiveTypeEnum = FindObject<UEnum>(ANY_PACKAGE, TEXT("EAdditiveAnimationType"), true);
+	const UEnum* AdditiveTypeEnum = StaticEnum<EAdditiveAnimationType>();
 	// For each type check whether or not the blend space is compatible with it and cache the result
 	for (int32 TypeValue = 0; TypeValue < (int32)EAdditiveAnimationType::AAT_MAX; ++TypeValue)
 	{
@@ -207,7 +207,7 @@ bool FBlendSampleDetails::ShouldFilterAssetStatic(const FAssetData& AssetData, c
 	/** Cached flags to check whether or not an additive animation type is compatible with the blend space*/
 	TMap<FString, bool> bValidAdditiveTypes;
 	// Retrieve the additive animation type enum
-	const UEnum* AdditiveTypeEnum = FindObject<UEnum>(ANY_PACKAGE, TEXT("EAdditiveAnimationType"), true);
+	const UEnum* AdditiveTypeEnum = StaticEnum<EAdditiveAnimationType>();
 	// For each type check whether or not the blend space is compatible with it and cache the result
 	for (int32 TypeValue = 0; TypeValue < (int32)EAdditiveAnimationType::AAT_MAX; ++TypeValue)
 	{
diff --git a/Engine/Source/Editor/Persona/Private/SAnimSegmentsPanel.cpp b/Engine/Source/Editor/Persona/Private/SAnimSegmentsPanel.cpp
--- a/Engine/Source/Editor/Persona/Private/SAnimSegmentsPanel.cpp
+++ b/Engine/Source/Editor/Persona/Private/SAnimSegmentsPanel.cpp
@@ -451,7 +451,7 @@ void SAnimSegmentsPanel::FillSubMenu(FMenuBuilder& MenuBuilder, int32 AnimSegmen
 
 bool SAnimSegmentsPanel::ShouldFilter(const FAssetData& DataToDisplay, TEnumAsByte<EAdditiveAnimationType> InAdditiveType)
 {
-	UEnum* AdditiveTypeEnum = FindObject<UEnum>(ANY_PACKAGE, TEXT("EAdditiveAnimationType"), true);
+	UEnum* AdditiveTypeEnum = StaticEnum<EAdditiveAnimationType>();
 	const FString EnumString = DataToDisplay.GetTagValueRef<FString>(GET_MEMBER_NAME_CHECKED(UAnimSequence, AdditiveAnimType));
 	EAdditiveAnimationType AdditiveType = (!EnumString.IsEmpty() ? (EAdditiveAnimationType)AdditiveTypeEnum->GetValueByName(*EnumString) : AAT_None);
 	return (AdditiveType != InAdditiveType);
diff --git a/Engine/Source/Editor/Sequencer/Private/SSequencer.cpp b/Engine/Source/Editor/Sequencer/Private/SSequencer.cpp
--- a/Engine/Source/Editor/Sequencer/Private/SSequencer.cpp
+++ b/Engine/Source/Editor/Sequencer/Private/SSequencer.cpp
@@ -1093,7 +1093,7 @@ void SSequencer::FillTimeDisplayFormatMenu(FMenuBuilder& MenuBuilder)
 	TSharedPtr<FSequencer> Sequencer = SequencerPtr.Pin();
 	bool bSupportsDropFormatDisplay = FTimecode::IsDropFormatTimecodeSupported(Sequencer->GetFocusedDisplayRate());
 
-	const UEnum* FrameNumberDisplayEnum = FindObject<UEnum>(ANY_PACKAGE, TEXT("EFrameNumberDisplayFormats"), true);
+	const UEnum* FrameNumberDisplayEnum = StaticEnum<EFrameNumberDisplayFormats>();
 	check(FrameNumberDisplayEnum);
 
 	if (Settings)
diff --git a/Engine/Source/Editor/Sequencer/Private/SSequencerPlayRateCombo.cpp b/Engine/Source/Editor/Sequencer/Private/SSequencerPlayRateCombo.cpp
--- a/Engine/Source/Editor/Sequencer/Private/SSequencerPlayRateCombo.cpp
+++ b/Engine/Source/Editor/Sequencer/Private/SSequencerPlayRateCombo.cpp
@@ -259,7 +259,7 @@ void SSequencerPlayRateCombo::PopulateClockSourceMenu(FMenuBuilder& MenuBuilder)
 	TSharedPtr<FSequencer> Sequencer    = WeakSequencer.Pin();
 	UMovieSceneSequence*   RootSequence = Sequencer.IsValid() ? Sequencer->GetRootMovieSceneSequence() : nullptr;
 
-	const UEnum* ClockSourceEnum = FindObject<UEnum>(ANY_PACKAGE, TEXT("EUpdateClockSource"), true);
+	const UEnum* ClockSourceEnum = StaticEnum<EUpdateClockSource>();
 
 	check(ClockSourceEnum);
 
diff --git a/Engine/Source/Editor/SkeletonEditor/Private/SkeletonTreeBoneItem.cpp b/Engine/Source/Editor/SkeletonEditor/Private/SkeletonTreeBoneItem.cpp
--- a/Engine/Source/Editor/SkeletonEditor/Private/SkeletonTreeBoneItem.cpp
+++ b/Engine/Source/Editor/SkeletonEditor/Private/SkeletonTreeBoneItem.cpp
@@ -161,7 +161,7 @@ TSharedRef< SWidget > FSkeletonTreeBoneItem::CreateBoneTranslationRetargetingMod
 
 	MenuBuilder.BeginSection("BoneTranslationRetargetingMode", LOCTEXT( "BoneTranslationRetargetingModeMenuHeading", "Bone Translation Retargeting Mode" ) );
 	{
-		UEnum* const Enum = FindObject<UEnum>(ANY_PACKAGE, TEXT("EBoneTranslationRetargetingMode"), true);	
+		UEnum* const Enum = StaticEnum<EBoneTranslationRetargetingMode::Type>();	
 		check(Enum);
 
 		FUIAction ActionRetargetingAnimation = FUIAction(FExecuteAction::CreateSP(this, &FSkeletonTreeBoneItem::SetBoneTranslationRetargetingMode, EBoneTranslationRetargetingMode::Animation));
@@ -193,7 +193,7 @@ FText FSkeletonTreeBoneItem::GetTranslationRetargetingModeMenuTitle() const
 	if( BoneIndex != INDEX_NONE )
 	{
 		const EBoneTranslationRetargetingMode::Type RetargetingMode = Skeleton.GetBoneTranslationRetargetingMode(BoneIndex);
-		UEnum* const Enum = FindObject<UEnum>(ANY_PACKAGE, TEXT("EBoneTranslationRetargetingMode"), true);	
+		UEnum* const Enum = StaticEnum<EBoneTranslationRetargetingMode::Type>();	
 		if (Enum)
 		{
 			return Enum->GetDisplayNameTextByValue(RetargetingMode);
diff --git a/Engine/Source/Editor/UMGEditor/Private/Designer/SDesignerView.cpp b/Engine/Source/Editor/UMGEditor/Private/Designer/SDesignerView.cpp
--- a/Engine/Source/Editor/UMGEditor/Private/Designer/SDesignerView.cpp
+++ b/Engine/Source/Editor/UMGEditor/Private/Designer/SDesignerView.cpp
@@ -3457,7 +3457,7 @@ TSharedRef<SWidget> SDesignerView::GetScreenSizingFillMenu()
 
 void SDesignerView::CreateScreenFillEntry(FMenuBuilder& MenuBuilder, EDesignPreviewSizeMode SizeMode)
 {
-	const static UEnum* PreviewSizeEnum = FindObject<UEnum>(ANY_PACKAGE, TEXT("EDesignPreviewSizeMode"), true);
+	const static UEnum* PreviewSizeEnum = StaticEnum<EDesignPreviewSizeMode>();
 
 	// Add desired size option
 	FUIAction DesiredSizeAction(
@@ -3471,7 +3471,7 @@ void SDesignerView::CreateScreenFillEntry(FMenuBuilder& MenuBuilder, EDesignPrev
 
 FText SDesignerView::GetScreenSizingFillText() const
 {
-	const static UEnum* PreviewSizeEnum = FindObject<UEnum>(ANY_PACKAGE, TEXT("EDesignPreviewSizeMode"), true);
+	const static UEnum* PreviewSizeEnum = StaticEnum<EDesignPreviewSizeMode>();
 
 	if ( UUserWidget* DefaultWidget = GetDefaultWidget() )
 	{
diff --git a/Engine/Source/Editor/UnrealEd/Private/Commandlets/GenerateBlueprintAPICommandlet.cpp b/Engine/Source/Editor/UnrealEd/Private/Commandlets/GenerateBlueprintAPICommandlet.cpp
--- a/Engine/Source/Editor/UnrealEd/Private/Commandlets/GenerateBlueprintAPICommandlet.cpp
+++ b/Engine/Source/Editor/UnrealEd/Private/Commandlets/GenerateBlueprintAPICommandlet.cpp
@@ -937,7 +937,7 @@ static void GenerateBlueprintAPIUtils::DumpActionMenuItem(uint32 Indent, FGraphA
 			const FString PinEntryIndentedNewline = "\n" + BuildIndentString(Indent+1);
 			const FString PinDetailsIndentedNewline = "\n" + BuildIndentString(Indent+2);
 
-			static UEnum* PinDirectionEnum = FindObject<UEnum>(ANY_PACKAGE, TEXT("EEdGraphPinDirection"));
+			static UEnum* PinDirectionEnum = StaticEnum<EEdGraphPinDirection>();
 
 			bool bFirst = true;
 			for (UEdGraphPin* Pin : Node->Pins)
diff --git a/Engine/Source/Editor/UnrealEd/Private/Fbx/FbxMainImport.cpp b/Engine/Source/Editor/UnrealEd/Private/Fbx/FbxMainImport.cpp
--- a/Engine/Source/Editor/UnrealEd/Private/Fbx/FbxMainImport.cpp
+++ b/Engine/Source/Editor/UnrealEd/Private/Fbx/FbxMainImport.cpp
@@ -1365,7 +1365,7 @@ bool FFbxImporter::ImportFromFile(const FString& Filename, const FString& Type,
 				{
 					if( FEngineAnalytics::IsAvailable() )
 					{
-						const static UEnum* FBXImportTypeEnum = FindObject<UEnum>(ANY_PACKAGE, TEXT("EFBXImportType"));
+						const static UEnum* FBXImportTypeEnum = StaticEnum<EFBXImportType>();
 						TArray<FAnalyticsEventAttribute> Attribs;
 
 						FString OriginalVendor(UTF8_TO_TCHAR(DocInfo->Original_ApplicationVendor.Get().Buffer()));
diff --git a/Engine/Source/Editor/UnrealEd/Private/Settings/ProjectPackagingSettingsCustomization.h b/Engine/Source/Editor/UnrealEd/Private/Settings/ProjectPackagingSettingsCustomization.h
--- a/Engine/Source/Editor/UnrealEd/Private/Settings/ProjectPackagingSettingsCustomization.h
+++ b/Engine/Source/Editor/UnrealEd/Private/Settings/ProjectPackagingSettingsCustomization.h
@@ -153,7 +153,7 @@ protected:
 			IDetailCategoryBuilder& ProjectCategory = LayoutBuilder.EditCategory("Project");
 			{
 				TSharedRef<FPropertyRestriction> BuildConfigurationRestriction = MakeShareable(new FPropertyRestriction(LOCTEXT("ContentOnlyRestrictionReason", "The DebugGame and Client build configurations are not available in content-only projects.")));
-				const UEnum* const ProjectPackagingBuildConfigurationsEnum = FindObject<UEnum>(ANY_PACKAGE, TEXT("EProjectPackagingBuildConfigurations"));		
+				const UEnum* const ProjectPackagingBuildConfigurationsEnum = StaticEnum<EProjectPackagingBuildConfigurations>();		
 				BuildConfigurationRestriction->AddDisabledValue(ProjectPackagingBuildConfigurationsEnum->GetNameStringByValue((uint8)EProjectPackagingBuildConfigurations::PPBC_DebugGame));
 				BuildConfigurationRestriction->AddDisabledValue(ProjectPackagingBuildConfigurationsEnum->GetNameStringByValue((uint8)EProjectPackagingBuildConfigurations::PPBC_DebugGameClient));
 
@@ -169,7 +169,7 @@ protected:
 			if (ClientTargetFileNames.Num() == 0)
 			{
 				TSharedRef<FPropertyRestriction> BuildConfigurationRestriction = MakeShareable(new FPropertyRestriction(LOCTEXT("ClientRestrictionReason", "The Client build configurations require a {ProjectName}Client.Target.cs file in your Project/Source folder.")));
-				const UEnum* const ProjectPackagingBuildConfigurationsEnum = FindObject<UEnum>(ANY_PACKAGE, TEXT("EProjectPackagingBuildConfigurations"));
+				const UEnum* const ProjectPackagingBuildConfigurationsEnum = StaticEnum<EProjectPackagingBuildConfigurations>();
 				BuildConfigurationRestriction->AddDisabledValue(ProjectPackagingBuildConfigurationsEnum->GetNameStringByValue((uint8)EProjectPackagingBuildConfigurations::PPBC_DebugGameClient));
 				BuildConfigurationRestriction->AddDisabledValue(ProjectPackagingBuildConfigurationsEnum->GetNameStringByValue((uint8)EProjectPackagingBuildConfigurations::PPBC_DevelopmentClient));
 				BuildConfigurationRestriction->AddDisabledValue(ProjectPackagingBuildConfigurationsEnum->GetNameStringByValue((uint8)EProjectPackagingBuildConfigurations::PPBC_ShippingClient));
diff --git a/Engine/Source/Programs/UnrealHeaderTool/Private/CodeGenerator.cpp b/Engine/Source/Programs/UnrealHeaderTool/Private/CodeGenerator.cpp
--- a/Engine/Source/Programs/UnrealHeaderTool/Private/CodeGenerator.cpp
+++ b/Engine/Source/Programs/UnrealHeaderTool/Private/CodeGenerator.cpp
@@ -2178,6 +2178,11 @@ void FNativeClassHeaderGenerator::ExportNativeGeneratedInitCode(FOutputDevice& O
 		Out.Logf(TEXT("\tIMPLEMENT_DYNAMIC_CLASS(%s, TEXT(\"%s\"), %u);\r\n"), ClassNameCPP, *OverriddenClassName, ClassHash);
 	}
 
+	Out.Logf(TEXT("\ttemplate<> %sUClass* StaticClass<%s>()\r\n"), *GetAPIString(), ClassNameCPP);
+	Out.Logf(TEXT("\t{\r\n"));
+	Out.Logf(TEXT("\t\treturn %s::StaticClass();\r\n"), ClassNameCPP);
+	Out.Logf(TEXT("\t}\r\n"));
+
 	Out.Logf(TEXT("\tstatic FCompiledInDefer Z_CompiledInDefer_UClass_%s(%s, &%s::StaticClass, TEXT(\"%s\"), TEXT(\"%s\"), %s, %s, %s, %s);\r\n"),
 		ClassNameCPP,
 		*SingletonName,
@@ -2939,6 +2944,9 @@ void FNativeClassHeaderGenerator::ExportClassFromSourceFileInner(
 		OutGeneratedHeaderText.Log(*BodyMacros);
 	}
 
+	// Forward declare the StaticClass specialisation in the header
+	OutGeneratedHeaderText.Logf(TEXT("template<> %sUClass* StaticClass<class %s>();\r\n\r\n"), *GetAPIString(), ClassCPPName);
+
 	// If there is a serialization function implementation for the CPP file, add it now
 	if (GeneratedSerializeFunctionCPP.Len())
 	{
@@ -3202,6 +3210,38 @@ void FNativeClassHeaderGenerator::ExportEnum(FOutputDevice& Out, UEnum* Enum)
 		Out.Logf( TEXT("\\\r\n\top(%s) "), *QualifiedEnumValue );
 	}
 	Out.Logf( TEXT("\r\n") );
+
+	// Forward declare the StaticEnum<> specialisation for enum classes
+	if (const EUnderlyingEnumType* EnumPropType = GEnumUnderlyingTypes.Find(Enum))
+	{
+		check(Enum->GetCppForm() == UEnum::ECppForm::EnumClass);
+
+		FString UnderlyingTypeString;
+
+		if (*EnumPropType != EUnderlyingEnumType::Unspecified)
+		{
+			UnderlyingTypeString = TEXT(" : ");
+
+			switch (*EnumPropType)
+			{
+			case EUnderlyingEnumType::int8:        UnderlyingTypeString += TNameOf<int8>::GetName();	break;
+			case EUnderlyingEnumType::int16:       UnderlyingTypeString += TNameOf<int16>::GetName();	break;
+			case EUnderlyingEnumType::int32:       UnderlyingTypeString += TNameOf<int32>::GetName();	break;
+			case EUnderlyingEnumType::int64:       UnderlyingTypeString += TNameOf<int64>::GetName();	break;
+			case EUnderlyingEnumType::uint8:       UnderlyingTypeString += TNameOf<uint8>::GetName();	break;
+			case EUnderlyingEnumType::uint16:      UnderlyingTypeString += TNameOf<uint16>::GetName();	break;
+			case EUnderlyingEnumType::uint32:      UnderlyingTypeString += TNameOf<uint32>::GetName();	break;
+			case EUnderlyingEnumType::uint64:      UnderlyingTypeString += TNameOf<uint64>::GetName();	break;
+			default:
+				check(false);
+			}
+		}
+
+		Out.Logf( TEXT("\r\n") );
+		Out.Logf( TEXT("enum class %s%s;\r\n"), *Enum->CppType, *UnderlyingTypeString );
+		Out.Logf( TEXT("template<> %sUEnum* StaticEnum<%s>();\r\n"), *GetAPIString(), *Enum->CppType );
+		Out.Logf( TEXT("\r\n") );
+	}
 }
 
 // Exports the header text for the list of structs specified (GENERATED_BODY impls)
@@ -3269,6 +3309,15 @@ void FNativeClassHeaderGenerator::ExportGeneratedStructBodyMacros(FOutputDevice&
 		Out.Logf(TEXT("\treturn Singleton;\r\n"));
 		Out.Logf(TEXT("}\r\n"));
 
+		// Forward declare the StaticStruct specialisation in the header
+		OutGeneratedHeaderText.Logf(TEXT("template<> %sUScriptStruct* StaticStruct<struct %s>();\r\n\r\n"), *GetAPIString(), StructNameCPP);
+
+		// Generate the StaticStruct specialisation
+		Out.Logf(TEXT("template<> %sUScriptStruct* StaticStruct<%s>()\r\n"), *GetAPIString(), StructNameCPP);
+		Out.Logf(TEXT("{\r\n"));
+		Out.Logf(TEXT("\treturn %s::StaticStruct();\r\n"), StructNameCPP);
+		Out.Logf(TEXT("}\r\n"));
+
 		Out.Logf(TEXT("static FCompiledInDeferStruct Z_CompiledInDeferStruct_UScriptStruct_%s(%s::StaticStruct, TEXT(\"%s\"), TEXT(\"%s\"), %s, %s, %s);\r\n"),
 			StructNameCPP,
 			StructNameCPP,
@@ -3477,6 +3526,11 @@ void FNativeClassHeaderGenerator::ExportGeneratedEnumInitCode(FOutputDevice& Out
 	Out.Logf(TEXT("\t\treturn Singleton;\r\n"));
 	Out.Logf(TEXT("\t}\r\n"));
 
+	Out.Logf(TEXT("\ttemplate<> %sUEnum* StaticEnum<%s>()\r\n"), *GetAPIString(), *Enum->CppType);
+	Out.Logf(TEXT("\t{\r\n"));
+	Out.Logf(TEXT("\t\treturn %s_StaticEnum();\r\n"), *Enum->GetName());
+	Out.Logf(TEXT("\t}\r\n"));
+
 	Out.Logf(
 		TEXT("\tstatic FCompiledInDeferEnum Z_CompiledInDeferEnum_UEnum_%s(%s_StaticEnum, TEXT(\"%s\"), TEXT(\"%s\"), %s, %s, %s);\r\n"),
 		*EnumNameCpp,
diff --git a/Engine/Source/Programs/UnrealHeaderTool/Private/NativeClassExporter.h b/Engine/Source/Programs/UnrealHeaderTool/Private/NativeClassExporter.h
--- a/Engine/Source/Programs/UnrealHeaderTool/Private/NativeClassExporter.h
+++ b/Engine/Source/Programs/UnrealHeaderTool/Private/NativeClassExporter.h
@@ -188,7 +188,7 @@ private:
 	 * @param	Out		the output device for the mirror struct
 	 * @param	Enums	the enum to export
 	 */
-	static void ExportEnum(FOutputDevice& Out, UEnum* Enum);
+	void ExportEnum(FOutputDevice& Out, UEnum* Enum);
 
 	/**
 	 * Exports the inl text for enums declared in non-UClass headers.
diff --git a/Engine/Source/Runtime/AIModule/Private/AITypes.cpp b/Engine/Source/Runtime/AIModule/Private/AITypes.cpp
--- a/Engine/Source/Runtime/AIModule/Private/AITypes.cpp
+++ b/Engine/Source/Runtime/AIModule/Private/AITypes.cpp
@@ -64,7 +64,7 @@ void FAIResourceLock::ForceClearAllLocks()
 
 FString FAIResourceLock::GetLockPriorityName() const
 {
-	const static UEnum* SourceEnum = FindObject<UEnum>(ANY_PACKAGE, TEXT("EAIRequestPriority"));
+	const static UEnum* SourceEnum = StaticEnum<EAIRequestPriority::Type>();
 
 	FString LockNames;
 
diff --git a/Engine/Source/Runtime/AIModule/Private/Actions/PawnAction.cpp b/Engine/Source/Runtime/AIModule/Private/Actions/PawnAction.cpp
--- a/Engine/Source/Runtime/AIModule/Private/Actions/PawnAction.cpp
+++ b/Engine/Source/Runtime/AIModule/Private/Actions/PawnAction.cpp
@@ -14,7 +14,7 @@ namespace
 {
 	FString GetActionResultName(int64 Value)
 	{
-		static const UEnum* Enum = FindObject<UEnum>(ANY_PACKAGE, TEXT("EPawnActionResult"));
+		static const UEnum* Enum = StaticEnum<EPawnActionResult::Type>();
 		check(Enum);
 		return Enum->GetNameStringByValue(Value);
 	}
@@ -403,7 +403,7 @@ UPawnAction* UPawnAction::CreateActionInstance(UObject* WorldContextObject, TSub
 
 FString UPawnAction::GetStateDescription() const
 {
-	static const UEnum* AbortStateEnum = FindObject<UEnum>(ANY_PACKAGE, TEXT("EPawnActionAbortState")); 
+	static const UEnum* AbortStateEnum = StaticEnum<EPawnActionAbortState::Type>(); 
 		
 	if (AbortState != EPawnActionAbortState::NotBeingAborted)
 	{
@@ -414,7 +414,7 @@ FString UPawnAction::GetStateDescription() const
 
 FString UPawnAction::GetPriorityName() const
 {
-	static const UEnum* Enum = FindObject<UEnum>(ANY_PACKAGE, TEXT("EAIRequestPriority"));
+	static const UEnum* Enum = StaticEnum<EAIRequestPriority::Type>();
 	check(Enum);
 	return Enum->GetNameStringByValue(GetPriority());
 }
diff --git a/Engine/Source/Runtime/AIModule/Private/Actions/PawnActionsComponent.cpp b/Engine/Source/Runtime/AIModule/Private/Actions/PawnActionsComponent.cpp
--- a/Engine/Source/Runtime/AIModule/Private/Actions/PawnActionsComponent.cpp
+++ b/Engine/Source/Runtime/AIModule/Private/Actions/PawnActionsComponent.cpp
@@ -17,14 +17,14 @@ namespace
 {
 	FString GetEventName(int64 Value)
 	{
-		static const UEnum* Enum = FindObject<UEnum>(ANY_PACKAGE, TEXT("EPawnActionEventType"));
+		static const UEnum* Enum = StaticEnum<EPawnActionEventType::Type>();
 		check(Enum);
 		return Enum->GetNameStringByValue(Value);
 	}
 
 	FString GetPriorityName(int64 Value)
 	{
-		static const UEnum* Enum = FindObject<UEnum>(ANY_PACKAGE, TEXT("EAIRequestPriority"));
+		static const UEnum* Enum = StaticEnum<EAIRequestPriority::Type>();
 		check(Enum);
 		return Enum->GetNameStringByValue(Value);
 	}
diff --git a/Engine/Source/Runtime/AIModule/Private/BehaviorTree/Decorators/BTDecorator_Blackboard.cpp b/Engine/Source/Runtime/AIModule/Private/BehaviorTree/Decorators/BTDecorator_Blackboard.cpp
--- a/Engine/Source/Runtime/AIModule/Private/BehaviorTree/Decorators/BTDecorator_Blackboard.cpp
+++ b/Engine/Source/Runtime/AIModule/Private/BehaviorTree/Decorators/BTDecorator_Blackboard.cpp
@@ -105,19 +105,19 @@ static void CacheOperationEnums()
 {
 	if (BasicOpEnum == NULL)
 	{
-		BasicOpEnum = FindObject<UEnum>(ANY_PACKAGE, TEXT("EBasicKeyOperation"));
+		BasicOpEnum = StaticEnum<EBasicKeyOperation::Type>();
 		check(BasicOpEnum);
 	}
 
 	if (ArithmeticOpEnum == NULL)
 	{
-		ArithmeticOpEnum = FindObject<UEnum>(ANY_PACKAGE, TEXT("EArithmeticKeyOperation"));
+		ArithmeticOpEnum = StaticEnum<EArithmeticKeyOperation::Type>();
 		check(ArithmeticOpEnum);
 	}
 
 	if (TextOpEnum == NULL)
 	{
-		TextOpEnum = FindObject<UEnum>(ANY_PACKAGE, TEXT("ETextKeyOperation"));
+		TextOpEnum = StaticEnum<ETextKeyOperation::Type>();
 		check(TextOpEnum);
 	}
 }
diff --git a/Engine/Source/Runtime/AIModule/Private/BehaviorTree/Decorators/BTDecorator_DoesPathExist.cpp b/Engine/Source/Runtime/AIModule/Private/BehaviorTree/Decorators/BTDecorator_DoesPathExist.cpp
--- a/Engine/Source/Runtime/AIModule/Private/BehaviorTree/Decorators/BTDecorator_DoesPathExist.cpp
+++ b/Engine/Source/Runtime/AIModule/Private/BehaviorTree/Decorators/BTDecorator_DoesPathExist.cpp
@@ -89,7 +89,7 @@ bool UBTDecorator_DoesPathExist::CalculateRawConditionValue(UBehaviorTreeCompone
 
 FString UBTDecorator_DoesPathExist::GetStaticDescription() const 
 {
-	const UEnum* PathTypeEnum = FindObject<UEnum>(ANY_PACKAGE, TEXT("EPathExistanceQueryType"));
+	const UEnum* PathTypeEnum = StaticEnum<EPathExistanceQueryType::Type>();
 	check(PathTypeEnum);
 
 	return FString::Printf(TEXT("%s: Find path from %s to %s (mode:%s)"),
diff --git a/Engine/Source/Runtime/AIModule/Private/BrainComponent.cpp b/Engine/Source/Runtime/AIModule/Private/BrainComponent.cpp
--- a/Engine/Source/Runtime/AIModule/Private/BrainComponent.cpp
+++ b/Engine/Source/Runtime/AIModule/Private/BrainComponent.cpp
@@ -195,7 +195,7 @@ UBrainComponent::UBrainComponent(const FObjectInitializer& ObjectInitializer) :
 #if ENABLE_VISUAL_LOG
 void UBrainComponent::DescribeSelfToVisLog(FVisualLogEntry* Snapshot) const
 {
-	const static UEnum* PriorityEnum = FindObject<UEnum>(ANY_PACKAGE, TEXT("EAIRequestPriority"));
+	const static UEnum* PriorityEnum = StaticEnum<EAIRequestPriority::Type>();
 
 	if (IsPendingKill())
 	{
diff --git a/Engine/Source/Runtime/AIModule/Private/EnvironmentQuery/EnvQueryManager.cpp b/Engine/Source/Runtime/AIModule/Private/EnvironmentQuery/EnvQueryManager.cpp
--- a/Engine/Source/Runtime/AIModule/Private/EnvironmentQuery/EnvQueryManager.cpp
+++ b/Engine/Source/Runtime/AIModule/Private/EnvironmentQuery/EnvQueryManager.cpp
@@ -692,7 +692,7 @@ TSharedPtr<FEnvQueryInstance> UEnvQueryManager::CreateQueryInstance(const UEnvQu
 	if (InstanceTemplate == NULL)
 	{
 		SCOPE_CYCLE_COUNTER(STAT_AI_EQS_LoadTime);
-		static const UEnum* RunModeEnum = FindObject<UEnum>(ANY_PACKAGE, TEXT("EEnvQueryRunMode"));
+		static const UEnum* RunModeEnum = StaticEnum<EEnvQueryRunMode::Type>();
 
 		// duplicate template in manager's world for BP based nodes
 		const FString NewInstanceName = RunModeEnum
diff --git a/Engine/Source/Runtime/AIModule/Private/EnvironmentQuery/EnvQueryTypes.cpp b/Engine/Source/Runtime/AIModule/Private/EnvironmentQuery/EnvQueryTypes.cpp
--- a/Engine/Source/Runtime/AIModule/Private/EnvironmentQuery/EnvQueryTypes.cpp
+++ b/Engine/Source/Runtime/AIModule/Private/EnvironmentQuery/EnvQueryTypes.cpp
@@ -154,7 +154,7 @@ FText FEnvTraceData::ToText(FEnvTraceData::EDescriptionMode DescMode) const
 
 		if (DescMode == FEnvTraceData::Brief)
 		{
-			static UEnum* ChannelEnum = FindObject<UEnum>(ANY_PACKAGE, TEXT("ETraceTypeQuery"), true);
+			static UEnum* ChannelEnum = StaticEnum<ETraceTypeQuery>();
 
 			FFormatNamedArguments Args;
 			Args.Add(TEXT("ExtentDescription"), Desc);
diff --git a/Engine/Source/Runtime/AIModule/Private/EnvironmentQuery/Tests/EnvQueryTest_Overlap.cpp b/Engine/Source/Runtime/AIModule/Private/EnvironmentQuery/Tests/EnvQueryTest_Overlap.cpp
--- a/Engine/Source/Runtime/AIModule/Private/EnvironmentQuery/Tests/EnvQueryTest_Overlap.cpp
+++ b/Engine/Source/Runtime/AIModule/Private/EnvironmentQuery/Tests/EnvQueryTest_Overlap.cpp
@@ -71,7 +71,7 @@ void UEnvQueryTest_Overlap::RunTest(FEnvQueryInstance& QueryInstance) const
 
 FText UEnvQueryTest_Overlap::GetDescriptionTitle() const
 {
-	UEnum* ShapeEnum = FindObject<UEnum>(ANY_PACKAGE, TEXT("EEnvOverlapShape"), true);
+	UEnum* ShapeEnum = StaticEnum<EEnvOverlapShape::Type>();
 	FString ShapeDesc = ShapeEnum->GetDisplayNameTextByValue(OverlapData.OverlapShape).ToString();
 
 	return FText::FromString(FString::Printf(TEXT("%s: %s"), 
@@ -80,10 +80,10 @@ FText UEnvQueryTest_Overlap::GetDescriptionTitle() const
 
 FText UEnvQueryTest_Overlap::GetDescriptionDetails() const
 {
-	UEnum* ShapeEnum = FindObject<UEnum>(ANY_PACKAGE, TEXT("EEnvOverlapShape"), true);
+	UEnum* ShapeEnum = StaticEnum<EEnvOverlapShape::Type>();
 	FString ShapeDesc = ShapeEnum->GetDisplayNameTextByValue(OverlapData.OverlapShape).ToString();
 
-	UEnum* ChannelEnum = FindObject<UEnum>(ANY_PACKAGE, TEXT("ECollisionChannel"), true);
+	UEnum* ChannelEnum = StaticEnum<ECollisionChannel>();
 	FString ChannelDesc = ChannelEnum->GetDisplayNameTextByValue(OverlapData.OverlapChannel).ToString();
 
 	FString SizeDesc;
diff --git a/Engine/Source/Runtime/AIModule/Private/EnvironmentQuery/Tests/EnvQueryTest_Trace.cpp b/Engine/Source/Runtime/AIModule/Private/EnvironmentQuery/Tests/EnvQueryTest_Trace.cpp
--- a/Engine/Source/Runtime/AIModule/Private/EnvironmentQuery/Tests/EnvQueryTest_Trace.cpp
+++ b/Engine/Source/Runtime/AIModule/Private/EnvironmentQuery/Tests/EnvQueryTest_Trace.cpp
@@ -99,7 +99,7 @@ void UEnvQueryTest_Trace::PostLoad()
 
 FText UEnvQueryTest_Trace::GetDescriptionTitle() const
 {
-	UEnum* ChannelEnum = FindObject<UEnum>(ANY_PACKAGE, TEXT("ETraceTypeQuery"), true);
+	UEnum* ChannelEnum = StaticEnum<ETraceTypeQuery>();
 	FString ChannelDesc = ChannelEnum->GetDisplayNameTextByValue(TraceData.TraceChannel).ToString();
 
 	FString DirectionDesc = TraceFromContext.IsDynamic() ?
diff --git a/Engine/Source/Runtime/AIModule/Private/Navigation/PathFollowingComponent.cpp b/Engine/Source/Runtime/AIModule/Private/Navigation/PathFollowingComponent.cpp
--- a/Engine/Source/Runtime/AIModule/Private/Navigation/PathFollowingComponent.cpp
+++ b/Engine/Source/Runtime/AIModule/Private/Navigation/PathFollowingComponent.cpp
@@ -71,7 +71,7 @@ FPathFollowingResult::FPathFollowingResult(EPathFollowingResult::Type ResultCode
 
 FString FPathFollowingResult::ToString() const
 {
-	static UEnum* ResultEnum = FindObject<UEnum>(ANY_PACKAGE, TEXT("EPathFollowingResult"));
+	static UEnum* ResultEnum = StaticEnum<EPathFollowingResult::Type>();
 	return FString::Printf(TEXT("%s[%s]"), ResultEnum ? *ResultEnum->GetNameStringByValue(Code) : TEXT("??"), *FPathFollowingResultFlags::ToString(Flags));
 }
 
@@ -228,7 +228,7 @@ FString GetPathDescHelper(FNavPathSharedPtr Path)
 
 void UPathFollowingComponent::OnPathEvent(FNavigationPath* InPath, ENavPathEvent::Type Event)
 {
-	const static UEnum* NavPathEventEnum = FindObject<UEnum>(ANY_PACKAGE, TEXT("ENavPathEvent"));
+	const static UEnum* NavPathEventEnum = StaticEnum<ENavPathEvent::Type>();
 	UE_VLOG(GetOwner(), LogPathFollowing, Log, TEXT("OnPathEvent: %s"), *NavPathEventEnum->GetNameStringByValue(Event));
 
 	if (InPath && Path.Get() == InPath)
@@ -1624,7 +1624,7 @@ bool UPathFollowingComponent::HasDirectPath() const
 
 FString UPathFollowingComponent::GetStatusDesc() const
 {
-	const static UEnum* StatusEnum = FindObject<UEnum>(ANY_PACKAGE, TEXT("EPathFollowingStatus"));
+	const static UEnum* StatusEnum = StaticEnum<EPathFollowingStatus::Type>();
 	if (StatusEnum)
 	{
 		return StatusEnum->GetNameStringByValue(Status);
@@ -1635,7 +1635,7 @@ FString UPathFollowingComponent::GetStatusDesc() const
 
 FString UPathFollowingComponent::GetResultDesc(EPathFollowingResult::Type Result) const
 {
-	const static UEnum* ResultEnum = FindObject<UEnum>(ANY_PACKAGE, TEXT("EPathFollowingResult"));
+	const static UEnum* ResultEnum = StaticEnum<EPathFollowingResult::Type>();
 	if (ResultEnum)
 	{
 		return ResultEnum->GetNameStringByValue(Result);
@@ -1783,7 +1783,7 @@ void UPathFollowingComponent::OnStartedFalling()
 
 void UPathFollowingComponent::LockResource(EAIRequestPriority::Type LockSource)
 {
-	const static UEnum* SourceEnum = FindObject<UEnum>(ANY_PACKAGE, TEXT("EAILockSource"));
+	const static UEnum* SourceEnum = StaticEnum<EAILockSource::Type>();
 	const bool bWasLocked = ResourceLock.IsLocked();
 
 	ResourceLock.SetLock(LockSource);
diff --git a/Engine/Source/Runtime/AIModule/Private/Tasks/AITask_MoveTo.cpp b/Engine/Source/Runtime/AIModule/Private/Tasks/AITask_MoveTo.cpp
--- a/Engine/Source/Runtime/AIModule/Private/Tasks/AITask_MoveTo.cpp
+++ b/Engine/Source/Runtime/AIModule/Private/Tasks/AITask_MoveTo.cpp
@@ -315,7 +315,7 @@ void UAITask_MoveTo::OnRequestFinished(FAIRequestID RequestID, const FPathFollow
 
 void UAITask_MoveTo::OnPathEvent(FNavigationPath* InPath, ENavPathEvent::Type Event)
 {
-	const static UEnum* NavPathEventEnum = FindObject<UEnum>(ANY_PACKAGE, TEXT("ENavPathEvent"));
+	const static UEnum* NavPathEventEnum = StaticEnum<ENavPathEvent::Type>();
 	UE_VLOG(GetGameplayTasksComponent(), LogGameplayTasks, Log, TEXT("%s> Path event: %s"), *GetName(), *NavPathEventEnum->GetNameStringByValue(Event));
 
 	switch (Event)
diff --git a/Engine/Source/Runtime/CoreUObject/Public/UObject/Class.h b/Engine/Source/Runtime/CoreUObject/Public/UObject/Class.h
--- a/Engine/Source/Runtime/CoreUObject/Public/UObject/Class.h
+++ b/Engine/Source/Runtime/CoreUObject/Public/UObject/Class.h
@@ -15,6 +15,7 @@
 #include "Math/RandomStream.h"
 #include "UObject/GarbageCollection.h"
 #include "UObject/CoreNative.h"
+#include "UObject/ReflectedTypeAccessors.h"
 #include "Templates/HasGetTypeHash.h"
 #include "Templates/IsAbstract.h"
 #include "Templates/IsEnum.h"
@@ -1854,6 +1855,86 @@ public:
 		out_TextValue = GetDisplayValueAsText( EnumPath, EnumeratorValue);
 	}
 
+	/**
+	 * @param EnumeratorValue  Enumerator Value.
+	 *
+	 * @return the name associated with the enumerator for the specified enum value for the enum specified by the template type.
+	 */
+	template<typename EnumType>
+	FORCEINLINE static FName GetValueAsName(const EnumType EnumeratorValue)
+	{
+		// For the C++ enum.
+		static_assert(TIsEnum<EnumType>::Value, "Should only call this with enum types");
+		UEnum* EnumClass = StaticEnum<EnumType>();
+		check(EnumClass != nullptr);
+		return EnumClass->GetNameByValue((int64)EnumeratorValue);
+	}
+
+	template<typename EnumType>
+	FORCEINLINE static FName GetValueAsName(const TEnumAsByte<EnumType> EnumeratorValue)
+	{
+		return GetValueAsName((int64)EnumeratorValue.GetValue());
+	}
+
+	template<typename EnumType>
+	FORCEINLINE static void GetValueAsName(const EnumType EnumeratorValue, FName& out_NameValue )
+	{
+		out_NameValue = GetValueAsName(EnumeratorValue);
+	}
+
+	/**
+	 * @param EnumeratorValue  Enumerator Value.
+	 *
+	 * @return the string associated with the enumerator for the specified enum value for the enum specified by the template type.
+	 */
+	template<typename EnumType>
+	FORCEINLINE static FString GetValueAsString(const EnumType EnumeratorValue)
+	{
+		// For the C++ enum.
+		static_assert(TIsEnum<EnumType>::Value, "Should only call this with enum types");
+		return GetValueAsName(EnumeratorValue).ToString();
+	}
+
+	template<typename EnumType>
+	FORCEINLINE static FString GetValueAsString(const TEnumAsByte<EnumType> EnumeratorValue)
+	{
+		return GetValueAsString((int64)EnumeratorValue.GetValue());
+	}
+
+	template<typename EnumType>
+	FORCEINLINE static void GetValueAsString(const EnumType EnumeratorValue, FString& out_StringValue )
+	{
+		out_StringValue = GetValueAsString(EnumeratorValue );
+	}
+
+
+	/**
+	 * @param EnumeratorValue  Enumerator Value.
+	 *
+	 * @return the localized display string associated with the specified enum value for the enum specified by the template type.
+	 */
+	template<typename EnumType>
+	FORCEINLINE static FText GetDisplayValueAsText(const EnumType EnumeratorValue )
+	{
+		// For the C++ enum.
+		static_assert(TIsEnum<EnumType>::Value, "Should only call this with enum types");
+		UEnum* EnumClass = StaticEnum<EnumType>();
+		check(EnumClass != nullptr);
+		return EnumClass->GetDisplayNameTextByValue((int64)EnumeratorValue);
+	}
+
+	template<typename EnumType>
+	FORCEINLINE static FText GetDisplayValueAsText(const TEnumAsByte<EnumType> EnumeratorValue)
+	{
+		return GetDisplayValueAsText((int64)EnumeratorValue.GetValue());
+	}
+
+	template<typename EnumType>
+	FORCEINLINE static void GetDisplayValueAsText(const EnumType EnumeratorValue, FText& out_TextValue )
+	{
+		out_TextValue = GetDisplayValueAsText(EnumeratorValue);
+	}
+
 	// Deprecated Functions
 	DEPRECATED(4.16, "FindEnumIndex is deprecated, call GetIndexByName or GetValueByName instead")
 	int32 FindEnumIndex(FName InName) const { return GetIndexByName(InName, EGetByNameFlags::ErrorIfNotFound); }
diff --git a/Engine/Source/Runtime/CoreUObject/Public/UObject/ReflectedTypeAccessors.h b/Engine/Source/Runtime/CoreUObject/Public/UObject/ReflectedTypeAccessors.h
new file mode 100644
--- /dev/null
+++ b/Engine/Source/Runtime/CoreUObject/Public/UObject/ReflectedTypeAccessors.h
@@ -0,0 +1,15 @@
+// Copyright 1998-2018 Epic Games, Inc. All Rights Reserved.
+
+#pragma once
+
+class UClass;
+class UScriptStruct;
+class UEnum;
+
+/*-----------------------------------------------------------------------------
+C++ templated Static(Class/Struct/Enum) retrieval function prototypes.
+-----------------------------------------------------------------------------*/
+
+template<typename ClassType>	UClass*			StaticClass();
+template<typename StructType>	UScriptStruct*	StaticStruct();
+template<typename EnumType>		UEnum*			StaticEnum();
diff --git a/Engine/Source/Runtime/Engine/Private/Collision/Collision.cpp b/Engine/Source/Runtime/Engine/Private/Collision/Collision.cpp
--- a/Engine/Source/Runtime/Engine/Private/Collision/Collision.cpp
+++ b/Engine/Source/Runtime/Engine/Private/Collision/Collision.cpp
@@ -655,7 +655,7 @@ namespace CollisionResponseConsoleCommands
 
 	void ListCollisionChannelNames()
 	{
-		const UEnum *ChannelEnum = FindObject<UEnum>(ANY_PACKAGE, TEXT("ECollisionChannel"), true);
+		const UEnum *ChannelEnum = StaticEnum<ECollisionChannel>();
 		if (ChannelEnum)
 		{
 			for (int32 ChannelValue = 0; ChannelValue < ECollisionChannel::ECC_MAX; ChannelValue++)
@@ -754,7 +754,7 @@ namespace CollisionResponseConsoleCommands
 
 	ECollisionChannel StringToCollisionChannel(const FString& InString)
 	{
-		const UEnum *ChannelEnum = FindObject<UEnum>(ANY_PACKAGE, TEXT("ECollisionChannel"), true);
+		const UEnum *ChannelEnum = StaticEnum<ECollisionChannel>();
 		if (ChannelEnum)
 		{
 			int32 ChannelInt = INDEX_NONE;
@@ -816,7 +816,7 @@ namespace CollisionResponseConsoleCommands
 
 	ECollisionTraceFlag StringToCollisionComplexity(const FString& InString)
 	{
-		const UEnum *ComplexityEnum = FindObject<UEnum>(ANY_PACKAGE, TEXT("ECollisionTraceFlag"), true);
+		const UEnum *ComplexityEnum = StaticEnum<ECollisionTraceFlag>();
 		if (ComplexityEnum)
 		{
 			int32 ComplexityInt = INDEX_NONE;
@@ -900,7 +900,7 @@ namespace CollisionResponseConsoleCommands
 		if (Results.Num() > 0)
 		{
 			Results.Sort(FSortComponentsWithResponseToProfile(RequiredResponse));
-			const UEnum *ChannelEnum = FindObject<UEnum>(ANY_PACKAGE, TEXT("ECollisionChannel"), true);
+			const UEnum *ChannelEnum = StaticEnum<ECollisionChannel>();
 			TMap<ECollisionChannel, FString> EnumToDisplayNameMap;
 
 			// Get max column widths for some data
@@ -1037,7 +1037,7 @@ namespace CollisionResponseConsoleCommands
 		}
 		// Display Summary
 		check(RequiredResponse < ECollisionResponse::ECR_MAX);
-		const UEnum *ChannelEnum = FindObject<UEnum>(ANY_PACKAGE, TEXT("ECollisionChannel"), true);
+		const UEnum *ChannelEnum = StaticEnum<ECollisionChannel>();
 		const FString ChannelName = (ChannelEnum ? ChannelEnum->GetNameStringByValue(TestChannel) : TEXT("<unknown>"));
 		const FString ChannelDisplayName = GetDisplayNameText(ChannelEnum, TestChannel, ChannelName);
 		UE_LOG(LogCollisionCommands, Log, TEXT("----------------------------------------------------------------------"));
@@ -1147,7 +1147,7 @@ namespace CollisionResponseConsoleCommands
 			}
 
 			Results.Sort(FSortComponentsForComplexity(ComponentAssetNameMap));
-			const UEnum *ComplexityEnum = FindObject<UEnum>(ANY_PACKAGE, TEXT("ECollisionTraceFlag"), true);
+			const UEnum *ComplexityEnum = StaticEnum<ECollisionTraceFlag>();
 			if (!ensure(ComplexityEnum != nullptr))
 			{
 				return;
diff --git a/Engine/Source/Runtime/Engine/Private/Collision/CollisionProfile.cpp b/Engine/Source/Runtime/Engine/Private/Collision/CollisionProfile.cpp
--- a/Engine/Source/Runtime/Engine/Private/Collision/CollisionProfile.cpp
+++ b/Engine/Source/Runtime/Engine/Private/Collision/CollisionProfile.cpp
@@ -289,7 +289,7 @@ void UCollisionProfile::LoadProfileConfig(bool bForceInit)
 	FString GameTraceChannel	= TEXT("GameTraceChannel");
 
 	// find the enum
-	UEnum* Enum = FindObject<UEnum>(ANY_PACKAGE, TEXT("ECollisionChannel"), true);
+	UEnum* Enum = StaticEnum<ECollisionChannel>();
 	// we need this Enum
 	check (Enum);
 	UStruct* Struct = FCollisionResponseContainer::StaticStruct(); 
@@ -310,8 +310,8 @@ void UCollisionProfile::LoadProfileConfig(bool bForceInit)
 	int32 PrefixLen = FString(TEXT("ECC_")).Len();
 
 	// need to have mapping table between ECollisionChannel and EObjectTypeQuery/ETraceTypeQuery
-	UEnum* ObjectTypeEnum = FindObject<UEnum>(ANY_PACKAGE, TEXT("EObjectTypeQuery"), true);
-	UEnum* TraceTypeEnum = FindObject<UEnum>(ANY_PACKAGE, TEXT("ETraceTypeQuery"), true);
+	UEnum* ObjectTypeEnum = StaticEnum<EObjectTypeQuery>();
+	UEnum* TraceTypeEnum = StaticEnum<ETraceTypeQuery>();
 	check (ObjectTypeEnum && TraceTypeEnum);
 
 	for ( int32 EnumIndex=0; EnumIndex<NumEnum; ++EnumIndex )
diff --git a/Engine/Source/Runtime/Engine/Private/GameViewportClient.cpp b/Engine/Source/Runtime/Engine/Private/GameViewportClient.cpp
--- a/Engine/Source/Runtime/Engine/Private/GameViewportClient.cpp
+++ b/Engine/Source/Runtime/Engine/Private/GameViewportClient.cpp
@@ -371,7 +371,7 @@ void UGameViewportClient::Init(struct FWorldContext& WorldContext, UGameInstance
 	FString DefaultViewportMouseCaptureMode;
 	if (FParse::Value(FCommandLine::Get(), TEXT("DefaultViewportMouseCaptureMode="), DefaultViewportMouseCaptureMode))
 	{
-		const UEnum* EnumPtr = FindObject<UEnum>(ANY_PACKAGE, TEXT("EMouseCaptureMode"));
+		const UEnum* EnumPtr = StaticEnum<EMouseCaptureMode>();
 		checkf(EnumPtr, TEXT("Unable to find EMouseCaptureMode enum"));
 		if (EnumPtr)
 		{
diff --git a/Engine/Source/Runtime/Engine/Private/Materials/Material.cpp b/Engine/Source/Runtime/Engine/Private/Materials/Material.cpp
--- a/Engine/Source/Runtime/Engine/Private/Materials/Material.cpp
+++ b/Engine/Source/Runtime/Engine/Private/Materials/Material.cpp
@@ -607,7 +607,7 @@ void SetCompactFullNameFromObject(FCompactFullName &Dest, UObject* InDepObject)
 
 FString MaterialDomainString(EMaterialDomain MaterialDomain)
 {
-	static const UEnum* Enum = FindObject<UEnum>(ANY_PACKAGE, TEXT("EMaterialDomain"));
+	static const UEnum* Enum = StaticEnum<EMaterialDomain>();
 	check(Enum);
 	return Enum->GetNameStringByValue(int64(MaterialDomain));
 }
@@ -4000,7 +4000,7 @@ void UMaterial::DumpDebugInfo()
 	UE_LOG(LogConsoleResponse, Display, TEXT("----------------------------- %s"), *GetFullName());
 
 	{
-		static const UEnum* Enum = FindObject<UEnum>(ANY_PACKAGE, TEXT("EMaterialDomain"));
+		static const UEnum* Enum = StaticEnum<EMaterialDomain>();
 		check(Enum);
 		UE_LOG(LogConsoleResponse, Display, TEXT("  MaterialDomain %s"), *Enum->GetNameStringByValue(int64(MaterialDomain)));
 	}
diff --git a/Engine/Source/Runtime/Engine/Private/Materials/MaterialExpressions.cpp b/Engine/Source/Runtime/Engine/Private/Materials/MaterialExpressions.cpp
--- a/Engine/Source/Runtime/Engine/Private/Materials/MaterialExpressions.cpp
+++ b/Engine/Source/Runtime/Engine/Private/Materials/MaterialExpressions.cpp
@@ -2182,7 +2182,7 @@ void UMaterialExpressionTextureProperty::GetTexturesForceMaterialRecompile(TArra
 void UMaterialExpressionTextureProperty::GetCaption(TArray<FString>& OutCaptions) const
 {
 #if WITH_EDITOR
-	const UEnum* TexturePropertyEnum = FindObject<UEnum>(NULL, TEXT("Engine.EMaterialExposedTextureProperty"));
+	const UEnum* TexturePropertyEnum = StaticEnum<EMaterialExposedTextureProperty>();
 	check(TexturePropertyEnum);
 
 	const FString PropertyDisplayName = TexturePropertyEnum->GetDisplayNameTextByValue(Property).ToString();
@@ -7653,7 +7653,7 @@ int32 UMaterialExpressionViewProperty::Compile(class FMaterialCompiler* Compiler
 
 void UMaterialExpressionViewProperty::GetCaption(TArray<FString>& OutCaptions) const
 {
-	const UEnum* ViewPropertyEnum = FindObject<UEnum>(NULL, TEXT("Engine.EMaterialExposedViewProperty"));
+	const UEnum* ViewPropertyEnum = StaticEnum<EMaterialExposedViewProperty>();
 	check(ViewPropertyEnum);
 
 	OutCaptions.Add(ViewPropertyEnum->GetDisplayNameTextByValue(Property).ToString());
@@ -7955,7 +7955,7 @@ int32 UMaterialExpressionSceneTexture::Compile(class FMaterialCompiler* Compiler
 
 void UMaterialExpressionSceneTexture::GetCaption(TArray<FString>& OutCaptions) const
 {
-	UEnum* Enum = FindObject<UEnum>(NULL, TEXT("Engine.ESceneTextureId"));
+	UEnum* Enum = StaticEnum<ESceneTextureId>();
 
 	check(Enum);
 
@@ -8430,8 +8430,8 @@ int32 UMaterialExpressionTransform::Compile(class FMaterialCompiler* Compiler, i
 void UMaterialExpressionTransform::GetCaption(TArray<FString>& OutCaptions) const
 {
 #if WITH_EDITOR
-	const UEnum* MVCTSEnum = FindObject<UEnum>(nullptr, TEXT("Engine.EMaterialVectorCoordTransformSource"));
-	const UEnum* MVCTEnum = FindObject<UEnum>(nullptr, TEXT("Engine.EMaterialVectorCoordTransform"));
+	const UEnum* MVCTSEnum = StaticEnum<EMaterialVectorCoordTransformSource>();
+	const UEnum* MVCTEnum = StaticEnum<EMaterialVectorCoordTransform>();
 	check(MVCTSEnum);
 	check(MVCTEnum);
 	
@@ -8532,7 +8532,7 @@ int32 UMaterialExpressionTransformPosition::Compile(class FMaterialCompiler* Com
 void UMaterialExpressionTransformPosition::GetCaption(TArray<FString>& OutCaptions) const
 {
 #if WITH_EDITOR
-	const UEnum* MPTSEnum = FindObject<UEnum>(nullptr, TEXT("Engine.EMaterialPositionTransformSource"));
+	const UEnum* MPTSEnum = StaticEnum<EMaterialPositionTransformSource>();
 	check(MPTSEnum);
 	
 	FString TransformDesc;
@@ -12876,7 +12876,7 @@ int32 UMaterialExpressionNoise::Compile(class FMaterialCompiler* Compiler, int32
 
 void UMaterialExpressionNoise::GetCaption(TArray<FString>& OutCaptions) const
 {
-	const UEnum* NFEnum = FindObject<UEnum>(nullptr, TEXT("Engine.ENoiseFunction"));
+	const UEnum* NFEnum = StaticEnum<ENoiseFunction>();
 	check(NFEnum);
 	OutCaptions.Add(NFEnum->GetDisplayNameTextByValue(NoiseFunction).ToString());
 	OutCaptions.Add(TEXT("Noise"));
@@ -12952,7 +12952,7 @@ int32 UMaterialExpressionVectorNoise::Compile(class FMaterialCompiler* Compiler,
 
 void UMaterialExpressionVectorNoise::GetCaption(TArray<FString>& OutCaptions) const
 {
-	const UEnum* VNFEnum = FindObject<UEnum>(nullptr, TEXT("Engine.EVectorNoiseFunction"));
+	const UEnum* VNFEnum = StaticEnum<EVectorNoiseFunction>();
 	check(VNFEnum);
 	OutCaptions.Add(VNFEnum->GetDisplayNameTextByValue(NoiseFunction).ToString());
 	OutCaptions.Add(TEXT("Vector Noise"));
diff --git a/Engine/Source/Runtime/Engine/Private/Materials/MaterialInstance.cpp b/Engine/Source/Runtime/Engine/Private/Materials/MaterialInstance.cpp
--- a/Engine/Source/Runtime/Engine/Private/Materials/MaterialInstance.cpp
+++ b/Engine/Source/Runtime/Engine/Private/Materials/MaterialInstance.cpp
@@ -4128,7 +4128,7 @@ void UMaterialInstance::DumpDebugInfo()
 
 		if (Base)
 		{
-			static const UEnum* Enum = FindObject<UEnum>(ANY_PACKAGE, TEXT("EMaterialDomain"));
+			static const UEnum* Enum = StaticEnum<EMaterialDomain>();
 			check(Enum);
 			UE_LOG(LogConsoleResponse, Display, TEXT("  MaterialDomain %s"), *Enum->GetNameStringByValue(int64(Base->MaterialDomain)));
 		}
diff --git a/Engine/Source/Runtime/Engine/Private/Materials/MaterialInterface.cpp b/Engine/Source/Runtime/Engine/Private/Materials/MaterialInterface.cpp
--- a/Engine/Source/Runtime/Engine/Private/Materials/MaterialInterface.cpp
+++ b/Engine/Source/Runtime/Engine/Private/Materials/MaterialInterface.cpp
@@ -70,7 +70,7 @@ UMaterialInterface::UMaterialInterface(const FObjectInitializer& ObjectInitializ
 
 		if (SamplerTypeEnum == nullptr)
 		{
-			SamplerTypeEnum = FindObject<UEnum>(NULL, TEXT("/Script/Engine.EMaterialSamplerType"));
+			SamplerTypeEnum = StaticEnum<EMaterialSamplerType>();
 			check(SamplerTypeEnum);
 		}
 
diff --git a/Engine/Source/Runtime/Engine/Private/PhysicsEngine/PhysicsSettings.cpp b/Engine/Source/Runtime/Engine/Private/PhysicsEngine/PhysicsSettings.cpp
--- a/Engine/Source/Runtime/Engine/Private/PhysicsEngine/PhysicsSettings.cpp
+++ b/Engine/Source/Runtime/Engine/Private/PhysicsEngine/PhysicsSettings.cpp
@@ -119,7 +119,7 @@ void UPhysicsSettings::LoadSurfaceType()
 {
 	// read "SurfaceType" defines and set meta data for the enum
 	// find the enum
-	UEnum * Enum = FindObject<UEnum>(ANY_PACKAGE, TEXT("EPhysicalSurface"), true);
+	UEnum * Enum = StaticEnum<EPhysicalSurface>();
 	// we need this Enum
 	check(Enum);
 
diff --git a/Engine/Source/Runtime/Engine/Private/VoiceConfig.cpp b/Engine/Source/Runtime/Engine/Private/VoiceConfig.cpp
--- a/Engine/Source/Runtime/Engine/Private/VoiceConfig.cpp
+++ b/Engine/Source/Runtime/Engine/Private/VoiceConfig.cpp
@@ -39,7 +39,7 @@ int32 UVOIPStatics::GetVoiceSampleRate()
 
 	if (GConfig->GetString(TEXT("/Script/Engine.AudioSettings"), TEXT("VoiPSampleRate"), DesiredSampleRateStr, GEngineIni))
 	{
-		SampleRate = (int32)GetEnumValueFromString<EVoiceSampleRate>(TEXT("EVoiceSampleRate"), DesiredSampleRateStr);
+		SampleRate = (int32)GetEnumValueFromString<EVoiceSampleRate>(DesiredSampleRateStr);
 		if (SampleRate > 0)
 		{		
 			return SampleRate;
diff --git a/Engine/Source/Runtime/Engine/Public/Net/VoiceConfig.h b/Engine/Source/Runtime/Engine/Public/Net/VoiceConfig.h
--- a/Engine/Source/Runtime/Engine/Public/Net/VoiceConfig.h
+++ b/Engine/Source/Runtime/Engine/Public/Net/VoiceConfig.h
@@ -43,9 +43,9 @@ namespace MicSilenceDetectionConfig
 class USoundAttenuation;
 
 template <typename EnumType>
-static FORCEINLINE EnumType GetEnumValueFromString(const FString& EnumName, const FString& String)
+static FORCEINLINE EnumType GetEnumValueFromString(const FString& String)
 {
-	UEnum* Enum = FindObject<UEnum>((UObject*) ANY_PACKAGE, *EnumName, true);
+	UEnum* Enum = StaticEnum<EnumType>();
 	if (!Enum)
 	{
 		return EnumType(0);
diff --git a/Engine/Source/Runtime/GameplayTasks/Private/GameplayTask.cpp b/Engine/Source/Runtime/GameplayTasks/Private/GameplayTask.cpp
--- a/Engine/Source/Runtime/GameplayTasks/Private/GameplayTask.cpp
+++ b/Engine/Source/Runtime/GameplayTasks/Private/GameplayTask.cpp
@@ -400,7 +400,7 @@ FString UGameplayTask::GenerateDebugDescription() const
 
 FString UGameplayTask::GetTaskStateName() const
 {
-	static const UEnum* Enum = FindObject<UEnum>(ANY_PACKAGE, TEXT("EGameplayTaskState"));
+	static const UEnum* Enum = StaticEnum<EGameplayTaskState>();
 	check(Enum);
 	return Enum->GetNameStringByValue(int64(TaskState));
 }
diff --git a/Engine/Source/Runtime/GameplayTasks/Private/GameplayTasksComponent.cpp b/Engine/Source/Runtime/GameplayTasks/Private/GameplayTasksComponent.cpp
--- a/Engine/Source/Runtime/GameplayTasks/Private/GameplayTasksComponent.cpp
+++ b/Engine/Source/Runtime/GameplayTasks/Private/GameplayTasksComponent.cpp
@@ -16,7 +16,7 @@ namespace
 {
 	FORCEINLINE const TCHAR* GetGameplayTaskEventName(EGameplayTaskEvent Event)
 	{
-		/*static const UEnum* GameplayTaskEventEnum = FindObject<UEnum>(ANY_PACKAGE, TEXT("EGameplayTaskEvent"));
+		/*static const UEnum* GameplayTaskEventEnum = StaticEnum<EGameplayTaskEvent>();
 		return GameplayTaskEventEnum->GetDisplayNameTextByValue(static_cast<int64>(Event)).ToString();*/
 
 		return Event == EGameplayTaskEvent::Add ? TEXT("Add") : TEXT("Remove");
@@ -572,7 +572,7 @@ FString UGameplayTasksComponent::GetTasksPriorityQueueDescription() const
 
 FString UGameplayTasksComponent::GetTaskStateName(EGameplayTaskState Value)
 {
-	static const UEnum* Enum = FindObject<UEnum>(ANY_PACKAGE, TEXT("EGameplayTaskState"));
+	static const UEnum* Enum = StaticEnum<EGameplayTaskState>();
 	check(Enum);
 	return Enum->GetNameStringByValue(int64(Value));
 }
diff --git a/Engine/Source/Runtime/HeadMountedDisplay/Private/HeadMountedDisplayConsoleCommands.cpp b/Engine/Source/Runtime/HeadMountedDisplay/Private/HeadMountedDisplayConsoleCommands.cpp
--- a/Engine/Source/Runtime/HeadMountedDisplay/Private/HeadMountedDisplayConsoleCommands.cpp
+++ b/Engine/Source/Runtime/HeadMountedDisplay/Private/HeadMountedDisplayConsoleCommands.cpp
@@ -190,7 +190,7 @@ static FAutoConsoleCommand CShowTrackingSensorsCmd(
 
 static void TrackingOrigin(const TArray<FString>& Args, UWorld* , FOutputDevice& Ar)
 {
-	const static UEnum* TrackingOriginEnum = FindObject<UEnum>(ANY_PACKAGE, TEXT("EHMDTrackingOrigin"));
+	const static UEnum* TrackingOriginEnum = StaticEnum<EHMDTrackingOrigin::Type>();
 	int Origin = INDEX_NONE;
 	if (Args.Num())
 	{
@@ -248,7 +248,7 @@ static void SpectatorScreenMode(const TArray<FString>& Args, UWorld* , FOutputDe
 		return;
 	}
 
-	const static UEnum* ScreenModeEnum = FindObject<UEnum>(ANY_PACKAGE, TEXT("ESpectatorScreenMode"));
+	const static UEnum* ScreenModeEnum = StaticEnum<ESpectatorScreenMode>();
 	int ModeVal = INDEX_NONE;
 	if (Args.Num())
 	{
diff --git a/Engine/Source/Runtime/HeadMountedDisplay/Private/MotionControllerComponent.cpp b/Engine/Source/Runtime/HeadMountedDisplay/Private/MotionControllerComponent.cpp
--- a/Engine/Source/Runtime/HeadMountedDisplay/Private/MotionControllerComponent.cpp
+++ b/Engine/Source/Runtime/HeadMountedDisplay/Private/MotionControllerComponent.cpp
@@ -42,7 +42,7 @@ namespace LegacyMotionSources
 {
 	static bool GetSourceNameForHand(EControllerHand InHand, FName& OutSourceName)
 	{
-		UEnum* HandEnum = FindObject<UEnum>(ANY_PACKAGE, TEXT("EControllerHand"));
+		UEnum* HandEnum = StaticEnum<EControllerHand>();
 		if (HandEnum)
 		{
 			FString ValueName = HandEnum->GetNameStringByValue((int64)InHand);
diff --git a/Engine/Source/Runtime/MovieSceneTracks/Private/Sections/MovieSceneParticleSection.cpp b/Engine/Source/Runtime/MovieSceneTracks/Private/Sections/MovieSceneParticleSection.cpp
--- a/Engine/Source/Runtime/MovieSceneTracks/Private/Sections/MovieSceneParticleSection.cpp
+++ b/Engine/Source/Runtime/MovieSceneTracks/Private/Sections/MovieSceneParticleSection.cpp
@@ -8,7 +8,7 @@
 
 FMovieSceneParticleChannel::FMovieSceneParticleChannel()
 {
-	SetEnum(FindObject<UEnum>(ANY_PACKAGE, TEXT("EParticleKey")));
+	SetEnum(StaticEnum<EParticleKey>());
 }
 
 UMovieSceneParticleSection::UMovieSceneParticleSection( const FObjectInitializer& ObjectInitializer )
diff --git a/Engine/Source/Runtime/Renderer/Private/PostProcess/VisualizeShadingModels.cpp b/Engine/Source/Runtime/Renderer/Private/PostProcess/VisualizeShadingModels.cpp
--- a/Engine/Source/Runtime/Renderer/Private/PostProcess/VisualizeShadingModels.cpp
+++ b/Engine/Source/Runtime/Renderer/Private/PostProcess/VisualizeShadingModels.cpp
@@ -157,7 +157,7 @@ void FRCPassPostProcessVisualizeShadingModels::Process(FRenderingCompositePassCo
 	Canvas.DrawShadowedString( X, Y, *Line, GetStatsFont(), FLinearColor(0.5f, 0.5f, 0.5f));
 	Y += YStep;
 
-	UEnum* Enum = FindObject<UEnum>(NULL, TEXT("Engine.EMaterialShadingModel"));
+	UEnum* Enum = StaticEnum<EMaterialShadingModel>();
 	check(Enum);
 
 	Y += 5;
diff --git a/Engine/Source/Runtime/UMG/Private/Components/RichTextBlockImageDecorator.cpp b/Engine/Source/Runtime/UMG/Private/Components/RichTextBlockImageDecorator.cpp
--- a/Engine/Source/Runtime/UMG/Private/Components/RichTextBlockImageDecorator.cpp
+++ b/Engine/Source/Runtime/UMG/Private/Components/RichTextBlockImageDecorator.cpp
@@ -113,7 +113,7 @@ protected:
 		EStretch::Type Stretch = EStretch::ScaleToFit;
 		if (const FString* SstretchString = RunInfo.MetaData.Find(TEXT("stretch")))
 		{
-			static const UEnum* StretchEnum = FindObject<UEnum>(ANY_PACKAGE, TEXT("EStretch"));
+			static const UEnum* StretchEnum = StaticEnum<EStretch::Type>();
 			int64 StretchValue = StretchEnum->GetValueByNameString(*SstretchString);
 			if (StretchValue != INDEX_NONE)
 			{
diff --git a/Engine/Source/Runtime/VectorVM/Private/VectorVM.cpp b/Engine/Source/Runtime/VectorVM/Private/VectorVM.cpp
--- a/Engine/Source/Runtime/VectorVM/Private/VectorVM.cpp
+++ b/Engine/Source/Runtime/VectorVM/Private/VectorVM.cpp
@@ -1635,8 +1635,8 @@ void VectorVM::Init()
 	static bool Inited = false;
 	if (Inited == false)
 	{
-		g_VectorVMEnumStateObj = FindObject<UEnum>(ANY_PACKAGE, TEXT("EVectorVMOp"), true);
-		g_VectorVMEnumOperandObj = FindObject<UEnum>(ANY_PACKAGE, TEXT("EVectorVMOperandLocation"), true);
+		g_VectorVMEnumStateObj = StaticEnum<EVectorVMOp>();
+		g_VectorVMEnumOperandObj = StaticEnum<EVectorVMOperandLocation>();
 
 		// random noise
 		float TempTable[17][17][17];
