GitHub #4995 : [UE-62943] Fixed CodeGenerator to respect const keyword on return type of interface functions
Some tests added.

diff --git a/Engine/Source/Programs/UnrealHeaderTool/Private/CodeGenerator.cpp b/Engine/Source/Programs/UnrealHeaderTool/Private/CodeGenerator.cpp
--- a/Engine/Source/Programs/UnrealHeaderTool/Private/CodeGenerator.cpp
+++ b/Engine/Source/Programs/UnrealHeaderTool/Private/CodeGenerator.cpp
@@ -4209,6 +4209,11 @@ void FNativeClassHeaderGenerator::ExportNativeFunctionHeader(
 	UProperty* ReturnProperty = Function->GetReturnProperty();
 	if (ReturnProperty != nullptr)
 	{
+		if (ReturnProperty->HasAnyPropertyFlags(EPropertyFlags::CPF_ConstParm))
+		{
+			Out.Log(TEXT("const "));
+		}
+
 		FString ExtendedReturnType;
 		FString ReturnType = ReturnProperty->GetCPPType(&ExtendedReturnType, (FunctionHeaderStyle == EExportFunctionHeaderStyle::Definition && (FunctionType != EExportFunctionType::Interface) ? CPPF_Implementation : 0) | CPPF_ArgumentOrReturnValue);
 		OutFwdDecls.Add(ReturnProperty->GetCPPTypeForwardDeclaration());
diff --git a/Engine/Source/Programs/UnrealHeaderTool/Resources/UHTDebugging/TestInterface.h b/Engine/Source/Programs/UnrealHeaderTool/Resources/UHTDebugging/TestInterface.h
--- a/Engine/Source/Programs/UnrealHeaderTool/Resources/UHTDebugging/TestInterface.h
+++ b/Engine/Source/Programs/UnrealHeaderTool/Resources/UHTDebugging/TestInterface.h
@@ -19,4 +19,13 @@ public:
 
 	UFUNCTION(BlueprintNativeEvent)
 	FString SomeFunction(int32 Val) const;
+
+	UFUNCTION(BlueprintNativeEvent, BlueprintCallable)
+	const float SomeFunctionWithConstReturnType() const;
+
+	UFUNCTION(BlueprintNativeEvent, BlueprintCallable)
+	UObject* SomeFunctionWithNonConstPointerReturnType() const;
+
+	UFUNCTION(BlueprintNativeEvent, BlueprintCallable)
+	const UObject* SomeFunctionWithConstPointerReturnType() const;
 };
