Refactored structs to reduce executable space from UHT-generated code tables.

diff --git a/Engine/Source/Programs/UnrealHeaderTool/Private/CodeGenerator.cpp b/Engine/Source/Programs/UnrealHeaderTool/Private/CodeGenerator.cpp
--- a/Engine/Source/Programs/UnrealHeaderTool/Private/CodeGenerator.cpp
+++ b/Engine/Source/Programs/UnrealHeaderTool/Private/CodeGenerator.cpp
@@ -913,14 +913,14 @@
 		DeclOut.Logf(TEXT("%sstatic const UE4CodeGen_Private::FBytePropertyParams %s;\r\n"), DeclSpaces, *NameWithoutScope);
 
 		Out.Logf(
-			TEXT("%sconst UE4CodeGen_Private::FBytePropertyParams %s = { UE4CodeGen_Private::EPropertyClass::Byte, %s, %s, (EPropertyFlags)0x%016llx, %s, %s, %s, %s, %s };%s\r\n"),
+			TEXT("%sconst UE4CodeGen_Private::FBytePropertyParams %s = { %s, %s, (EPropertyFlags)0x%016llx, UE4CodeGen_Private::EPropertyGenFlags::Byte, %s, %s, %s, %s, %s };%s\r\n"),
 			Spaces,
 			Name,
 			*PropName,
-			UPropertyObjectFlags,
+			*PropNotifyFunc,
 			PropFlags,
+			UPropertyObjectFlags,
 			*ArrayDim,
-			*PropNotifyFunc,
 			OffsetStr,
 			*GetSingletonNameFuncAddr(TypedProp->Enum),
 			*MetaDataParams,
@@ -935,14 +935,14 @@
 		DeclOut.Logf(TEXT("%sstatic const UE4CodeGen_Private::FInt8PropertyParams %s;\r\n"), DeclSpaces, *NameWithoutScope);
 
 		Out.Logf(
-			TEXT("%sconst UE4CodeGen_Private::FInt8PropertyParams %s = { UE4CodeGen_Private::EPropertyClass::Int8, %s, %s, (EPropertyFlags)0x%016llx, %s, %s, %s, %s };%s\r\n"),
+			TEXT("%sconst UE4CodeGen_Private::FInt8PropertyParams %s = { %s, %s, (EPropertyFlags)0x%016llx, UE4CodeGen_Private::EPropertyGenFlags::Int8, %s, %s, %s, %s };%s\r\n"),
 			Spaces,
 			Name,
 			*PropName,
-			UPropertyObjectFlags,
+			*PropNotifyFunc,
 			PropFlags,
+			UPropertyObjectFlags,
 			*ArrayDim,
-			*PropNotifyFunc,
 			OffsetStr,
 			*MetaDataParams,
 			*PropTag
@@ -956,14 +956,14 @@
 		DeclOut.Logf(TEXT("%sstatic const UE4CodeGen_Private::FInt16PropertyParams %s;\r\n"), DeclSpaces, *NameWithoutScope);
 
 		Out.Logf(
-			TEXT("%sconst UE4CodeGen_Private::FInt16PropertyParams %s = { UE4CodeGen_Private::EPropertyClass::Int16, %s, %s, (EPropertyFlags)0x%016llx, %s, %s, %s, %s };%s\r\n"),
+			TEXT("%sconst UE4CodeGen_Private::FInt16PropertyParams %s = { %s, %s, (EPropertyFlags)0x%016llx, UE4CodeGen_Private::EPropertyGenFlags::Int16, %s, %s, %s, %s };%s\r\n"),
 			Spaces,
 			Name,
 			*PropName,
-			UPropertyObjectFlags,
+			*PropNotifyFunc,
 			PropFlags,
+			UPropertyObjectFlags,
 			*ArrayDim,
-			*PropNotifyFunc,
 			OffsetStr,
 			*MetaDataParams,
 			*PropTag
@@ -979,15 +979,15 @@
 		DeclOut.Logf(TEXT("%sstatic const UE4CodeGen_Private::%s %s;\r\n"), DeclSpaces, PropTypeName, *NameWithoutScope);
 
 		Out.Logf(
-			TEXT("%sconst UE4CodeGen_Private::%s %s = { UE4CodeGen_Private::EPropertyClass::Int, %s, %s, (EPropertyFlags)0x%016llx, %s, %s, %s, %s };%s\r\n"),
+			TEXT("%sconst UE4CodeGen_Private::%s %s = { %s, %s, (EPropertyFlags)0x%016llx, UE4CodeGen_Private::EPropertyGenFlags::Int, %s, %s, %s, %s };%s\r\n"),
 			Spaces,
 			PropTypeName,
 			Name,
 			*PropName,
-			UPropertyObjectFlags,
+			*PropNotifyFunc,
 			PropFlags,
+			UPropertyObjectFlags,
 			*ArrayDim,
-			*PropNotifyFunc,
 			OffsetStr,
 			*MetaDataParams,
 			*PropTag
@@ -1001,14 +1001,14 @@
 		DeclOut.Logf(TEXT("%sstatic const UE4CodeGen_Private::FInt64PropertyParams %s;\r\n"), DeclSpaces, *NameWithoutScope);
 
 		Out.Logf(
-			TEXT("%sconst UE4CodeGen_Private::FInt64PropertyParams %s = { UE4CodeGen_Private::EPropertyClass::Int64, %s, %s, (EPropertyFlags)0x%016llx, %s, %s, %s, %s };%s\r\n"),
+			TEXT("%sconst UE4CodeGen_Private::FInt64PropertyParams %s = { %s, %s, (EPropertyFlags)0x%016llx, UE4CodeGen_Private::EPropertyGenFlags::Int64, %s, %s, %s, %s };%s\r\n"),
 			Spaces,
 			Name,
 			*PropName,
-			UPropertyObjectFlags,
+			*PropNotifyFunc,
 			PropFlags,
+			UPropertyObjectFlags,
 			*ArrayDim,
-			*PropNotifyFunc,
 			OffsetStr,
 			*MetaDataParams,
 			*PropTag
@@ -1022,14 +1022,14 @@
 		DeclOut.Logf(TEXT("%sstatic const UE4CodeGen_Private::FUInt16PropertyParams %s;\r\n"), DeclSpaces, *NameWithoutScope);
 
 		Out.Logf(
-			TEXT("%sconst UE4CodeGen_Private::FUInt16PropertyParams %s = { UE4CodeGen_Private::EPropertyClass::UInt16, %s, %s, (EPropertyFlags)0x%016llx, %s, %s, %s, %s };%s\r\n"),
+			TEXT("%sconst UE4CodeGen_Private::FUInt16PropertyParams %s = { %s, %s, (EPropertyFlags)0x%016llx, UE4CodeGen_Private::EPropertyGenFlags::UInt16, %s, %s, %s, %s };%s\r\n"),
 			Spaces,
 			Name,
 			*PropName,
-			UPropertyObjectFlags,
+			*PropNotifyFunc,
 			PropFlags,
+			UPropertyObjectFlags,
 			*ArrayDim,
-			*PropNotifyFunc,
 			OffsetStr,
 			*MetaDataParams,
 			*PropTag
@@ -1045,15 +1045,15 @@
 		DeclOut.Logf(TEXT("%sstatic const UE4CodeGen_Private::%s %s;\r\n"), DeclSpaces, PropTypeName, *NameWithoutScope);
 
 		Out.Logf(
-			TEXT("%sconst UE4CodeGen_Private::%s %s = { UE4CodeGen_Private::EPropertyClass::UInt32, %s, %s, (EPropertyFlags)0x%016llx, %s, %s, %s, %s };%s\r\n"),
+			TEXT("%sconst UE4CodeGen_Private::%s %s = { %s, %s, (EPropertyFlags)0x%016llx, UE4CodeGen_Private::EPropertyGenFlags::UInt32, %s, %s, %s, %s };%s\r\n"),
 			Spaces,
 			PropTypeName,
 			Name,
 			*PropName,
-			UPropertyObjectFlags,
+			*PropNotifyFunc,
 			PropFlags,
+			UPropertyObjectFlags,
 			*ArrayDim,
-			*PropNotifyFunc,
 			OffsetStr,
 			*MetaDataParams,
 			*PropTag
@@ -1067,14 +1067,14 @@
 		DeclOut.Logf(TEXT("%sstatic const UE4CodeGen_Private::FUInt64PropertyParams %s;\r\n"), DeclSpaces, *NameWithoutScope);
 
 		Out.Logf(
-			TEXT("%sconst UE4CodeGen_Private::FUInt64PropertyParams %s = { UE4CodeGen_Private::EPropertyClass::UInt64, %s, %s, (EPropertyFlags)0x%016llx, %s, %s, %s, %s };%s\r\n"),
+			TEXT("%sconst UE4CodeGen_Private::FUInt64PropertyParams %s = { %s, %s, (EPropertyFlags)0x%016llx, UE4CodeGen_Private::EPropertyGenFlags::UInt64, %s, %s, %s, %s };%s\r\n"),
 			Spaces,
 			Name,
 			*PropName,
-			UPropertyObjectFlags,
+			*PropNotifyFunc,
 			PropFlags,
+			UPropertyObjectFlags,
 			*ArrayDim,
-			*PropNotifyFunc,
 			OffsetStr,
 			*MetaDataParams,
 			*PropTag
@@ -1088,14 +1088,14 @@
 		DeclOut.Logf(TEXT("%sstatic const UE4CodeGen_Private::FFloatPropertyParams %s;\r\n"), DeclSpaces, *NameWithoutScope);
 
 		Out.Logf(
-			TEXT("%sconst UE4CodeGen_Private::FFloatPropertyParams %s = { UE4CodeGen_Private::EPropertyClass::Float, %s, %s, (EPropertyFlags)0x%016llx, %s, %s, %s, %s };%s\r\n"),
+			TEXT("%sconst UE4CodeGen_Private::FFloatPropertyParams %s = { %s, %s, (EPropertyFlags)0x%016llx, UE4CodeGen_Private::EPropertyGenFlags::Float, %s, %s, %s, %s };%s\r\n"),
 			Spaces,
 			Name,
 			*PropName,
-			UPropertyObjectFlags,
+			*PropNotifyFunc,
 			PropFlags,
+			UPropertyObjectFlags,
 			*ArrayDim,
-			*PropNotifyFunc,
 			OffsetStr,
 			*MetaDataParams,
 			*PropTag
@@ -1109,14 +1109,14 @@
 		DeclOut.Logf(TEXT("%sstatic const UE4CodeGen_Private::FDoublePropertyParams %s;\r\n"), DeclSpaces, *NameWithoutScope);
 
 		Out.Logf(
-			TEXT("%sconst UE4CodeGen_Private::FDoublePropertyParams %s = { UE4CodeGen_Private::EPropertyClass::Double, %s, %s, (EPropertyFlags)0x%016llx, %s, %s, %s, %s };%s\r\n"),
+			TEXT("%sconst UE4CodeGen_Private::FDoublePropertyParams %s = { %s, %s, (EPropertyFlags)0x%016llx, UE4CodeGen_Private::EPropertyGenFlags::Double, %s, %s, %s, %s };%s\r\n"),
 			Spaces,
 			Name,
 			*PropName,
-			UPropertyObjectFlags,
+			*PropNotifyFunc,
 			PropFlags,
+			UPropertyObjectFlags,
 			*ArrayDim,
-			*PropNotifyFunc,
 			OffsetStr,
 			*MetaDataParams,
 			*PropTag
@@ -1152,16 +1152,16 @@
 		DeclOut.Logf(TEXT("%sstatic const UE4CodeGen_Private::FBoolPropertyParams %s;\r\n"), DeclSpaces, *NameWithoutScope);
 
 		Out.Logf(
-			TEXT("%sconst UE4CodeGen_Private::FBoolPropertyParams %s = { UE4CodeGen_Private::EPropertyClass::Bool, %s, %s, (EPropertyFlags)0x%016llx, %s, %s, sizeof(%s), %s, %s, %s, %s };%s\r\n"),
+			TEXT("%sconst UE4CodeGen_Private::FBoolPropertyParams %s = { %s, %s, (EPropertyFlags)0x%016llx, UE4CodeGen_Private::EPropertyGenFlags::Bool %s, %s, %s, sizeof(%s), %s, %s, %s };%s\r\n"),
 			Spaces,
 			Name,
 			*PropName,
-			UPropertyObjectFlags,
+			*PropNotifyFunc,
 			PropFlags,
+			TypedProp->IsNativeBool() ? TEXT("| UE4CodeGen_Private::EPropertyGenFlags::NativeBool") : TEXT(""),
+			UPropertyObjectFlags,
 			*ArrayDim,
-			*PropNotifyFunc,
 			*TypedProp->GetCPPType(nullptr, 0),
-			TypedProp->IsNativeBool() ? TEXT("UE4CodeGen_Private::ENativeBool::Native") : TEXT("UE4CodeGen_Private::ENativeBool::NotNative"),
 			*OuterSize,
 			*Setter,
 			*MetaDataParams,
@@ -1176,14 +1176,14 @@
 		DeclOut.Logf(TEXT("%sstatic const UE4CodeGen_Private::FSoftClassPropertyParams %s;\r\n"), DeclSpaces, *NameWithoutScope);
 
 		Out.Logf(
-			TEXT("%sconst UE4CodeGen_Private::FSoftClassPropertyParams %s = { UE4CodeGen_Private::EPropertyClass::SoftClass, %s, %s, (EPropertyFlags)0x%016llx, %s, %s, %s, %s, %s };%s\r\n"),
+			TEXT("%sconst UE4CodeGen_Private::FSoftClassPropertyParams %s = { %s, %s, (EPropertyFlags)0x%016llx, UE4CodeGen_Private::EPropertyGenFlags::SoftClass, %s, %s, %s, %s, %s };%s\r\n"),
 			Spaces,
 			Name,
 			*PropName,
-			UPropertyObjectFlags,
+			*PropNotifyFunc,
 			PropFlags,
+			UPropertyObjectFlags,
 			*ArrayDim,
-			*PropNotifyFunc,
 			OffsetStr,
 			*GetSingletonNameFuncAddr(TypedProp->MetaClass, false),
 			*MetaDataParams,
@@ -1198,14 +1198,14 @@
 		DeclOut.Logf(TEXT("%sstatic const UE4CodeGen_Private::FWeakObjectPropertyParams %s;\r\n"), DeclSpaces, *NameWithoutScope);
 
 		Out.Logf(
-			TEXT("%sconst UE4CodeGen_Private::FWeakObjectPropertyParams %s = { UE4CodeGen_Private::EPropertyClass::WeakObject, %s, %s, (EPropertyFlags)0x%016llx, %s, %s, %s, %s, %s };%s\r\n"),
+			TEXT("%sconst UE4CodeGen_Private::FWeakObjectPropertyParams %s = { %s, %s, (EPropertyFlags)0x%016llx, UE4CodeGen_Private::EPropertyGenFlags::WeakObject, %s, %s, %s, %s, %s };%s\r\n"),
 			Spaces,
 			Name,
 			*PropName,
-			UPropertyObjectFlags,
+			*PropNotifyFunc,
 			PropFlags,
+			UPropertyObjectFlags,
 			*ArrayDim,
-			*PropNotifyFunc,
 			OffsetStr,
 			*GetSingletonNameFuncAddr(TypedProp->PropertyClass, false),
 			*MetaDataParams,
@@ -1220,14 +1220,14 @@
 		DeclOut.Logf(TEXT("%sstatic const UE4CodeGen_Private::FLazyObjectPropertyParams %s;\r\n"), DeclSpaces, *NameWithoutScope);
 
 		Out.Logf(
-			TEXT("%sconst UE4CodeGen_Private::FLazyObjectPropertyParams %s = { UE4CodeGen_Private::EPropertyClass::LazyObject, %s, %s, (EPropertyFlags)0x%016llx, %s, %s, %s, %s, %s };%s\r\n"),
+			TEXT("%sconst UE4CodeGen_Private::FLazyObjectPropertyParams %s = { %s, %s, (EPropertyFlags)0x%016llx, UE4CodeGen_Private::EPropertyGenFlags::LazyObject, %s, %s, %s, %s, %s };%s\r\n"),
 			Spaces,
 			Name,
 			*PropName,
-			UPropertyObjectFlags,
+			*PropNotifyFunc,
 			PropFlags,
+			UPropertyObjectFlags,
 			*ArrayDim,
-			*PropNotifyFunc,
 			OffsetStr,
 			*GetSingletonNameFuncAddr(TypedProp->PropertyClass, false),
 			*MetaDataParams,
@@ -1242,14 +1242,14 @@
 		DeclOut.Logf(TEXT("%sstatic const UE4CodeGen_Private::FSoftObjectPropertyParams %s;\r\n"), DeclSpaces, *NameWithoutScope);
 
 		Out.Logf(
-			TEXT("%sconst UE4CodeGen_Private::FSoftObjectPropertyParams %s = { UE4CodeGen_Private::EPropertyClass::SoftObject, %s, %s, (EPropertyFlags)0x%016llx, %s, %s, %s, %s, %s };%s\r\n"),
+			TEXT("%sconst UE4CodeGen_Private::FSoftObjectPropertyParams %s = { %s, %s, (EPropertyFlags)0x%016llx, UE4CodeGen_Private::EPropertyGenFlags::SoftObject, %s, %s, %s, %s, %s };%s\r\n"),
 			Spaces,
 			Name,
 			*PropName,
-			UPropertyObjectFlags,
+			*PropNotifyFunc,
 			PropFlags,
+			UPropertyObjectFlags,
 			*ArrayDim,
-			*PropNotifyFunc,
 			OffsetStr,
 			*GetSingletonNameFuncAddr(TypedProp->PropertyClass, false),
 			*MetaDataParams,
@@ -1264,14 +1264,14 @@
 		DeclOut.Logf(TEXT("%sstatic const UE4CodeGen_Private::FClassPropertyParams %s;\r\n"), DeclSpaces, *NameWithoutScope);
 
 		Out.Logf(
-			TEXT("%sconst UE4CodeGen_Private::FClassPropertyParams %s = { UE4CodeGen_Private::EPropertyClass::Class, %s, %s, (EPropertyFlags)0x%016llx, %s, %s, %s, %s, %s, %s };%s\r\n"),
+			TEXT("%sconst UE4CodeGen_Private::FClassPropertyParams %s = { %s, %s, (EPropertyFlags)0x%016llx, UE4CodeGen_Private::EPropertyGenFlags::Class, %s, %s, %s, %s, %s, %s };%s\r\n"),
 			Spaces,
 			Name,
 			*PropName,
-			UPropertyObjectFlags,
+			*PropNotifyFunc,
 			PropFlags,
+			UPropertyObjectFlags,
 			*ArrayDim,
-			*PropNotifyFunc,
 			OffsetStr,
 			*GetSingletonNameFuncAddr(TypedProp->MetaClass, false),
 			*GetSingletonNameFuncAddr(TypedProp->PropertyClass, false),
@@ -1287,14 +1287,14 @@
 		DeclOut.Logf(TEXT("%sstatic const UE4CodeGen_Private::FObjectPropertyParams %s;\r\n"), DeclSpaces, *NameWithoutScope);
 
 		Out.Logf(
-			TEXT("%sconst UE4CodeGen_Private::FObjectPropertyParams %s = { UE4CodeGen_Private::EPropertyClass::Object, %s, %s, (EPropertyFlags)0x%016llx, %s, %s, %s, %s, %s };%s\r\n"),
+			TEXT("%sconst UE4CodeGen_Private::FObjectPropertyParams %s = { %s, %s, (EPropertyFlags)0x%016llx, UE4CodeGen_Private::EPropertyGenFlags::Object, %s, %s, %s, %s, %s };%s\r\n"),
 			Spaces,
 			Name,
 			*PropName,
-			UPropertyObjectFlags,
+			*PropNotifyFunc,
 			PropFlags,
+			UPropertyObjectFlags,
 			*ArrayDim,
-			*PropNotifyFunc,
 			OffsetStr,
 			*GetSingletonNameFuncAddr(TypedProp->PropertyClass, false),
 			*MetaDataParams,
@@ -1309,14 +1309,14 @@
 		DeclOut.Logf(TEXT("%sstatic const UE4CodeGen_Private::FInterfacePropertyParams %s;\r\n"), DeclSpaces, *NameWithoutScope);
 
 		Out.Logf(
-			TEXT("%sconst UE4CodeGen_Private::FInterfacePropertyParams %s = { UE4CodeGen_Private::EPropertyClass::Interface, %s, %s, (EPropertyFlags)0x%016llx, %s, %s, %s, %s, %s };%s\r\n"),
+			TEXT("%sconst UE4CodeGen_Private::FInterfacePropertyParams %s = { %s, %s, (EPropertyFlags)0x%016llx, UE4CodeGen_Private::EPropertyGenFlags::Interface, %s, %s, %s, %s, %s };%s\r\n"),
 			Spaces,
 			Name,
 			*PropName,
-			UPropertyObjectFlags,
+			*PropNotifyFunc,
 			PropFlags,
+			UPropertyObjectFlags,
 			*ArrayDim,
-			*PropNotifyFunc,
 			OffsetStr,
 			*GetSingletonNameFuncAddr(TypedProp->InterfaceClass, false),
 			*MetaDataParams,
@@ -1331,14 +1331,14 @@
 		DeclOut.Logf(TEXT("%sstatic const UE4CodeGen_Private::FNamePropertyParams %s;\r\n"), DeclSpaces, *NameWithoutScope);
 
 		Out.Logf(
-			TEXT("%sconst UE4CodeGen_Private::FNamePropertyParams %s = { UE4CodeGen_Private::EPropertyClass::Name, %s, %s, (EPropertyFlags)0x%016llx, %s, %s, %s, %s };%s\r\n"),
+			TEXT("%sconst UE4CodeGen_Private::FNamePropertyParams %s = { %s, %s, (EPropertyFlags)0x%016llx, UE4CodeGen_Private::EPropertyGenFlags::Name, %s, %s, %s, %s };%s\r\n"),
 			Spaces,
 			Name,
 			*PropName,
-			UPropertyObjectFlags,
+			*PropNotifyFunc,
 			PropFlags,
+			UPropertyObjectFlags,
 			*ArrayDim,
-			*PropNotifyFunc,
 			OffsetStr,
 			*MetaDataParams,
 			*PropTag
@@ -1352,14 +1352,14 @@
 		DeclOut.Logf(TEXT("%sstatic const UE4CodeGen_Private::FStrPropertyParams %s;\r\n"), DeclSpaces, *NameWithoutScope);
 
 		Out.Logf(
-			TEXT("%sconst UE4CodeGen_Private::FStrPropertyParams %s = { UE4CodeGen_Private::EPropertyClass::Str, %s, %s, (EPropertyFlags)0x%016llx, %s, %s, %s, %s };%s\r\n"),
+			TEXT("%sconst UE4CodeGen_Private::FStrPropertyParams %s = { %s, %s, (EPropertyFlags)0x%016llx, UE4CodeGen_Private::EPropertyGenFlags::Str, %s, %s, %s, %s };%s\r\n"),
 			Spaces,
 			Name,
 			*PropName,
-			UPropertyObjectFlags,
+			*PropNotifyFunc,
 			PropFlags,
+			UPropertyObjectFlags,
 			*ArrayDim,
-			*PropNotifyFunc,
 			OffsetStr,
 			*MetaDataParams,
 			*PropTag
@@ -1373,14 +1373,14 @@
 		DeclOut.Logf(TEXT("%sstatic const UE4CodeGen_Private::FArrayPropertyParams %s;\r\n"), DeclSpaces, *NameWithoutScope);
 
 		Out.Logf(
-			TEXT("%sconst UE4CodeGen_Private::FArrayPropertyParams %s = { UE4CodeGen_Private::EPropertyClass::Array, %s, %s, (EPropertyFlags)0x%016llx, %s, %s, %s, %s };%s\r\n"),
+			TEXT("%sconst UE4CodeGen_Private::FArrayPropertyParams %s = { %s, %s, (EPropertyFlags)0x%016llx, UE4CodeGen_Private::EPropertyGenFlags::Array, %s, %s, %s, %s };%s\r\n"),
 			Spaces,
 			Name,
 			*PropName,
-			UPropertyObjectFlags,
+			*PropNotifyFunc,
 			PropFlags,
+			UPropertyObjectFlags,
 			*ArrayDim,
-			*PropNotifyFunc,
 			OffsetStr,
 			*MetaDataParams,
 			*PropTag
@@ -1394,14 +1394,14 @@
 		DeclOut.Logf(TEXT("%sstatic const UE4CodeGen_Private::FMapPropertyParams %s;\r\n"), DeclSpaces, *NameWithoutScope);
 
 		Out.Logf(
-			TEXT("%sconst UE4CodeGen_Private::FMapPropertyParams %s = { UE4CodeGen_Private::EPropertyClass::Map, %s, %s, (EPropertyFlags)0x%016llx, %s, %s, %s, %s };%s\r\n"),
+			TEXT("%sconst UE4CodeGen_Private::FMapPropertyParams %s = { %s, %s, (EPropertyFlags)0x%016llx, UE4CodeGen_Private::EPropertyGenFlags::Map, %s, %s, %s, %s };%s\r\n"),
 			Spaces,
 			Name,
 			*PropName,
-			UPropertyObjectFlags,
+			*PropNotifyFunc,
 			PropFlags,
+			UPropertyObjectFlags,
 			*ArrayDim,
-			*PropNotifyFunc,
 			OffsetStr,
 			*MetaDataParams,
 			*PropTag
@@ -1415,14 +1415,14 @@
 		DeclOut.Logf(TEXT("%sstatic const UE4CodeGen_Private::FSetPropertyParams %s;\r\n"), DeclSpaces, *NameWithoutScope);
 
 		Out.Logf(
-			TEXT("%sconst UE4CodeGen_Private::FSetPropertyParams %s = { UE4CodeGen_Private::EPropertyClass::Set, %s, %s, (EPropertyFlags)0x%016llx, %s, %s, %s, %s };%s\r\n"),
+			TEXT("%sconst UE4CodeGen_Private::FSetPropertyParams %s = { %s, %s, (EPropertyFlags)0x%016llx, UE4CodeGen_Private::EPropertyGenFlags::Set, %s, %s, %s, %s };%s\r\n"),
 			Spaces,
 			Name,
 			*PropName,
-			UPropertyObjectFlags,
+			*PropNotifyFunc,
 			PropFlags,
+			UPropertyObjectFlags,
 			*ArrayDim,
-			*PropNotifyFunc,
 			OffsetStr,
 			*MetaDataParams,
 			*PropTag
@@ -1436,14 +1436,14 @@
 		DeclOut.Logf(TEXT("%sstatic const UE4CodeGen_Private::FStructPropertyParams %s;\r\n"), DeclSpaces, *NameWithoutScope);
 
 		Out.Logf(
-			TEXT("%sconst UE4CodeGen_Private::FStructPropertyParams %s = { UE4CodeGen_Private::EPropertyClass::Struct, %s, %s, (EPropertyFlags)0x%016llx, %s, %s, %s, %s, %s };%s\r\n"),
+			TEXT("%sconst UE4CodeGen_Private::FStructPropertyParams %s = { %s, %s, (EPropertyFlags)0x%016llx, UE4CodeGen_Private::EPropertyGenFlags::Struct, %s, %s, %s, %s, %s };%s\r\n"),
 			Spaces,
 			Name,
 			*PropName,
-			UPropertyObjectFlags,
+			*PropNotifyFunc,
 			PropFlags,
+			UPropertyObjectFlags,
 			*ArrayDim,
-			*PropNotifyFunc,
 			OffsetStr,
 			*GetSingletonNameFuncAddr(TypedProp->Struct),
 			*MetaDataParams,
@@ -1458,14 +1458,14 @@
 		DeclOut.Logf(TEXT("%sstatic const UE4CodeGen_Private::FDelegatePropertyParams %s;\r\n"), DeclSpaces, *NameWithoutScope);
 
 		Out.Logf(
-			TEXT("%sconst UE4CodeGen_Private::FDelegatePropertyParams %s = { UE4CodeGen_Private::EPropertyClass::Delegate, %s, %s, (EPropertyFlags)0x%016llx, %s, %s, %s, %s, %s };%s\r\n"),
+			TEXT("%sconst UE4CodeGen_Private::FDelegatePropertyParams %s = { %s, %s, (EPropertyFlags)0x%016llx, UE4CodeGen_Private::EPropertyGenFlags::Delegate, %s, %s, %s, %s, %s };%s\r\n"),
 			Spaces,
 			Name,
 			*PropName,
-			UPropertyObjectFlags,
+			*PropNotifyFunc,
 			PropFlags,
+			UPropertyObjectFlags,
 			*ArrayDim,
-			*PropNotifyFunc,
 			OffsetStr,
 			*GetSingletonNameFuncAddr(TypedProp->SignatureFunction),
 			*MetaDataParams,
@@ -1480,14 +1480,14 @@
 		DeclOut.Logf(TEXT("%sstatic const UE4CodeGen_Private::FMulticastDelegatePropertyParams %s;\r\n"), DeclSpaces, *NameWithoutScope);
 
 		Out.Logf(
-			TEXT("%sconst UE4CodeGen_Private::FMulticastDelegatePropertyParams %s = { UE4CodeGen_Private::EPropertyClass::MulticastDelegate, %s, %s, (EPropertyFlags)0x%016llx, %s, %s, %s, %s, %s };%s\r\n"),
+			TEXT("%sconst UE4CodeGen_Private::FMulticastDelegatePropertyParams %s = { %s, %s, (EPropertyFlags)0x%016llx, UE4CodeGen_Private::EPropertyGenFlags::MulticastDelegate, %s, %s, %s, %s, %s };%s\r\n"),
 			Spaces,
 			Name,
 			*PropName,
-			UPropertyObjectFlags,
+			*PropNotifyFunc,
 			PropFlags,
+			UPropertyObjectFlags,
 			*ArrayDim,
-			*PropNotifyFunc,
 			OffsetStr,
 			*GetSingletonNameFuncAddr(TypedProp->SignatureFunction),
 			*MetaDataParams,
@@ -1502,14 +1502,14 @@
 		DeclOut.Logf(TEXT("%sstatic const UE4CodeGen_Private::FTextPropertyParams %s;\r\n"), DeclSpaces, *NameWithoutScope);
 
 		Out.Logf(
-			TEXT("%sconst UE4CodeGen_Private::FTextPropertyParams %s = { UE4CodeGen_Private::EPropertyClass::Text, %s, %s, (EPropertyFlags)0x%016llx, %s, %s, %s, %s };%s\r\n"),
+			TEXT("%sconst UE4CodeGen_Private::FTextPropertyParams %s = { %s, %s, (EPropertyFlags)0x%016llx, UE4CodeGen_Private::EPropertyGenFlags::Text, %s, %s, %s, %s };%s\r\n"),
 			Spaces,
 			Name,
 			*PropName,
-			UPropertyObjectFlags,
+			*PropNotifyFunc,
 			PropFlags,
+			UPropertyObjectFlags,
 			*ArrayDim,
-			*PropNotifyFunc,
 			OffsetStr,
 			*MetaDataParams,
 			*PropTag
@@ -1523,14 +1523,14 @@
 		DeclOut.Logf(TEXT("%sstatic const UE4CodeGen_Private::FEnumPropertyParams %s;\r\n"), DeclSpaces, *NameWithoutScope);
 
 		Out.Logf(
-			TEXT("%sconst UE4CodeGen_Private::FEnumPropertyParams %s = { UE4CodeGen_Private::EPropertyClass::Enum, %s, %s, (EPropertyFlags)0x%016llx, %s, %s, %s, %s, %s };%s\r\n"),
+			TEXT("%sconst UE4CodeGen_Private::FEnumPropertyParams %s = { %s, %s, (EPropertyFlags)0x%016llx, UE4CodeGen_Private::EPropertyGenFlags::Enum, %s, %s, %s, %s, %s };%s\r\n"),
 			Spaces,
 			Name,
 			*PropName,
-			UPropertyObjectFlags,
+			*PropNotifyFunc,
 			PropFlags,
+			UPropertyObjectFlags,
 			*ArrayDim,
-			*PropNotifyFunc,
 			OffsetStr,
 			*GetSingletonNameFuncAddr(TypedProp->Enum),
 			*MetaDataParams,
@@ -1559,11 +1559,11 @@
 	return false;
 }
 
-FString FNativeClassHeaderGenerator::OutputProperties(FOutputDevice& DeclOut, FOutputDevice& Out, const TCHAR* Scope, const TArray<UProperty*>& Properties, const TCHAR* DeclSpaces, const TCHAR* Spaces)
+TTuple<FString, FString> FNativeClassHeaderGenerator::OutputProperties(FOutputDevice& DeclOut, FOutputDevice& Out, const TCHAR* Scope, const TArray<UProperty*>& Properties, const TCHAR* DeclSpaces, const TCHAR* Spaces)
 {
 	if (Properties.Num() == 0)
 	{
-		return TEXT("nullptr, 0");
+		return TTuple<FString, FString>(TEXT("nullptr"), TEXT("0"));
 	}
 
 	TArray<FPropertyNamePointerPair> PropertyNamesAndPointers;
@@ -1609,11 +1609,17 @@
 
 	if (bHasAllEditorOnlyDataProperties)
 	{
-		return FString::Printf(TEXT("IF_WITH_EDITORONLY_DATA(%sPropPointers, nullptr), IF_WITH_EDITORONLY_DATA(ARRAY_COUNT(%sPropPointers), 0)"), Scope, Scope);
+		return TTuple<FString, FString>(
+			FString::Printf(TEXT("IF_WITH_EDITORONLY_DATA(%sPropPointers, nullptr)"), Scope),
+			FString::Printf(TEXT("IF_WITH_EDITORONLY_DATA(ARRAY_COUNT(%sPropPointers), 0)"), Scope)
+		);
 	}
 	else
 	{
-		return FString::Printf(TEXT("%sPropPointers, ARRAY_COUNT(%sPropPointers)"), Scope, Scope);
+		return TTuple<FString, FString>(
+			FString::Printf(TEXT("%sPropPointers"), Scope),
+			FString::Printf(TEXT("ARRAY_COUNT(%sPropPointers)"), Scope)
+		);
 	}
 }
 
@@ -1842,8 +1848,8 @@
 		for (UField* ScriptType : *SingletonsToOutput)
 		{
 			Out.Log(FTypeSingletonCache::Get(ScriptType, true).GetExternDecl());
+			}
 		}
-	}
 
 	FOutputDeviceNull OutputDeviceNull;
 	FString MetaDataParams = OutputMetaDataCodeForObject(OutputDeviceNull, Out, InPackage, TEXT("Package_MetaDataParams"), TEXT(""), TEXT("\t\t\t"));
@@ -1854,7 +1860,8 @@
 	Out.Logf(TEXT("\t\tif (!ReturnPackage)\r\n"));
 	Out.Logf(TEXT("\t\t{\r\n"));
 
-	FString SingletonRange;
+	const TCHAR* SingletonArray;
+	const TCHAR* SingletonCount;
 	if (SingletonsToOutput)
 	{
 		Out.Logf(TEXT("\t\t\tstatic UObject* (*const SingletonFuncArray[])() = {\r\n"));
@@ -1866,19 +1873,22 @@
 		}
 		Out.Logf(TEXT("\t\t\t};\r\n"));
 
-		SingletonRange = TEXT("SingletonFuncArray, ARRAY_COUNT(SingletonFuncArray)");
+		SingletonArray = TEXT("SingletonFuncArray");
+		SingletonCount = TEXT("ARRAY_COUNT(SingletonFuncArray)");
 	}
 	else
 	{
-		SingletonRange = TEXT("nullptr, 0");
+		SingletonArray = TEXT("nullptr");
+		SingletonCount = TEXT("0");
 	}
 
 	Out.Logf(TEXT("\t\t\tstatic const UE4CodeGen_Private::FPackageParams PackageParams = {\r\n"));
 	Out.Logf(TEXT("\t\t\t\t%s,\r\n"), *CreateUTF8LiteralString(InPackage->GetName()));
+	Out.Logf(TEXT("\t\t\t\t%s,\r\n"), SingletonArray);
+	Out.Logf(TEXT("\t\t\t\t%s,\r\n"), SingletonCount);
 	Out.Logf(TEXT("\t\t\t\tPKG_CompiledIn | 0x%08X,\r\n"), InPackage->GetPackageFlags() & (PKG_ClientOptional | PKG_ServerSideOnly | PKG_EditorOnly | PKG_Developer));
 	Out.Logf(TEXT("\t\t\t\t0x%08X,\r\n"), Hash);
 	Out.Logf(TEXT("\t\t\t\t0x%08X,\r\n"), GenerateTextHash(InDeclarations));
-	Out.Logf(TEXT("\t\t\t\t%s,\r\n"), *SingletonRange);
 	Out.Logf(TEXT("\t\t\t\t%s\r\n"), *MetaDataParams);
 	Out.Logf(TEXT("\t\t\t};\r\n"));
 	Out.Logf(TEXT("\t\t\tUE4CodeGen_Private::ConstructUPackage(ReturnPackage, PackageParams);\r\n"));
@@ -1969,7 +1979,8 @@
 			Singletons.Logf(TEXT("\t\t(UObject* (*)())%s,\r\n"), *PackageSingletonName.LeftChop(2));
 		}
 
-		const TCHAR* SingletonsRange;
+		const TCHAR* SingletonsArray;
+		const TCHAR* SingletonsCount;
 		if (Singletons.Len() != 0)
 		{
 			GeneratedClassRegisterFunctionText.Logf(TEXT("\t\tstatic UObject* (*const DependentSingletons[])();\r\n"));
@@ -1978,13 +1989,17 @@
 			StaticDefinitions.Log (*Singletons);
 			StaticDefinitions.Logf(TEXT("\t};\r\n"));
 
-			SingletonsRange = TEXT("DependentSingletons, ARRAY_COUNT(DependentSingletons)");
+			SingletonsArray = TEXT("DependentSingletons");
+			SingletonsCount = TEXT("ARRAY_COUNT(DependentSingletons)");
 		}
 		else
 		{
-			SingletonsRange = TEXT("nullptr, 0");
+			SingletonsArray = TEXT("nullptr");
+			SingletonsCount = TEXT("0");
 		}
 
+		const TCHAR* FunctionsArray;
+		const TCHAR* FunctionsCount;
 		if (FunctionsToExport.Num() != 0)
 		{
 			GeneratedClassRegisterFunctionText.Log(BEGIN_WRAP_EDITOR_ONLY(bAllEditorOnlyFunctions));
@@ -2021,6 +2036,22 @@
 
 			StaticDefinitions.Log(TEXT("\t};\r\n"));
 			StaticDefinitions.Log(END_WRAP_EDITOR_ONLY(bAllEditorOnlyFunctions));
+
+			if (bAllEditorOnlyFunctions)
+			{
+				FunctionsArray = TEXT("IF_WITH_EDITOR(FuncInfo, nullptr)");
+				FunctionsCount = TEXT("IF_WITH_EDITOR(ARRAY_COUNT(FuncInfo), 0)");
+			}
+			else
+			{
+				FunctionsArray = TEXT("FuncInfo");
+				FunctionsCount = TEXT("ARRAY_COUNT(FuncInfo)");
+			}
+		}
+		else
+		{
+			FunctionsArray = TEXT("nullptr");
+			FunctionsCount = TEXT("0");
 		}
 
 		TMap<FName, FString>* MetaDataMap = UMetaData::GetMapForObject(Class);
@@ -2038,9 +2069,10 @@
 		TArray<UProperty*> Props;
 		Algo::Copy(TFieldRange<UProperty>(Class, EFieldIteratorFlags::ExcludeSuper), Props);
 
-		FString PropertyRange = OutputProperties(GeneratedClassRegisterFunctionText, StaticDefinitions, *FString::Printf(TEXT("%s::"), *StaticsStructName), Props, TEXT("\t\t"), TEXT("\t"));
+		TTuple<FString, FString> PropertyRange = OutputProperties(GeneratedClassRegisterFunctionText, StaticDefinitions, *FString::Printf(TEXT("%s::"), *StaticsStructName), Props, TEXT("\t\t"), TEXT("\t"));
 
-		const TCHAR* InterfaceRange;
+		const TCHAR* InterfaceArray;
+		const TCHAR* InterfaceCount;
 		if (Class->Interfaces.Num() > 0)
 		{
 			GeneratedClassRegisterFunctionText.Log(TEXT("\t\tstatic const UE4CodeGen_Private::FImplementedInterfaceParams InterfaceParams[];\r\n"));
@@ -2067,11 +2099,13 @@
 			}
 			StaticDefinitions.Log(TEXT("\t\t};\r\n"));
 
-			InterfaceRange = TEXT("InterfaceParams, ARRAY_COUNT(InterfaceParams)");
+			InterfaceArray = TEXT("InterfaceParams");
+			InterfaceCount = TEXT("ARRAY_COUNT(InterfaceParams)");
 		}
 		else
 		{
-			InterfaceRange = TEXT("nullptr, 0");
+			InterfaceArray = TEXT("nullptr");
+			InterfaceCount = TEXT("0");
 		}
 
 		GeneratedClassRegisterFunctionText.Logf(TEXT("\t\tstatic const FCppClassTypeInfoStatic StaticCppClassTypeInfo;\r\n"));
@@ -2091,13 +2125,17 @@
 
 		StaticDefinitions.Logf(TEXT("\tconst UE4CodeGen_Private::FClassParams %s::ClassParams = {\r\n"), *StaticsStructName);
 		StaticDefinitions.Logf(TEXT("\t\t&%s::StaticClass,\r\n"), ClassNameCPP);
-		StaticDefinitions.Logf(TEXT("\t\t%s,\r\n"), SingletonsRange);
-		StaticDefinitions.Logf(TEXT("\t\t0x%08Xu,\r\n"), ClassFlags);
-		StaticDefinitions.Logf(TEXT("\t\t%s,\r\n"), (FunctionsToExport.Num() != 0) ? (bAllEditorOnlyFunctions ? TEXT("IF_WITH_EDITOR(FuncInfo, nullptr), IF_WITH_EDITOR(ARRAY_COUNT(FuncInfo), 0)") : TEXT("FuncInfo, ARRAY_COUNT(FuncInfo)")) : TEXT("nullptr, 0"));
-		StaticDefinitions.Logf(TEXT("\t\t%s,\r\n"), *PropertyRange);
 		StaticDefinitions.Logf(TEXT("\t\t%s,\r\n"), (Class->ClassConfigName != NAME_None) ? *CreateUTF8LiteralString(Class->ClassConfigName.ToString()) : TEXT("nullptr"));
 		StaticDefinitions.Log (TEXT("\t\t&StaticCppClassTypeInfo,\r\n"));
-		StaticDefinitions.Logf(TEXT("\t\t%s,\r\n"), InterfaceRange);
+		StaticDefinitions.Logf(TEXT("\t\t%s,\r\n"), SingletonsArray);
+		StaticDefinitions.Logf(TEXT("\t\t%s,\r\n"), FunctionsArray);
+		StaticDefinitions.Logf(TEXT("\t\t%s,\r\n"), *PropertyRange.Get<0>());
+		StaticDefinitions.Logf(TEXT("\t\t%s,\r\n"), InterfaceArray);
+		StaticDefinitions.Logf(TEXT("\t\t%s,\r\n"), SingletonsCount);
+		StaticDefinitions.Logf(TEXT("\t\t%s,\r\n"), FunctionsCount);
+		StaticDefinitions.Logf(TEXT("\t\t%s,\r\n"), *PropertyRange.Get<1>());
+		StaticDefinitions.Logf(TEXT("\t\t%s,\r\n"), InterfaceCount);
+		StaticDefinitions.Logf(TEXT("\t\t0x%08Xu,\r\n"), ClassFlags);
 		StaticDefinitions.Logf(TEXT("\t\t%s\r\n"), *MetaDataParams);
 		StaticDefinitions.Log (TEXT("\t};\r\n"));
 
@@ -2269,7 +2307,7 @@
 	const TCHAR* UFunctionType = bIsDelegate ? TEXT("UDelegateFunction") : TEXT("UFunction");
 	const TCHAR* UFunctionObjectFlags = FClass::IsOwnedByDynamicType(Function) ? TEXT("RF_Public|RF_Transient") : TEXT("RF_Public|RF_Transient|RF_MarkAsNative");
 
-	FString PropertyRange = OutputProperties(CurrentFunctionText, StaticDefinitions, *FString::Printf(TEXT("%s::"), *StaticsStructName), Props, TEXT("\t\t"), TEXT("\t"));
+	TTuple<FString, FString> PropertyRange = OutputProperties(CurrentFunctionText, StaticDefinitions, *FString::Printf(TEXT("%s::"), *StaticsStructName), Props, TEXT("\t\t"), TEXT("\t"));
 
 	const FFunctionData* CompilerInfo = FFunctionData::FindForFunction(Function);
 	const FFuncInfo&     FunctionData = CompilerInfo->GetFunctionData();
@@ -2280,15 +2318,16 @@
 	CurrentFunctionText.Log(TEXT("\t\tstatic const UE4CodeGen_Private::FFunctionParams FuncParams;\r\n"));
 
 	StaticDefinitions.Logf(
-		TEXT("\tconst UE4CodeGen_Private::FFunctionParams %s::FuncParams = { (UObject*(*)())%s, %s, %s, %s, (EFunctionFlags)0x%08X, %s, %s, %d, %d, %s };\r\n"),
+		TEXT("\tconst UE4CodeGen_Private::FFunctionParams %s::FuncParams = { (UObject*(*)())%s, %s, %s, %s, %s, %s, %s, (EFunctionFlags)0x%08X, %d, %d, %s };\r\n"),
 		*StaticsStructName,
 		*OuterFunc,
+		*GetSingletonNameFuncAddr(SuperFunction),
 		*CreateUTF8LiteralString(FNativeClassHeaderGenerator::GetOverriddenName(Function)),
+		*StructureSize,
+		*PropertyRange.Get<0>(),
+		*PropertyRange.Get<1>(),
 		UFunctionObjectFlags,
-		*GetSingletonNameFuncAddr(SuperFunction),
 		(uint32)Function->FunctionFlags,
-		*StructureSize,
-		*PropertyRange,
 		bIsNet ? FunctionData.RPCId : 0,
 		bIsNet ? FunctionData.RPCResponseId : 0,
 		*MetaDataParams
@@ -3410,7 +3449,7 @@
 		NewStructOps = TEXT("nullptr");
 	}
 
-	FString PropertyRange = OutputProperties(GeneratedStructRegisterFunctionText, StaticDefinitions, *FString::Printf(TEXT("%s::"), *StaticsStructName), Props, TEXT("\t\t"), TEXT("\t"));
+	TTuple<FString, FString> PropertyRange = OutputProperties(GeneratedStructRegisterFunctionText, StaticDefinitions, *FString::Printf(TEXT("%s::"), *StaticsStructName), Props, TEXT("\t\t"), TEXT("\t"));
 
 	GeneratedStructRegisterFunctionText.Log (TEXT("\t\tstatic const UE4CodeGen_Private::FStructParams ReturnStructParams;\r\n"));
 
@@ -3419,11 +3458,12 @@
 	StaticDefinitions.Logf(TEXT("\t\t%s,\r\n"), *GetSingletonNameFuncAddr(BaseStruct));
 	StaticDefinitions.Logf(TEXT("\t\t%s,\r\n"), *NewStructOps);
 	StaticDefinitions.Logf(TEXT("\t\t%s,\r\n"), *CreateUTF8LiteralString(ActualStructName));
-	StaticDefinitions.Logf(TEXT("\t\t%s,\r\n"), bIsDynamic ? TEXT("RF_Public|RF_Transient") : TEXT("RF_Public|RF_Transient|RF_MarkAsNative"));
-	StaticDefinitions.Logf(TEXT("\t\tEStructFlags(0x%08X),\r\n"), (uint32)UncomputedFlags);
 	StaticDefinitions.Logf(TEXT("\t\tsizeof(%s),\r\n"), StructNameCPP);
 	StaticDefinitions.Logf(TEXT("\t\talignof(%s),\r\n"), StructNameCPP);
-	StaticDefinitions.Logf(TEXT("\t\t%s,\r\n"), *PropertyRange);
+	StaticDefinitions.Logf(TEXT("\t\t%s,\r\n"), *PropertyRange.Get<0>());
+	StaticDefinitions.Logf(TEXT("\t\t%s,\r\n"), *PropertyRange.Get<1>());
+	StaticDefinitions.Logf(TEXT("\t\t%s,\r\n"), bIsDynamic ? TEXT("RF_Public|RF_Transient") : TEXT("RF_Public|RF_Transient|RF_MarkAsNative"));
+	StaticDefinitions.Logf(TEXT("\t\tEStructFlags(0x%08X),\r\n"), (uint32)UncomputedFlags);
 	StaticDefinitions.Logf(TEXT("\t\t%s\r\n"), *MetaDataParams);
 	StaticDefinitions.Log (TEXT("\t};\r\n"));
 
@@ -3597,14 +3637,14 @@
 
 	GeneratedEnumRegisterFunctionText.Logf(TEXT("\t\t\tstatic const UE4CodeGen_Private::FEnumParams EnumParams = {\r\n"));
 	GeneratedEnumRegisterFunctionText.Logf(TEXT("\t\t\t\t(UObject*(*)())%s,\r\n"), *OuterString.LeftChop(2));
-	GeneratedEnumRegisterFunctionText.Logf(TEXT("\t\t\t\tUE4CodeGen_Private::EDynamicType::%s,\r\n"), bIsDynamic ? TEXT("Dynamic") : TEXT("NotDynamic"));
-	GeneratedEnumRegisterFunctionText.Logf(TEXT("\t\t\t\t%s,\r\n"), *CreateUTF8LiteralString(OverriddenEnumNameCpp));
-	GeneratedEnumRegisterFunctionText.Logf(TEXT("\t\t\t\t%s,\r\n"), UEnumObjectFlags);
 	GeneratedEnumRegisterFunctionText.Logf(TEXT("\t\t\t\t%s,\r\n"), EnumDisplayNameFn.IsEmpty() ? TEXT("nullptr") : *EnumDisplayNameFn);
-	GeneratedEnumRegisterFunctionText.Logf(TEXT("\t\t\t\t(uint8)%s,\r\n"), *EnumFormStr);
+	GeneratedEnumRegisterFunctionText.Logf(TEXT("\t\t\t\t%s,\r\n"), *CreateUTF8LiteralString(OverriddenEnumNameCpp));
 	GeneratedEnumRegisterFunctionText.Logf(TEXT("\t\t\t\t%s,\r\n"), *CreateUTF8LiteralString(Enum->CppType));
 	GeneratedEnumRegisterFunctionText.Logf(TEXT("\t\t\t\tEnumerators,\r\n"));
 	GeneratedEnumRegisterFunctionText.Logf(TEXT("\t\t\t\tARRAY_COUNT(Enumerators),\r\n"));
+	GeneratedEnumRegisterFunctionText.Logf(TEXT("\t\t\t\t%s,\r\n"), UEnumObjectFlags);
+	GeneratedEnumRegisterFunctionText.Logf(TEXT("\t\t\t\tUE4CodeGen_Private::EDynamicType::%s,\r\n"), bIsDynamic ? TEXT("Dynamic") : TEXT("NotDynamic"));
+	GeneratedEnumRegisterFunctionText.Logf(TEXT("\t\t\t\t(uint8)%s,\r\n"), *EnumFormStr);
 	GeneratedEnumRegisterFunctionText.Logf(TEXT("\t\t\t\t%s\r\n"), *MetaDataParams);
 	GeneratedEnumRegisterFunctionText.Logf(TEXT("\t\t\t};\r\n"));
 	GeneratedEnumRegisterFunctionText.Logf(TEXT("\t\t\tUE4CodeGen_Private::ConstructUEnum(ReturnEnum, EnumParams);\r\n"));
diff --git a/Engine/Source/Programs/UnrealHeaderTool/Private/NativeClassExporter.h b/Engine/Source/Programs/UnrealHeaderTool/Private/NativeClassExporter.h
--- a/Engine/Source/Programs/UnrealHeaderTool/Private/NativeClassExporter.h
+++ b/Engine/Source/Programs/UnrealHeaderTool/Private/NativeClassExporter.h
@@ -364,9 +364,9 @@ private:
 	 * @param	Spaces			String of spaces to use as an indent for the declaration
 	 * @param	Spaces			String of spaces to use as an indent
 	 *
-	 * @return	The string which represents a range of the emitted properties.
+	 * @return      A pair of strings which represents the pointer and a count of the emitted properties.
 	 */
-	FString OutputProperties(FOutputDevice& DeclOut, FOutputDevice& Out, const TCHAR* Scope, const TArray<UProperty*>& Properties, const TCHAR* DeclSpaces, const TCHAR* Spaces);
+	TTuple<FString, FString> OutputProperties(FOutputDevice& DeclOut, FOutputDevice& Out, const TCHAR* Scope, const TArray<UProperty*>& Properties, const TCHAR* DeclSpaces, const TCHAR* Spaces);
 
 	/**
 	 * Function to output the C++ code necessary to set up a property
diff --git a/Engine/Source/Runtime/CoreUObject/Private/UObject/UObjectGlobals.cpp b/Engine/Source/Runtime/CoreUObject/Private/UObject/UObjectGlobals.cpp
--- a/Engine/Source/Runtime/CoreUObject/Private/UObject/UObjectGlobals.cpp
+++ b/Engine/Source/Runtime/CoreUObject/Private/UObject/UObjectGlobals.cpp
@@ -4019,7 +4019,7 @@ namespace UE4CodeGen_Private
 #endif
 
 		UProperty* NewProp = nullptr;
-		switch (PropBase->Type)
+		switch (PropBase->Flags & PropertyTypeMask)
 		{
 			default:
 			{
@@ -4027,7 +4027,7 @@ namespace UE4CodeGen_Private
 				check(false);
 			}
 
-			case EPropertyClass::Byte:
+			case EPropertyGenFlags::Byte:
 			{
 				const FBytePropertyParams* Prop = (const FBytePropertyParams*)PropBase;
 				NewProp = new (EC_InternalUseOnlyConstructor, Outer, UTF8_TO_TCHAR(Prop->NameUTF8), Prop->ObjectFlags) UByteProperty(FObjectInitializer(), EC_CppProperty, Prop->Offset, Prop->PropertyFlags, Prop->EnumFunc ? Prop->EnumFunc() : nullptr);
@@ -4039,7 +4039,7 @@ namespace UE4CodeGen_Private
 			}
 			break;
 
-			case EPropertyClass::Int8:
+			case EPropertyGenFlags::Int8:
 			{
 				const FInt8PropertyParams* Prop = (const FInt8PropertyParams*)PropBase;
 				NewProp = new (EC_InternalUseOnlyConstructor, Outer, UTF8_TO_TCHAR(Prop->NameUTF8), Prop->ObjectFlags) UInt8Property(FObjectInitializer(), EC_CppProperty, Prop->Offset, Prop->PropertyFlags);
@@ -4051,7 +4051,7 @@ namespace UE4CodeGen_Private
 			}
 			break;
 
-			case EPropertyClass::Int16:
+			case EPropertyGenFlags::Int16:
 			{
 				const FInt16PropertyParams* Prop = (const FInt16PropertyParams*)PropBase;
 				NewProp = new (EC_InternalUseOnlyConstructor, Outer, UTF8_TO_TCHAR(Prop->NameUTF8), Prop->ObjectFlags) UInt16Property(FObjectInitializer(), EC_CppProperty, Prop->Offset, Prop->PropertyFlags);
@@ -4063,7 +4063,7 @@ namespace UE4CodeGen_Private
 			}
 			break;
 
-			case EPropertyClass::Int:
+			case EPropertyGenFlags::Int:
 			{
 				const FIntPropertyParams* Prop = (const FIntPropertyParams*)PropBase;
 				NewProp = new (EC_InternalUseOnlyConstructor, Outer, UTF8_TO_TCHAR(Prop->NameUTF8), Prop->ObjectFlags) UIntProperty(FObjectInitializer(), EC_CppProperty, Prop->Offset, Prop->PropertyFlags);
@@ -4075,7 +4075,7 @@ namespace UE4CodeGen_Private
 			}
 			break;
 
-			case EPropertyClass::Int64:
+			case EPropertyGenFlags::Int64:
 			{
 				const FInt64PropertyParams* Prop = (const FInt64PropertyParams*)PropBase;
 				NewProp = new (EC_InternalUseOnlyConstructor, Outer, UTF8_TO_TCHAR(Prop->NameUTF8), Prop->ObjectFlags) UInt64Property(FObjectInitializer(), EC_CppProperty, Prop->Offset, Prop->PropertyFlags);
@@ -4087,7 +4087,7 @@ namespace UE4CodeGen_Private
 			}
 			break;
 
-			case EPropertyClass::UInt16:
+			case EPropertyGenFlags::UInt16:
 			{
 				const FUInt16PropertyParams* Prop = (const FUInt16PropertyParams*)PropBase;
 				NewProp = new (EC_InternalUseOnlyConstructor, Outer, UTF8_TO_TCHAR(Prop->NameUTF8), Prop->ObjectFlags) UUInt16Property(FObjectInitializer(), EC_CppProperty, Prop->Offset, Prop->PropertyFlags);
@@ -4099,7 +4099,7 @@ namespace UE4CodeGen_Private
 			}
 			break;
 
-			case EPropertyClass::UInt32:
+			case EPropertyGenFlags::UInt32:
 			{
 				const FUInt32PropertyParams* Prop = (const FUInt32PropertyParams*)PropBase;
 				NewProp = new (EC_InternalUseOnlyConstructor, Outer, UTF8_TO_TCHAR(Prop->NameUTF8), Prop->ObjectFlags) UUInt32Property(FObjectInitializer(), EC_CppProperty, Prop->Offset, Prop->PropertyFlags);
@@ -4111,7 +4111,7 @@ namespace UE4CodeGen_Private
 			}
 			break;
 
-			case EPropertyClass::UInt64:
+			case EPropertyGenFlags::UInt64:
 			{
 				const FUInt64PropertyParams* Prop = (const FUInt64PropertyParams*)PropBase;
 				NewProp = new (EC_InternalUseOnlyConstructor, Outer, UTF8_TO_TCHAR(Prop->NameUTF8), Prop->ObjectFlags) UUInt64Property(FObjectInitializer(), EC_CppProperty, Prop->Offset, Prop->PropertyFlags);
@@ -4123,7 +4123,7 @@ namespace UE4CodeGen_Private
 			}
 			break;
 
-			case EPropertyClass::UnsizedInt:
+			case EPropertyGenFlags::UnsizedInt:
 			{
 				const FUnsizedIntPropertyParams* Prop = (const FUnsizedIntPropertyParams*)PropBase;
 				NewProp = new (EC_InternalUseOnlyConstructor, Outer, UTF8_TO_TCHAR(Prop->NameUTF8), Prop->ObjectFlags) UUInt64Property(FObjectInitializer(), EC_CppProperty, Prop->Offset, Prop->PropertyFlags);
@@ -4135,7 +4135,7 @@ namespace UE4CodeGen_Private
 			}
 			break;
 
-			case EPropertyClass::UnsizedUInt:
+			case EPropertyGenFlags::UnsizedUInt:
 			{
 				const FUnsizedUIntPropertyParams* Prop = (const FUnsizedUIntPropertyParams*)PropBase;
 				NewProp = new (EC_InternalUseOnlyConstructor, Outer, UTF8_TO_TCHAR(Prop->NameUTF8), Prop->ObjectFlags) UUInt64Property(FObjectInitializer(), EC_CppProperty, Prop->Offset, Prop->PropertyFlags);
@@ -4147,7 +4147,7 @@ namespace UE4CodeGen_Private
 			}
 			break;
 
-			case EPropertyClass::Float:
+			case EPropertyGenFlags::Float:
 			{
 				const FFloatPropertyParams* Prop = (const FFloatPropertyParams*)PropBase;
 				NewProp = new (EC_InternalUseOnlyConstructor, Outer, UTF8_TO_TCHAR(Prop->NameUTF8), Prop->ObjectFlags) UFloatProperty(FObjectInitializer(), EC_CppProperty, Prop->Offset, Prop->PropertyFlags);
@@ -4159,7 +4159,7 @@ namespace UE4CodeGen_Private
 			}
 			break;
 
-			case EPropertyClass::Double:
+			case EPropertyGenFlags::Double:
 			{
 				const FDoublePropertyParams* Prop = (const FDoublePropertyParams*)PropBase;
 				NewProp = new (EC_InternalUseOnlyConstructor, Outer, UTF8_TO_TCHAR(Prop->NameUTF8), Prop->ObjectFlags) UDoubleProperty(FObjectInitializer(), EC_CppProperty, Prop->Offset, Prop->PropertyFlags);
@@ -4171,7 +4171,7 @@ namespace UE4CodeGen_Private
 			}
 			break;
 
-			case EPropertyClass::Bool:
+			case EPropertyGenFlags::Bool:
 			{
 				auto DoDetermineBitfieldOffsetAndMask = [](uint32& Offset, uint32& BitMask, void (*SetBit)(void* Obj), const SIZE_T SizeOf)
 				{
@@ -4203,7 +4203,7 @@ namespace UE4CodeGen_Private
 					check(BitMask);
 				}
 
-				NewProp = new (EC_InternalUseOnlyConstructor, Outer, UTF8_TO_TCHAR(Prop->NameUTF8), Prop->ObjectFlags) UBoolProperty(FObjectInitializer(), EC_CppProperty, Offset, Prop->PropertyFlags, BitMask, Prop->ElementSize, Prop->NativeBool == ENativeBool::Native);
+				NewProp = new (EC_InternalUseOnlyConstructor, Outer, UTF8_TO_TCHAR(Prop->NameUTF8), Prop->ObjectFlags) UBoolProperty(FObjectInitializer(), EC_CppProperty, Offset, Prop->PropertyFlags, BitMask, Prop->ElementSize, !!(Prop->Flags & EPropertyGenFlags::NativeBool));
 
 #if WITH_METADATA
 				MetaDataArray = Prop->MetaDataArray;
@@ -4212,7 +4212,7 @@ namespace UE4CodeGen_Private
 			}
 			break;
 
-			case EPropertyClass::Object:
+			case EPropertyGenFlags::Object:
 			{
 				const FObjectPropertyParams* Prop = (const FObjectPropertyParams*)PropBase;
 				NewProp = new (EC_InternalUseOnlyConstructor, Outer, UTF8_TO_TCHAR(Prop->NameUTF8), Prop->ObjectFlags) UObjectProperty(FObjectInitializer(), EC_CppProperty, Prop->Offset, Prop->PropertyFlags, Prop->ClassFunc ? Prop->ClassFunc() : nullptr);
@@ -4224,7 +4224,7 @@ namespace UE4CodeGen_Private
 			}
 			break;
 
-			case EPropertyClass::WeakObject:
+			case EPropertyGenFlags::WeakObject:
 			{
 				const FWeakObjectPropertyParams* Prop = (const FWeakObjectPropertyParams*)PropBase;
 				NewProp = new (EC_InternalUseOnlyConstructor, Outer, UTF8_TO_TCHAR(Prop->NameUTF8), Prop->ObjectFlags) UWeakObjectProperty(FObjectInitializer(), EC_CppProperty, Prop->Offset, Prop->PropertyFlags, Prop->ClassFunc ? Prop->ClassFunc() : nullptr);
@@ -4236,7 +4236,7 @@ namespace UE4CodeGen_Private
 			}
 			break;
 
-			case EPropertyClass::LazyObject:
+			case EPropertyGenFlags::LazyObject:
 			{
 				const FLazyObjectPropertyParams* Prop = (const FLazyObjectPropertyParams*)PropBase;
 				NewProp = new (EC_InternalUseOnlyConstructor, Outer, UTF8_TO_TCHAR(Prop->NameUTF8), Prop->ObjectFlags) ULazyObjectProperty(FObjectInitializer(), EC_CppProperty, Prop->Offset, Prop->PropertyFlags, Prop->ClassFunc ? Prop->ClassFunc() : nullptr);
@@ -4248,7 +4248,7 @@ namespace UE4CodeGen_Private
 			}
 			break;
 
-			case EPropertyClass::SoftObject:
+			case EPropertyGenFlags::SoftObject:
 			{
 				const FSoftObjectPropertyParams* Prop = (const FSoftObjectPropertyParams*)PropBase;
 				NewProp = new (EC_InternalUseOnlyConstructor, Outer, UTF8_TO_TCHAR(Prop->NameUTF8), Prop->ObjectFlags) USoftObjectProperty(FObjectInitializer(), EC_CppProperty, Prop->Offset, Prop->PropertyFlags, Prop->ClassFunc ? Prop->ClassFunc() : nullptr);
@@ -4260,7 +4260,7 @@ namespace UE4CodeGen_Private
 			}
 			break;
 
-			case EPropertyClass::Class:
+			case EPropertyGenFlags::Class:
 			{
 				const FClassPropertyParams* Prop = (const FClassPropertyParams*)PropBase;
 				NewProp = new (EC_InternalUseOnlyConstructor, Outer, UTF8_TO_TCHAR(Prop->NameUTF8), Prop->ObjectFlags) UClassProperty(FObjectInitializer(), EC_CppProperty, Prop->Offset, Prop->PropertyFlags, Prop->MetaClassFunc ? Prop->MetaClassFunc() : nullptr, Prop->ClassFunc ? Prop->ClassFunc() : nullptr);
@@ -4272,7 +4272,7 @@ namespace UE4CodeGen_Private
 			}
 			break;
 
-			case EPropertyClass::SoftClass:
+			case EPropertyGenFlags::SoftClass:
 			{
 				const FSoftClassPropertyParams* Prop = (const FSoftClassPropertyParams*)PropBase;
 				NewProp = new (EC_InternalUseOnlyConstructor, Outer, UTF8_TO_TCHAR(Prop->NameUTF8), Prop->ObjectFlags) USoftClassProperty(FObjectInitializer(), EC_CppProperty, Prop->Offset, Prop->PropertyFlags, Prop->MetaClassFunc ? Prop->MetaClassFunc() : nullptr);
@@ -4284,7 +4284,7 @@ namespace UE4CodeGen_Private
 			}
 			break;
 
-			case EPropertyClass::Interface:
+			case EPropertyGenFlags::Interface:
 			{
 				const FInterfacePropertyParams* Prop = (const FInterfacePropertyParams*)PropBase;
 				NewProp = new (EC_InternalUseOnlyConstructor, Outer, UTF8_TO_TCHAR(Prop->NameUTF8), Prop->ObjectFlags) UInterfaceProperty(FObjectInitializer(), EC_CppProperty, Prop->Offset, Prop->PropertyFlags, Prop->InterfaceClassFunc ? Prop->InterfaceClassFunc() : nullptr);
@@ -4296,7 +4296,7 @@ namespace UE4CodeGen_Private
 			}
 			break;
 
-			case EPropertyClass::Name:
+			case EPropertyGenFlags::Name:
 			{
 				const FNamePropertyParams* Prop = (const FNamePropertyParams*)PropBase;
 				NewProp = new (EC_InternalUseOnlyConstructor, Outer, UTF8_TO_TCHAR(Prop->NameUTF8), Prop->ObjectFlags) UNameProperty(FObjectInitializer(), EC_CppProperty, Prop->Offset, Prop->PropertyFlags);
@@ -4308,7 +4308,7 @@ namespace UE4CodeGen_Private
 			}
 			break;
 
-			case EPropertyClass::Str:
+			case EPropertyGenFlags::Str:
 			{
 				const FStrPropertyParams* Prop = (const FStrPropertyParams*)PropBase;
 				NewProp = new (EC_InternalUseOnlyConstructor, Outer, UTF8_TO_TCHAR(Prop->NameUTF8), Prop->ObjectFlags) UStrProperty(FObjectInitializer(), EC_CppProperty, Prop->Offset, Prop->PropertyFlags);
@@ -4320,7 +4320,7 @@ namespace UE4CodeGen_Private
 			}
 			break;
 
-			case EPropertyClass::Array:
+			case EPropertyGenFlags::Array:
 			{
 				const FArrayPropertyParams* Prop = (const FArrayPropertyParams*)PropBase;
 				NewProp = new (EC_InternalUseOnlyConstructor, Outer, UTF8_TO_TCHAR(Prop->NameUTF8), Prop->ObjectFlags) UArrayProperty(FObjectInitializer(), EC_CppProperty, Prop->Offset, Prop->PropertyFlags);
@@ -4335,7 +4335,7 @@ namespace UE4CodeGen_Private
 			}
 			break;
 
-			case EPropertyClass::Map:
+			case EPropertyGenFlags::Map:
 			{
 				const FMapPropertyParams* Prop = (const FMapPropertyParams*)PropBase;
 				NewProp = new (EC_InternalUseOnlyConstructor, Outer, UTF8_TO_TCHAR(Prop->NameUTF8), Prop->ObjectFlags) UMapProperty(FObjectInitializer(), EC_CppProperty, Prop->Offset, Prop->PropertyFlags);
@@ -4350,7 +4350,7 @@ namespace UE4CodeGen_Private
 			}
 			break;
 
-			case EPropertyClass::Set:
+			case EPropertyGenFlags::Set:
 			{
 				const FSetPropertyParams* Prop = (const FSetPropertyParams*)PropBase;
 				NewProp = new (EC_InternalUseOnlyConstructor, Outer, UTF8_TO_TCHAR(Prop->NameUTF8), Prop->ObjectFlags) USetProperty(FObjectInitializer(), EC_CppProperty, Prop->Offset, Prop->PropertyFlags);
@@ -4365,7 +4365,7 @@ namespace UE4CodeGen_Private
 			}
 			break;
 
-			case EPropertyClass::Struct:
+			case EPropertyGenFlags::Struct:
 			{
 				const FStructPropertyParams* Prop = (const FStructPropertyParams*)PropBase;
 				NewProp = new (EC_InternalUseOnlyConstructor, Outer, UTF8_TO_TCHAR(Prop->NameUTF8), Prop->ObjectFlags) UStructProperty(FObjectInitializer(), EC_CppProperty, Prop->Offset, Prop->PropertyFlags, Prop->ScriptStructFunc ? Prop->ScriptStructFunc() : nullptr);
@@ -4377,7 +4377,7 @@ namespace UE4CodeGen_Private
 			}
 			break;
 
-			case EPropertyClass::Delegate:
+			case EPropertyGenFlags::Delegate:
 			{
 				const FDelegatePropertyParams* Prop = (const FDelegatePropertyParams*)PropBase;
 				NewProp = new (EC_InternalUseOnlyConstructor, Outer, UTF8_TO_TCHAR(Prop->NameUTF8), Prop->ObjectFlags) UDelegateProperty(FObjectInitializer(), EC_CppProperty, Prop->Offset, Prop->PropertyFlags, Prop->SignatureFunctionFunc ? Prop->SignatureFunctionFunc() : nullptr);
@@ -4389,7 +4389,7 @@ namespace UE4CodeGen_Private
 			}
 			break;
 
-			case EPropertyClass::MulticastDelegate:
+			case EPropertyGenFlags::MulticastDelegate:
 			{
 				const FMulticastDelegatePropertyParams* Prop = (const FMulticastDelegatePropertyParams*)PropBase;
 				NewProp = new (EC_InternalUseOnlyConstructor, Outer, UTF8_TO_TCHAR(Prop->NameUTF8), Prop->ObjectFlags) UMulticastDelegateProperty(FObjectInitializer(), EC_CppProperty, Prop->Offset, Prop->PropertyFlags, Prop->SignatureFunctionFunc ? Prop->SignatureFunctionFunc() : nullptr);
@@ -4401,7 +4401,7 @@ namespace UE4CodeGen_Private
 			}
 			break;
 
-			case EPropertyClass::Text:
+			case EPropertyGenFlags::Text:
 			{
 				const FTextPropertyParams* Prop = (const FTextPropertyParams*)PropBase;
 				NewProp = new (EC_InternalUseOnlyConstructor, Outer, UTF8_TO_TCHAR(Prop->NameUTF8), Prop->ObjectFlags) UTextProperty(FObjectInitializer(), EC_CppProperty, Prop->Offset, Prop->PropertyFlags);
@@ -4413,7 +4413,7 @@ namespace UE4CodeGen_Private
 			}
 			break;
 
-			case EPropertyClass::Enum:
+			case EPropertyGenFlags::Enum:
 			{
 				const FEnumPropertyParams* Prop = (const FEnumPropertyParams*)PropBase;
 				NewProp = new (EC_InternalUseOnlyConstructor, Outer, UTF8_TO_TCHAR(Prop->NameUTF8), Prop->ObjectFlags) UEnumProperty(FObjectInitializer(), EC_CppProperty, Prop->Offset, Prop->PropertyFlags, Prop->EnumFunc ? Prop->EnumFunc() : nullptr);
diff --git a/Engine/Source/Runtime/CoreUObject/Public/UObject/UObjectGlobals.h b/Engine/Source/Runtime/CoreUObject/Public/UObject/UObjectGlobals.h
--- a/Engine/Source/Runtime/CoreUObject/Public/UObject/UObjectGlobals.h
+++ b/Engine/Source/Runtime/CoreUObject/Public/UObject/UObjectGlobals.h
@@ -2108,52 +2108,57 @@ struct FCppClassTypeInfoStatic;
 
 namespace UE4CodeGen_Private
 {
-	enum class EPropertyClass
-	{
-		Byte,
-		Int8,
-		Int16,
-		Int,
-		Int64,
-		UInt16,
-		UInt32,
-		UInt64,
-		UnsizedInt,
-		UnsizedUInt,
-		Float,
-		Double,
-		Bool,
-		SoftClass,
-		WeakObject,
-		LazyObject,
-		SoftObject,
-		Class,
-		Object,
-		Interface,
-		Name,
-		Str,
-		Array,
-		Map,
-		Set,
-		Struct,
-		Delegate,
-		MulticastDelegate,
-		Text,
-		Enum,
+	enum class EPropertyGenFlags : uint32
+	{
+		None              = 0x00,
+
+		// First 5 bits are the property type
+		Byte              = 0x00,
+		Int8              = 0x01,
+		Int16             = 0x02,
+		Int               = 0x03,
+		Int64             = 0x04,
+		UInt16            = 0x05,
+		UInt32            = 0x06,
+		UInt64            = 0x07,
+		UnsizedInt        = 0x08,
+		UnsizedUInt       = 0x09,
+		Float             = 0x0A,
+		Double            = 0x0B,
+		Bool              = 0x0C,
+		SoftClass         = 0x0D,
+		WeakObject        = 0x0E,
+		LazyObject        = 0x0F,
+		SoftObject        = 0x10,
+		Class             = 0x11,
+		Object            = 0x12,
+		Interface         = 0x13,
+		Name              = 0x14,
+		Str               = 0x15,
+		Array             = 0x16,
+		Map               = 0x17,
+		Set               = 0x18,
+		Struct            = 0x19,
+		Delegate          = 0x1A,
+		MulticastDelegate = 0x1B,
+		Text              = 0x1C,
+		Enum              = 0x1D,
+
+		// Property-specific flags
+		NativeBool        = 0x20
 	};
 
+	ENUM_CLASS_FLAGS(EPropertyGenFlags)
+
+	// Value which masks out the type of combined EPropertyGenFlags.
+	constexpr EPropertyGenFlags PropertyTypeMask = (EPropertyGenFlags)0x1F;
+
 	enum class EDynamicType
 	{
 		NotDynamic,
 		Dynamic
 	};
 
-	enum class ENativeBool
-	{
-		NotNative,
-		Native
-	};
-
 #if WITH_METADATA
 	struct FMetaDataPairParam
 	{
@@ -2164,212 +2169,207 @@ namespace UE4CodeGen_Private
 
 	struct FEnumeratorParam
 	{
-		const char*               NameUTF8;
-		int64                     Value;
-#if WITH_METADATA
-		const FMetaDataPairParam* MetaDataArray;
-		int32                     NumMetaData;
-#endif
+		const char* NameUTF8;
+		int64       Value;
 	};
 
 	// This is not a base class but is just a common initial sequence of all of the F*PropertyParams types below.
 	// We don't want to use actual inheritance because we want to construct aggregated compile-time tables of these things.
 	struct FPropertyParamsBase
 	{
-		EPropertyClass Type;
-		const char*    NameUTF8;
-		EObjectFlags   ObjectFlags;
-		EPropertyFlags PropertyFlags;
-		int32          ArrayDim;
-		const char*    RepNotifyFuncUTF8;
+		const char*       NameUTF8;
+		const char*       RepNotifyFuncUTF8;
+		EPropertyFlags    PropertyFlags;
+		EPropertyGenFlags Flags;
+		EObjectFlags      ObjectFlags;
+		int32             ArrayDim;
 	};
 
 	struct FPropertyParamsBaseWithOffset // : FPropertyParamsBase
 	{
-		EPropertyClass Type;
-		const char*    NameUTF8;
-		EObjectFlags   ObjectFlags;
-		EPropertyFlags PropertyFlags;
-		int32          ArrayDim;
-		const char*    RepNotifyFuncUTF8;
-		int32          Offset;
+		const char*       NameUTF8;
+		const char*       RepNotifyFuncUTF8;
+		EPropertyFlags    PropertyFlags;
+		EPropertyGenFlags Flags;
+		EObjectFlags      ObjectFlags;
+		int32             ArrayDim;
+		int32             Offset;
 	};
 
 	struct FGenericPropertyParams // : FPropertyParamsBaseWithOffset
 	{
-		EPropertyClass   Type;
-		const char*      NameUTF8;
-		EObjectFlags     ObjectFlags;
-		EPropertyFlags   PropertyFlags;
-		int32            ArrayDim;
-		const char*      RepNotifyFuncUTF8;
-		int32            Offset;
+		const char*       NameUTF8;
+		const char*       RepNotifyFuncUTF8;
+		EPropertyFlags    PropertyFlags;
+		EPropertyGenFlags Flags;
+		EObjectFlags      ObjectFlags;
+		int32             ArrayDim;
+		int32             Offset;
 #if WITH_METADATA
-		const FMetaDataPairParam*           MetaDataArray;
-		int32                               NumMetaData;
+		const FMetaDataPairParam* MetaDataArray;
+		int32                     NumMetaData;
 #endif
 	};
 
 	struct FBytePropertyParams // : FPropertyParamsBaseWithOffset
 	{
-		EPropertyClass   Type;
-		const char*      NameUTF8;
-		EObjectFlags     ObjectFlags;
-		EPropertyFlags   PropertyFlags;
-		int32            ArrayDim;
-		const char*      RepNotifyFuncUTF8;
-		int32            Offset;
-		UEnum*         (*EnumFunc)();
+		const char*         NameUTF8;
+		const char*         RepNotifyFuncUTF8;
+		EPropertyFlags      PropertyFlags;
+		EPropertyGenFlags   Flags;
+		EObjectFlags        ObjectFlags;
+		int32               ArrayDim;
+		int32               Offset;
+		UEnum*            (*EnumFunc)();
 #if WITH_METADATA
-		const FMetaDataPairParam*           MetaDataArray;
-		int32                               NumMetaData;
+		const FMetaDataPairParam* MetaDataArray;
+		int32                     NumMetaData;
 #endif
 	};
 
 	struct FBoolPropertyParams // : FPropertyParamsBase
 	{
-		EPropertyClass   Type;
-		const char*      NameUTF8;
-		EObjectFlags     ObjectFlags;
-		EPropertyFlags   PropertyFlags;
-		int32            ArrayDim;
-		const char*      RepNotifyFuncUTF8;
-		uint32           ElementSize;
-		ENativeBool      NativeBool;
-		SIZE_T           SizeOfOuter;
-		void           (*SetBitFunc)(void* Obj);
+		const char*         NameUTF8;
+		const char*         RepNotifyFuncUTF8;
+		EPropertyFlags      PropertyFlags;
+		EPropertyGenFlags   Flags;
+		EObjectFlags        ObjectFlags;
+		int32               ArrayDim;
+		uint32              ElementSize;
+		SIZE_T              SizeOfOuter;
+		void              (*SetBitFunc)(void* Obj);
 #if WITH_METADATA
-		const FMetaDataPairParam*           MetaDataArray;
-		int32                               NumMetaData;
+		const FMetaDataPairParam* MetaDataArray;
+		int32                     NumMetaData;
 #endif
 	};
 
 	struct FObjectPropertyParams // : FPropertyParamsBaseWithOffset
 	{
-		EPropertyClass   Type;
-		const char*      NameUTF8;
-		EObjectFlags     ObjectFlags;
-		EPropertyFlags   PropertyFlags;
-		int32            ArrayDim;
-		const char*      RepNotifyFuncUTF8;
-		int32            Offset;
-		UClass*        (*ClassFunc)();
+		const char*         NameUTF8;
+		const char*         RepNotifyFuncUTF8;
+		EPropertyFlags      PropertyFlags;
+		EPropertyGenFlags   Flags;
+		EObjectFlags        ObjectFlags;
+		int32               ArrayDim;
+		int32               Offset;
+		UClass*          (*ClassFunc)();
 #if WITH_METADATA
-		const FMetaDataPairParam*           MetaDataArray;
-		int32                               NumMetaData;
+		const FMetaDataPairParam* MetaDataArray;
+		int32                     NumMetaData;
 #endif
 	};
 
 	struct FClassPropertyParams // : FPropertyParamsBaseWithOffset
 	{
-		EPropertyClass   Type;
-		const char*      NameUTF8;
-		EObjectFlags     ObjectFlags;
-		EPropertyFlags   PropertyFlags;
-		int32            ArrayDim;
-		const char*      RepNotifyFuncUTF8;
-		int32            Offset;
-		UClass*        (*MetaClassFunc)();
-		UClass*        (*ClassFunc)();
+		const char*         NameUTF8;
+		const char*         RepNotifyFuncUTF8;
+		EPropertyFlags      PropertyFlags;
+		EPropertyGenFlags   Flags;
+		EObjectFlags        ObjectFlags;
+		int32               ArrayDim;
+		int32               Offset;
+		UClass*           (*MetaClassFunc)();
+		UClass*           (*ClassFunc)();
 #if WITH_METADATA
-		const FMetaDataPairParam*           MetaDataArray;
-		int32                               NumMetaData;
+		const FMetaDataPairParam* MetaDataArray;
+		int32                     NumMetaData;
 #endif
 	};
 
 	struct FSoftClassPropertyParams // : FPropertyParamsBaseWithOffset
 	{
-		EPropertyClass   Type;
-		const char*      NameUTF8;
-		EObjectFlags     ObjectFlags;
-		EPropertyFlags   PropertyFlags;
-		int32            ArrayDim;
-		const char*      RepNotifyFuncUTF8;
-		int32            Offset;
-		UClass*        (*MetaClassFunc)();
+		const char*         NameUTF8;
+		const char*         RepNotifyFuncUTF8;
+		EPropertyFlags      PropertyFlags;
+		EPropertyGenFlags   Flags;
+		EObjectFlags        ObjectFlags;
+		int32               ArrayDim;
+		int32               Offset;
+		UClass*           (*MetaClassFunc)();
 #if WITH_METADATA
-		const FMetaDataPairParam*           MetaDataArray;
-		int32                               NumMetaData;
+		const FMetaDataPairParam* MetaDataArray;
+		int32                     NumMetaData;
 #endif
 	};
 
 	struct FInterfacePropertyParams // : FPropertyParamsBaseWithOffset
 	{
-		EPropertyClass   Type;
-		const char*      NameUTF8;
-		EObjectFlags     ObjectFlags;
-		EPropertyFlags   PropertyFlags;
-		int32            ArrayDim;
-		const char*      RepNotifyFuncUTF8;
-		int32            Offset;
-		UClass*        (*InterfaceClassFunc)();
+		const char*         NameUTF8;
+		const char*         RepNotifyFuncUTF8;
+		EPropertyFlags      PropertyFlags;
+		EPropertyGenFlags   Flags;
+		EObjectFlags        ObjectFlags;
+		int32               ArrayDim;
+		int32               Offset;
+		UClass*           (*InterfaceClassFunc)();
 #if WITH_METADATA
-		const FMetaDataPairParam*           MetaDataArray;
-		int32                               NumMetaData;
+		const FMetaDataPairParam* MetaDataArray;
+		int32                     NumMetaData;
 #endif
 	};
 
 	struct FStructPropertyParams // : FPropertyParamsBaseWithOffset
 	{
-		EPropertyClass   Type;
-		const char*      NameUTF8;
-		EObjectFlags     ObjectFlags;
-		EPropertyFlags   PropertyFlags;
-		int32            ArrayDim;
-		const char*      RepNotifyFuncUTF8;
-		int32            Offset;
-		UScriptStruct* (*ScriptStructFunc)();
+		const char*         NameUTF8;
+		const char*         RepNotifyFuncUTF8;
+		EPropertyFlags      PropertyFlags;
+		EPropertyGenFlags   Flags;
+		EObjectFlags        ObjectFlags;
+		int32               ArrayDim;
+		int32               Offset;
+		UScriptStruct*    (*ScriptStructFunc)();
 #if WITH_METADATA
-		const FMetaDataPairParam*           MetaDataArray;
-		int32                               NumMetaData;
+		const FMetaDataPairParam* MetaDataArray;
+		int32                     NumMetaData;
 #endif
 	};
 
 	struct FDelegatePropertyParams // : FPropertyParamsBaseWithOffset
 	{
-		EPropertyClass   Type;
-		const char*      NameUTF8;
-		EObjectFlags     ObjectFlags;
-		EPropertyFlags   PropertyFlags;
-		int32            ArrayDim;
-		const char*      RepNotifyFuncUTF8;
-		int32            Offset;
-		UFunction*     (*SignatureFunctionFunc)();
+		const char*         NameUTF8;
+		const char*         RepNotifyFuncUTF8;
+		EPropertyFlags      PropertyFlags;
+		EPropertyGenFlags   Flags;
+		EObjectFlags        ObjectFlags;
+		int32               ArrayDim;
+		int32               Offset;
+		UFunction*        (*SignatureFunctionFunc)();
 #if WITH_METADATA
-		const FMetaDataPairParam*           MetaDataArray;
-		int32                               NumMetaData;
+		const FMetaDataPairParam* MetaDataArray;
+		int32                     NumMetaData;
 #endif
 	};
 
 	struct FMulticastDelegatePropertyParams // : FPropertyParamsBaseWithOffset
 	{
-		EPropertyClass   Type;
-		const char*      NameUTF8;
-		EObjectFlags     ObjectFlags;
-		EPropertyFlags   PropertyFlags;
-		int32            ArrayDim;
-		const char*      RepNotifyFuncUTF8;
-		int32            Offset;
-		UFunction*     (*SignatureFunctionFunc)();
+		const char*         NameUTF8;
+		const char*         RepNotifyFuncUTF8;
+		EPropertyFlags      PropertyFlags;
+		EPropertyGenFlags   Flags;
+		EObjectFlags        ObjectFlags;
+		int32               ArrayDim;
+		int32               Offset;
+		UFunction*        (*SignatureFunctionFunc)();
 #if WITH_METADATA
-		const FMetaDataPairParam*           MetaDataArray;
-		int32                               NumMetaData;
+		const FMetaDataPairParam* MetaDataArray;
+		int32                     NumMetaData;
 #endif
 	};
 
 	struct FEnumPropertyParams // : FPropertyParamsBaseWithOffset
 	{
-		EPropertyClass   Type;
-		const char*      NameUTF8;
-		EObjectFlags     ObjectFlags;
-		EPropertyFlags   PropertyFlags;
-		int32            ArrayDim;
-		const char*      RepNotifyFuncUTF8;
-		int32            Offset;
-		UEnum*         (*EnumFunc)();
+		const char*        NameUTF8;
+		const char*        RepNotifyFuncUTF8;
+		EPropertyFlags     PropertyFlags;
+		EPropertyGenFlags  Flags;
+		EObjectFlags       ObjectFlags;
+		int32              ArrayDim;
+		int32              Offset;
+		UEnum*           (*EnumFunc)();
 #if WITH_METADATA
-		const FMetaDataPairParam*           MetaDataArray;
-		int32                               NumMetaData;
+		const FMetaDataPairParam* MetaDataArray;
+		int32                     NumMetaData;
 #endif
 	};
 
@@ -2398,13 +2398,13 @@ namespace UE4CodeGen_Private
 	struct FFunctionParams
 	{
 		UObject*                          (*OuterFunc)();
-		const char*                         NameUTF8;
-		EObjectFlags                        ObjectFlags;
 		UFunction*                        (*SuperFunc)();
-		EFunctionFlags                      FunctionFlags;
+		const char*                         NameUTF8;
 		SIZE_T                              StructureSize;
 		const FPropertyParamsBase* const*   PropertyArray;
 		int32                               NumProperties;
+		EObjectFlags                        ObjectFlags;
+		EFunctionFlags                      FunctionFlags;
 		uint16                              RPCId;
 		uint16                              RPCResponseId;
 #if WITH_METADATA
@@ -2416,14 +2416,14 @@ namespace UE4CodeGen_Private
 	struct FEnumParams
 	{
 		UObject*                  (*OuterFunc)();
-		EDynamicType                DynamicType;
-		const char*                 NameUTF8;
-		EObjectFlags                ObjectFlags;
 		FText                     (*DisplayNameFunc)(int32);
-		uint8                       CppForm; // this is of type UEnum::ECppForm
+		const char*                 NameUTF8;
 		const char*                 CppTypeUTF8;
 		const FEnumeratorParam*     EnumeratorParams;
 		int32                       NumEnumerators;
+		EObjectFlags                ObjectFlags;
+		EDynamicType                DynamicType;
+		uint8                       CppForm; // this is of type UEnum::ECppForm
 #if WITH_METADATA
 		const FMetaDataPairParam*   MetaDataArray;
 		int32                       NumMetaData;
@@ -2436,12 +2436,12 @@ namespace UE4CodeGen_Private
 		UScriptStruct*                    (*SuperFunc)();
 		void*                             (*StructOpsFunc)(); // really returns UScriptStruct::ICppStructOps*
 		const char*                         NameUTF8;
-		EObjectFlags                        ObjectFlags;
-		uint32                              StructFlags; // EStructFlags
 		SIZE_T                              SizeOf;
 		SIZE_T                              AlignOf;
 		const FPropertyParamsBase* const*   PropertyArray;
 		int32                               NumProperties;
+		EObjectFlags                        ObjectFlags;
+		uint32                              StructFlags; // EStructFlags
 #if WITH_METADATA
 		const FMetaDataPairParam*           MetaDataArray;
 		int32                               NumMetaData;
@@ -2451,11 +2451,11 @@ namespace UE4CodeGen_Private
 	struct FPackageParams
 	{
 		const char*                        NameUTF8;
+		UObject*                  (*const *SingletonFuncArray)();
+		int32                              NumSingletons;
 		uint32                             PackageFlags; // EPackageFlags
 		uint32                             BodyCRC;
 		uint32                             DeclarationsCRC;
-		UObject*                  (*const *SingletonFuncArray)();
-		int32                              NumSingletons;
 #if WITH_METADATA
 		const FMetaDataPairParam*          MetaDataArray;
 		int32                              NumMetaData;
@@ -2472,17 +2472,17 @@ namespace UE4CodeGen_Private
 	struct FClassParams
 	{
 		UClass*                                   (*ClassNoRegisterFunc)();
+		const char*                                 ClassConfigNameUTF8;
+		const FCppClassTypeInfoStatic*              CppClassInfo;
 		UObject*                           (*const *DependencySingletonFuncArray)();
-		int32                                       NumDependencySingletons;
-		uint32                                      ClassFlags; // EClassFlags
 		const FClassFunctionLinkInfo*               FunctionLinkArray;
-		int32                                       NumFunctions;
 		const FPropertyParamsBase* const*           PropertyArray;
-		int32                                       NumProperties;
-		const char*                                 ClassConfigNameUTF8;
-		const FCppClassTypeInfoStatic*              CppClassInfo;
 		const FImplementedInterfaceParams*          ImplementedInterfaceArray;
+		int32                                       NumDependencySingletons;
+		int32                                       NumFunctions;
+		int32                                       NumProperties;
 		int32                                       NumImplementedInterfaces;
+		uint32                                      ClassFlags; // EClassFlags
 #if WITH_METADATA
 		const FMetaDataPairParam*                   MetaDataArray;
 		int32                                       NumMetaData;
