Simple UHT optimization yields 1.1x wall time speedup on internal project

diff --git a/Engine/Source/Programs/UnrealHeaderTool/Private/Scope.h b/Engine/Source/Programs/UnrealHeaderTool/Private/Scope.h
--- a/Engine/Source/Programs/UnrealHeaderTool/Private/Scope.h
+++ b/Engine/Source/Programs/UnrealHeaderTool/Private/Scope.h
@@ -263,6 +263,8 @@ private:
 	friend struct FStructScopeArchiveProxy;
 };
 
+using FScopeSet = TSet<FScope*, DefaultKeyFuncs<FScope*>, TInlineSetAllocator<1024>>;
+
 /**
  * Represents a scope associated with source file.
  */
@@ -300,18 +302,15 @@ public:
 	 *
 	 * @param Out (Output parameter) Array to append scopes.
 	 */
-	void AppendIncludedFileScopes(TArray<FScope*>& Out)
+	void AppendIncludedFileScopes(FScopeSet& Out)
 	{
-		if (!Out.Contains(this))
-		{
-			Out.Add(this);
-		}
+		bool bAlreadyAdded = false;
+		Out.Add(this, &bAlreadyAdded);
 
-		for (FFileScope* IncludedScope : IncludedScopes)
+		if (!bAlreadyAdded)
 		{
-			if (!Out.Contains(IncludedScope))
+			for (FFileScope* IncludedScope : IncludedScopes)
 			{
-				Out.Add(IncludedScope);
 				IncludedScope->AppendIncludedFileScopes(Out);
 			}
 		}
@@ -381,7 +380,7 @@ public:
 	typedef typename TConditionalType<TIsConst, TMap<FName, UField*>::TConstIterator, TMap<FName, UField*>::TIterator>::Type MapIteratorType;
 	typedef typename TConditionalType<TIsConst, const FScope, FScope>::Type ScopeType;
 	typedef typename TConditionalType<TIsConst, const FFileScope, FFileScope>::Type FileScopeType;
-	typedef typename TConditionalType<TIsConst, typename TArray<ScopeType*>::TConstIterator, typename TArray<ScopeType*>::TIterator>::Type ScopeArrayIteratorType;
+	typedef typename TConditionalType<TIsConst, FScopeSet::TConstIterator, FScopeSet::TIterator>::Type ScopeArrayIteratorType;
 
 	// Constructor.
 	TDeepScopeTypeIterator(ScopeType* Scope)
@@ -418,7 +417,7 @@ public:
 	 */
 	bool MoveNext()
 	{
-		if (!ScopeIterator.IsValid() && !MoveToNextScope())
+		if (!ScopeIterator.IsSet() && !MoveToNextScope())
 		{
 			return false;
 		}
@@ -431,7 +430,7 @@ public:
 		{
 			do
 			{
-				ScopeIterator = nullptr;
+				ScopeIterator.Reset();
 				if (!MoveToNextScope())
 				{
 					return false;
@@ -461,9 +460,9 @@ private:
 	 */
 	bool MoveToNextScope()
 	{
-		if (!ScopesIterator.IsValid())
+		if (!ScopesIterator.IsSet())
 		{
-			ScopesIterator = MakeShareable(new ScopeArrayIteratorType(ScopesToTraverse));
+			ScopesIterator.Emplace(ScopesToTraverse);
 		}
 		else
 		{
@@ -475,17 +474,17 @@ private:
 			return false;
 		}
 
-		ScopeIterator = MakeShareable(new FScope::TScopeTypeIterator<TType, TIsConst>(ScopesIterator->operator*()));
+		ScopeIterator.Emplace(ScopesIterator->operator*());
 		return true;
 	}
 
-	// List of scopes to traverse.
-	TArray<ScopeType*> ScopesToTraverse;
-
 	// Current scope iterator.
-	TSharedPtr<FScope::TScopeTypeIterator<TType, TIsConst> > ScopeIterator;
+	TOptional<FScope::TScopeTypeIterator<TType, TIsConst> > ScopeIterator;
 
 	// Scopes list iterator.
-	TSharedPtr<ScopeArrayIteratorType> ScopesIterator;
+	TOptional<ScopeArrayIteratorType> ScopesIterator;
+
+	// List of scopes to traverse.
+	FScopeSet  ScopesToTraverse;
 };
 
