Removal of obsolete error, as UCLASSes can now be defined in any header, not just one entitled <ClassNameWithNoPrefix>.h.

diff --git a/Engine/Source/Programs/UnrealHeaderTool/Private/CodeGenerator.cpp b/Engine/Source/Programs/UnrealHeaderTool/Private/CodeGenerator.cpp
--- a/Engine/Source/Programs/UnrealHeaderTool/Private/CodeGenerator.cpp
+++ b/Engine/Source/Programs/UnrealHeaderTool/Private/CodeGenerator.cpp
@@ -6082,11 +6082,6 @@ UClass* ProcessParsedClass(bool bClassIsAnInterface, TArray<FHeaderProvider>& De
 		{
 			FError::Throwf(TEXT("No prefix or invalid identifier for base class %s.\nClass names must match Unreal prefix specifications (e.g., \"UObject\" or \"AActor\")"), *BaseClassName);
 		}
-
-		if (DependentOn.ContainsByPredicate([&](const FHeaderProvider& Dependency){ FString DependencyStr = Dependency.GetId(); return !DependencyStr.Contains(TEXT(".generated.h")) && FPaths::GetBaseFilename(DependencyStr) == ClassNameStripped; }))
-		{
-			FError::Throwf(TEXT("Class '%s' contains a dependency (#include or base class) to itself"), *ClassName);
-		}
 	}
 
 	//UE_LOG(LogCompile, Log, TEXT("Class: %s extends %s"),*ClassName,*BaseClassName);
