Produce a useful error message instead of a crash when a UINTERFACE() declaration isn't parsed properly.

diff --git a/Engine/Source/Programs/UnrealHeaderTool/Private/HeaderParser.cpp b/Engine/Source/Programs/UnrealHeaderTool/Private/HeaderParser.cpp
--- a/Engine/Source/Programs/UnrealHeaderTool/Private/HeaderParser.cpp
+++ b/Engine/Source/Programs/UnrealHeaderTool/Private/HeaderParser.cpp
@@ -5043,6 +5043,11 @@ bool FHeaderParser::CompileDeclaration(FClasses& AllClasses, TArray<UDelegateFun
 		RequireSymbol(TEXT(")"), Token.Identifier);
 
 		FClassMetaData* ClassData = GetCurrentClassData();
+		if (!ClassData)
+		{
+			FString CurrentClassName = GetCurrentClass()->GetName();
+			FError::Throwf(TEXT("Could not find the associated 'U%s' class while parsing 'I%s' - it could be missing or malformed"), *CurrentClassName, *CurrentClassName);
+		}
 
 		ClassData->GeneratedBodyMacroAccessSpecifier = CurrentAccessSpecifier;
 		ClassData->SetInterfaceGeneratedBodyLine(InputLine);
