Support uint64 backed bitfields as properties

Previously got an unintelligent compile error: "Missing ';' in 'variable declaration'"

diff --git a/Engine/Source/Programs/UnrealHeaderTool/Private/HeaderParser.cpp b/Engine/Source/Programs/UnrealHeaderTool/Private/HeaderParser.cpp
--- a/Engine/Source/Programs/UnrealHeaderTool/Private/HeaderParser.cpp
+++ b/Engine/Source/Programs/UnrealHeaderTool/Private/HeaderParser.cpp
@@ -3825,6 +3825,11 @@ void FHeaderParser::GetVarType(
 	{
 		VarProperty = FPropertyBase(CPT_Int64);
 	}
+	else if ( VarType.Matches(TEXT("uint64")) && IsBitfieldProperty() )
+	{
+		// 64-bit bitfield (bool) type, treat it like 8 bit type
+		VarProperty = FPropertyBase(CPT_Bool8);
+	}
 	else if ( VarType.Matches(TEXT("uint32")) && IsBitfieldProperty() )
 	{
 		// 32-bit bitfield (bool) type, treat it like 8 bit type
