Implemented a DocumentPolicy metadata tag for structs, classes and enums that can be used to drive documentation validation on

* Tooltips for Classes, Structs, Enums and Functions
* Tooltips for Properties (including uniqueness for Copy&Paste bugs)
* Tooltips for Enum Entries (including uniqueness for Copy&Paste bugs)
* Tooltips for Function Parameter(including uniqueness for Copy&Paste bugs)
* UIMin and UIMax for float properties
* UIMin and UIMax settings for float parameters on Functions

diff --git a/Engine/Source/Programs/UnrealHeaderTool/Private/BaseParser.cpp b/Engine/Source/Programs/UnrealHeaderTool/Private/BaseParser.cpp
--- a/Engine/Source/Programs/UnrealHeaderTool/Private/BaseParser.cpp
+++ b/Engine/Source/Programs/UnrealHeaderTool/Private/BaseParser.cpp
@@ -120,6 +120,9 @@ namespace
 
 			FMetadataKeyword& EarlyAccessFeature = Dictionary.Add(TEXT("EarlyAccessPreview"), EMetadataValueArgument::None);
 			EarlyAccessFeature.InsertAddAction(TEXT("DevelopmentStatus"), TEXT("EarlyAccess"));
+
+			FMetadataKeyword& DocumentationPolicy = Dictionary.Add(TEXT("DocumentationPolicy"), EMetadataValueArgument::None);
+			DocumentationPolicy.InsertAddAction(TEXT("DocumentationPolicy"), TEXT("Strict"));
 		}
 
 		return Dictionary.Find(Keyword);
diff --git a/Engine/Source/Programs/UnrealHeaderTool/Private/HeaderParser.cpp b/Engine/Source/Programs/UnrealHeaderTool/Private/HeaderParser.cpp
--- a/Engine/Source/Programs/UnrealHeaderTool/Private/HeaderParser.cpp
+++ b/Engine/Source/Programs/UnrealHeaderTool/Private/HeaderParser.cpp
@@ -66,13 +66,17 @@ private:
 
 enum {MAX_ARRAY_SIZE=2048};
 
+static const FName NAME_Comment(TEXT("Comment"));
 static const FName NAME_ToolTip(TEXT("ToolTip"));
+static const FName NAME_DocumentationPolicy(TEXT("DocumentationPolicy"));
 EGeneratedCodeVersion FHeaderParser::DefaultGeneratedCodeVersion = EGeneratedCodeVersion::V1;
 TArray<FString> FHeaderParser::StructsWithNoPrefix;
 TArray<FString> FHeaderParser::StructsWithTPrefix;
 TArray<FString> FHeaderParser::DelegateParameterCountStrings;
 TMap<FString, FString> FHeaderParser::TypeRedirectMap;
+TArray<FString> FHeaderParser::PropertyCPPTypesRequiringUIRanges;
 TMap<UClass*, ClassDefinitionRange> ClassDefinitionRanges;
+
 /**
  * Dirty hack global variable to allow different result codes passed through
  * exceptions. Needs to be fixed in future versions of UHT.
@@ -1080,6 +1084,16 @@ namespace
 				}
 			}
 			break;
+
+			case ECheckedMetadataSpecifier::DocumentationPolicy:
+			{
+				const TCHAR* StrictValue = L"Strict";
+				if (InValue != StrictValue)
+				{
+					FError::Throwf(TEXT("'%s' metadata was '%s' but it must be %s"), *InKey, *InValue, *StrictValue);
+				}
+			}
+			break;
 		}
 	}
 
@@ -1609,6 +1623,7 @@ UEnum* FHeaderParser::CompileEnum()
 
 	// Parse all enums tags.
 	FToken TagToken;
+	TArray<TMap<FName, FString>> EntryMetaData;
 
 	TArray<TPair<FName, int64>> EnumNames;
 	int64 CurrentEnumValue = 0;
@@ -1691,6 +1706,9 @@ UEnum* FHeaderParser::CompileEnum()
 			++CurrentEnumValue;
 		}
 
+		TagToken.MetaData.Add(TEXT("Name"), NewTag.ToString());
+		EntryMetaData.Add(TagToken.MetaData);
+
 		// check for metadata on this enum value
 		ParseFieldMetaData(TagToken.MetaData, TagToken.Identifier);
 		if (TagToken.MetaData.Num() > 0)
@@ -1755,6 +1773,7 @@ UEnum* FHeaderParser::CompileEnum()
 		FError::Throwf(TEXT("Unable to generate enum MAX entry '%s' due to name collision"), *MaxEnumItem.ToString());
 	}
 
+	CheckDocumentationPolicyForEnum(Enum, EnumValueMetaData, EntryMetaData);
 	return Enum;
 }
 
@@ -2014,6 +2033,79 @@ FString FHeaderParser::FormatCommentForToolTip(const FString& Input)
 	return Result;
 }
 
+TMap<FName, FString> FHeaderParser::GetParameterToolTipsFromFunctionComment(const FString& Input)
+{
+	TMap<FName, FString> Map;
+	TArray<FString> Params;
+
+	int32 Offset = 0;
+	int32 ParamStart = -1;
+	int32 ParamLength = 0;
+	FString ParamPrefix;
+	while (Offset < Input.Len())
+	{
+		if (ParamStart == -1)
+		{
+			if (Input.Find(TEXT("@param"), ESearchCase::CaseSensitive, ESearchDir::FromStart, Offset) == Offset)
+			{
+				ParamStart = Offset + 6;
+				Offset += 6;
+				continue;
+			}
+			else if (Input.Find(TEXT("@return"), ESearchCase::CaseSensitive, ESearchDir::FromStart, Offset) == Offset)
+			{
+				ParamStart = Offset + 7;
+				Offset += 7;
+				ParamPrefix = TEXT("ReturnValue ");
+				continue;
+			}
+			Offset++;
+			continue;
+		}
+
+		TCHAR CurrentChar = Input[Offset];
+		if ((CurrentChar == '\n') || (CurrentChar == '\r'))
+		{
+			if (ParamLength > 0)
+			{
+				Params.Add(ParamPrefix + Input.Mid(ParamStart, ParamLength));
+				ParamStart = -1;
+				ParamLength = 0;
+			}
+			ParamPrefix.Reset();
+		}
+		else
+		{
+			ParamLength++;
+		}
+		Offset++;
+	}
+	if (ParamLength > 0)
+	{
+		Params.Add(ParamPrefix + Input.Mid(ParamStart, ParamLength));
+	}
+
+	for (FString Param : Params)
+	{
+		Param = Param.ConvertTabsToSpaces(4).TrimStartAndEnd();
+
+		int32 FirstSpaceIndex = -1;
+		if (!Param.FindChar(' ', FirstSpaceIndex))
+		{
+			continue;
+		}
+
+		FString ParamName = Param.Left(FirstSpaceIndex);
+		FString ParamToolTip = Param.Mid(FirstSpaceIndex + 1);
+		while (ParamToolTip.RemoveFromStart(" ")) {}
+
+		Map.Add(*ParamName, ParamToolTip);
+	}
+
+	return Map;
+}
+
+
 void FHeaderParser::AddFormattedPrevCommentAsTooltipMetaData(TMap<FName, FString>& MetaData)
 {
 	// Don't add a tooltip if one already exists.
@@ -2022,6 +2114,12 @@ void FHeaderParser::AddFormattedPrevCommentAsTooltipMetaData(TMap<FName, FString
 		return;
 	}
 
+	// Add the comment if it is not empty
+	if (!PrevComment.IsEmpty())
+	{
+		MetaData.Add(NAME_Comment, *PrevComment);
+	}
+
 	// Don't add a tooltip if the comment is empty after formatting.
 	FString FormattedComment = FormatCommentForToolTip(PrevComment);
 	if (!FormattedComment.Len())
@@ -2525,6 +2623,7 @@ UScriptStruct* FHeaderParser::CompileStructDeclaration(FClasses& AllClasses)
 	// Link the properties within the struct
 	Struct->StaticLink(true);
 
+	CheckDocumentationPolicyForStruct(Struct, MetaData);
 	return Struct;
 }
 
@@ -5832,6 +5931,7 @@ UClass* FHeaderParser::CompileClassDeclaration(FClasses& AllClasses)
 		}
 	}
 
+	CheckDocumentationPolicyForStruct(Class, MetaData);
 	return Class;
 }
 
@@ -6603,7 +6703,7 @@ void FHeaderParser::CompileFunctionDeclaration(FClasses& AllClasses)
 			{
 				UE_LOG_ERROR_UHT(TEXT("An explicit Category specifier is required for Blueprint accessible functions in an Engine module."));
 			}
-		} 
+		}
 	}
 
 	// Verify interfaces with respect to their blueprint accessible functions
@@ -6994,6 +7094,9 @@ void FHeaderParser::CompileFunctionDeclaration(FClasses& AllClasses)
 		// Put the token back so we can continue parsing as normal
 		UngetToken(Token);
 	}
+
+	// perform documentation policy tests
+	CheckDocumentationPolicyForFunc(GetCurrentClass(), FuncInfo.FunctionReference, MetaData);
 }
 
 /** Parses optional metadata text. */
@@ -9454,6 +9557,274 @@ void FHeaderParser::PostPopNestInterface(FClasses& AllClasses, UClass* CurrentIn
 	}
 }
 
+FDocumentationPolicy FHeaderParser::GetDocumentationPolicyFromName(const FString& PolicyName)
+{
+	FDocumentationPolicy DocumentationPolicy;
+	if (PolicyName.Equals(TEXT("Strict")))
+	{
+		DocumentationPolicy.bClassOrStructCommentRequired = true;
+		DocumentationPolicy.bFunctionToolTipsRequired = true;
+		DocumentationPolicy.bMemberToolTipsRequired = true;
+		DocumentationPolicy.bParameterToolTipsRequired = true;
+		DocumentationPolicy.bFloatRangesRequired = true;
+	}
+	else
+	{
+		FError::Throwf(TEXT("Documentation Policy '%s' not yet supported"), *PolicyName);
+	}
+	return DocumentationPolicy;
+}
+
+
+FDocumentationPolicy FHeaderParser::GetDocumentationPolicyForStruct(UStruct* Struct)
+{
+	check(Struct!= nullptr);
+
+	FDocumentationPolicy DocumentationPolicy;
+	FString DocumentationPolicyName;
+	if (Struct->GetStringMetaDataHierarchical(NAME_DocumentationPolicy, &DocumentationPolicyName))
+	{
+		DocumentationPolicy = GetDocumentationPolicyFromName(DocumentationPolicyName);
+	}
+	return DocumentationPolicy;
+}
+
+void FHeaderParser::CheckDocumentationPolicyForEnum(UEnum* Enum, const TMap<FName, FString>& MetaData, TArray<TMap<FName, FString>> Entries)
+{
+	check(Enum != nullptr);
+
+	const FString* DocumentationPolicyName = MetaData.Find(NAME_DocumentationPolicy);
+	if (DocumentationPolicyName == nullptr)
+	{
+		return;
+	}
+
+	check(!DocumentationPolicyName->IsEmpty());
+
+	FDocumentationPolicy DocumentationPolicy = GetDocumentationPolicyFromName(*DocumentationPolicyName);
+	if (DocumentationPolicy.bClassOrStructCommentRequired)
+	{
+		const FString* EnumToolTip = MetaData.Find(NAME_ToolTip);
+		if (EnumToolTip == nullptr)
+		{
+			FError::Throwf(TEXT("Enum '%s' does not provide a tooltip / comment (DocumentationPolicy)."), *Enum->GetName());
+		}
+	}
+
+	TMap<FString, FString> ToolTipToEntry;
+	for (const TMap<FName, FString>& Entry : Entries)
+	{
+		const FString* EntryName = Entry.Find(TEXT("Name"));
+		if (EntryName == nullptr)
+		{
+			continue;
+		}
+
+		const FString* ToolTip = Entry.Find(NAME_ToolTip);
+		if (ToolTip == nullptr)
+		{
+			FError::Throwf(TEXT("Enum entry '%s::%s' does not provide a tooltip / comment (DocumentationPolicy)."), *Enum->GetName(), **EntryName);
+		}
+
+		const FString* ExistingEntry = ToolTipToEntry.Find(*ToolTip);
+		if (ExistingEntry != nullptr)
+		{
+			FError::Throwf(TEXT("Enum entries '%s::%s' and '%s::%s' have identical tooltips / comments (DocumentationPolicy)."), *Enum->GetName(), **ExistingEntry, *Enum->GetName(), **EntryName);
+		}
+		ToolTipToEntry.Add(*ToolTip, *EntryName);
+	}
+	
+}
+
+void FHeaderParser::CheckDocumentationPolicyForStruct(UStruct* Struct, const TMap<FName, FString>& MetaData)
+{
+	check(Struct != nullptr);
+
+	FDocumentationPolicy DocumentationPolicy = GetDocumentationPolicyForStruct(Struct);
+	if (DocumentationPolicy.bClassOrStructCommentRequired)
+	{
+		const FString* ClassTooltip = MetaData.Find(NAME_ToolTip);
+		if (ClassTooltip == nullptr)
+		{
+			FError::Throwf(TEXT("Struct '%s' does not provide a tooltip / comment (DocumentationPolicy)."), *Struct->GetName());
+		}
+	}
+
+	if (DocumentationPolicy.bMemberToolTipsRequired)
+	{
+		TMap<FString, FName> ToolTipToPropertyName;
+		for (UProperty* Property : TFieldRange<UProperty>(Struct))
+		{
+			FString ToolTip = Property->GetToolTipText().ToString();
+			if (ToolTip.IsEmpty() || ToolTip.Equals(Property->GetName()))
+			{
+				FError::Throwf(TEXT("Property '%s::%s' does not provide a tooltip / comment (DocumentationPolicy)."), *Struct->GetName(), *Property->GetName());
+			}
+			const FName* ExistingPropertyName = ToolTipToPropertyName.Find(ToolTip);
+			if (ExistingPropertyName != nullptr)
+			{
+				FError::Throwf(TEXT("Property '%s::%s' and '%s::%s' are using identical tooltips (DocumentationPolicy)."), *Struct->GetName(), *ExistingPropertyName->ToString(), *Struct->GetName(), *Property->GetName());
+			}
+			ToolTipToPropertyName.Add(ToolTip, Property->GetFName());
+		}
+	}
+
+	if (DocumentationPolicy.bFloatRangesRequired)
+	{
+		for (UProperty* Property : TFieldRange<UProperty>(Struct))
+		{
+			if(DoesCPPTypeRequireDocumentation(Property->GetCPPType()))
+			{
+				const FString& UIMin = Property->GetMetaData(TEXT("UIMin"));
+				const FString& UIMax = Property->GetMetaData(TEXT("UIMax"));
+				if(!CheckUIMinMaxRangeFromMetaData(UIMin, UIMax))
+				{
+					FError::Throwf(TEXT("Property '%s::%s' does not provide a valid UIMin / UIMax (DocumentationPolicy)."), *Struct->GetName(), *Property->GetName());
+				}
+			}
+		}
+	}
+
+	// also compare all tooltips to see if they are unique
+	if (DocumentationPolicy.bFunctionToolTipsRequired)
+	{
+		UClass* Class = Cast<UClass>(Struct);
+		if (Class != nullptr)
+		{
+			TMap<FString, FString> ToolTipToFunc;
+			for (UField* Field = Class->Children; Field; Field = Field->Next)
+			{
+				UFunction* Func = Cast<UFunction>(Field);
+				if (Func != NULL)
+				{
+					FString ToolTip = Func->GetToolTipText().ToString();
+					if (ToolTip.IsEmpty())
+					{
+						FError::Throwf(TEXT("Function '%s::%s' does not provide a tooltip / comment (DocumentationPolicy)."), *Class->GetName(), *Func->GetName());
+					}
+					const FString* ExistingFuncName = ToolTipToFunc.Find(ToolTip);
+					if (ExistingFuncName != nullptr)
+					{
+						FError::Throwf(TEXT("Functions '%s::%s' and '%s::%s' uses identical tooltips / comments (DocumentationPolicy)."), *Class->GetName(), **ExistingFuncName, *Class->GetName(), *Func->GetName());
+					}
+					ToolTipToFunc.Add(ToolTip, Func->GetName());
+				}
+			}
+		}
+	}
+}
+
+bool FHeaderParser::DoesCPPTypeRequireDocumentation(const FString& CPPType)
+{
+	if (PropertyCPPTypesRequiringUIRanges.Num() == 0)
+	{
+		PropertyCPPTypesRequiringUIRanges.Add(TEXT("float"));
+		PropertyCPPTypesRequiringUIRanges.Add(TEXT("double"));
+	}
+	return PropertyCPPTypesRequiringUIRanges.Find(CPPType) != INDEX_NONE;
+}
+
+// Validates the documentation for a given method
+void FHeaderParser::CheckDocumentationPolicyForFunc(UClass* Class, UFunction* Func, const TMap<FName, FString>& MetaData)
+{
+	check(Class != nullptr);
+	check(Func != nullptr);
+
+	FDocumentationPolicy DocumentationPolicy = GetDocumentationPolicyForStruct(Class);
+	if (DocumentationPolicy.bFunctionToolTipsRequired)
+	{
+		const FString* FunctionTooltip = MetaData.Find(NAME_ToolTip);
+		if (FunctionTooltip == nullptr)
+		{
+			FError::Throwf(TEXT("Function '%s::%s' does not provide a tooltip / comment (DocumentationPolicy)."), *Class->GetName(), *Func->GetName());
+		}
+	}
+
+	if (DocumentationPolicy.bParameterToolTipsRequired)
+	{
+		const FString* FunctionComment = MetaData.Find(NAME_Comment);
+		if (FunctionComment == nullptr)
+		{
+			FError::Throwf(TEXT("Function '%s::%s' does not provide a comment (DocumentationPolicy)."), *Class->GetName(), *Func->GetName());
+		}
+		TMap<FName, FString> ParamToolTips = GetParameterToolTipsFromFunctionComment(*FunctionComment);
+
+		// ensure each parameter has a tooltip
+		TSet<FName> ExistingFields;
+		for (TFieldIterator<UProperty> FieldIt(Func); !!FieldIt; ++FieldIt)
+		{
+			FName ParamName = FieldIt->GetFName();
+			const FString* ParamToolTip = ParamToolTips.Find(ParamName);
+			if (ParamToolTip == nullptr)
+			{
+				FError::Throwf(TEXT("Function '%s::%s' doesn't provide a tooltip for parameter '%s' (DocumentationPolicy)."), *Class->GetName(), *Func->GetName(), *ParamName.ToString());
+			}
+			ExistingFields.Add(ParamName);
+		}
+
+		// ensure we don't have parameter tooltips for parameters that don't exist
+		for (TPair<FName, FString>& Pair : ParamToolTips)
+		{
+			if (!ExistingFields.Contains(Pair.Key))
+			{
+				FError::Throwf(TEXT("Function '%s::%s' provides a tooltip for an unknown parameter '%s' (DocumentationPolicy)."), *Class->GetName(), *Func->GetName(), *Pair.Key.ToString());
+			}
+		}
+
+		// check for duplicate tooltips
+		TMap<FString, FName> ToolTipToParam;
+		for (TPair<FName, FString>& Pair : ParamToolTips)
+		{
+			const FName* ExistingParam = ToolTipToParam.Find(Pair.Value);
+			if (ExistingParam != nullptr)
+			{
+				FError::Throwf(TEXT("Function '%s::%s' uses identical tooltips for parameters '%s' and '%s' (DocumentationPolicy)."), *Class->GetName(), *Func->GetName(), *ExistingParam->ToString(), *Pair.Key.ToString());
+			}
+			ToolTipToParam.Add(Pair.Value, Pair.Key);
+		}
+	}
+
+	if (DocumentationPolicy.bFloatRangesRequired)
+	{
+		for (TFieldIterator<UProperty> FieldIt(Func); !!FieldIt; ++FieldIt)
+		{
+			if (DoesCPPTypeRequireDocumentation(FieldIt->GetCPPType()))
+			{
+				FString ParamName = FieldIt->GetName();
+				if (ParamName.Equals(TEXT("ReturnValue")))
+				{
+					continue;
+				}
+
+				const FString& UIMin = FieldIt->GetMetaData(TEXT("UIMin"));
+				const FString& UIMax = FieldIt->GetMetaData(TEXT("UIMax"));
+				if (!CheckUIMinMaxRangeFromMetaData(UIMin, UIMax))
+				{
+					FError::Throwf(TEXT("Function's '%s::%s' parameter '%s' does not provide valid UIMin / UIMax (DocumentationPolicy)."), *Class->GetName(), *Func->GetName(), *ParamName);
+				}
+			}
+		}
+	}
+
+}
+
+bool FHeaderParser::CheckUIMinMaxRangeFromMetaData(const FString& UIMin, const FString& UIMax)
+{
+	if (UIMin.IsEmpty() || UIMax.IsEmpty())
+	{
+		return false;
+	}
+
+	float UIMinValue = FCString::Atof(*UIMin);
+	float UIMaxValue = FCString::Atof(*UIMax);
+	if (UIMin > UIMax) // note that we actually allow UIMin == UIMax to disable the range manually.
+	{
+		return false;
+	}
+
+	return true;
+}
+
 template <class TFunctionType>
 TFunctionType* CreateFunctionImpl(const FFuncInfo& FuncInfo, UObject* Outer, FScope* CurrentScope)
 {
@@ -9498,3 +9869,4 @@ UDelegateFunction* FHeaderParser::CreateDelegateFunction(const FFuncInfo &FuncIn
 	UObject* CurrentPackage = LocSourceFile ? LocSourceFile->GetPackage() : nullptr;
 	return CreateFunctionImpl<UDelegateFunction>(FuncInfo, IsInAClass() ? (UObject*)GetCurrentClass() : CurrentPackage, GetCurrentScope());
 }
+
diff --git a/Engine/Source/Programs/UnrealHeaderTool/Private/HeaderParser.h b/Engine/Source/Programs/UnrealHeaderTool/Private/HeaderParser.h
--- a/Engine/Source/Programs/UnrealHeaderTool/Private/HeaderParser.h
+++ b/Engine/Source/Programs/UnrealHeaderTool/Private/HeaderParser.h
@@ -147,6 +147,15 @@ struct ClassDefinitionRange
 
 extern TMap<UClass*, ClassDefinitionRange> ClassDefinitionRanges;
 
+struct FDocumentationPolicy
+{
+	bool bClassOrStructCommentRequired = false;
+	bool bFunctionToolTipsRequired = false;
+	bool bMemberToolTipsRequired = false;
+	bool bParameterToolTipsRequired = false;
+	bool bFloatRangesRequired = false;
+};
+
 /////////////////////////////////////////////////////
 // FHeaderParser
 
@@ -469,6 +478,13 @@ protected:
 	 * @return					The input string, reformatted in such a way as to be appropriate for use as a tooltip.
 	 */
 	static FString FormatCommentForToolTip(const FString& Input);
+
+	/**
+	 * Retrieves parameter comments / tooltips from the function comment
+	 * @param		Input		An input string, expected to be a script comment.
+	 * @return					The map of parameter name to comment per parameter
+	*/
+	static TMap<FName, FString> GetParameterToolTipsFromFunctionComment(const FString& Input);
 	
 	/**
 	 * Begins the process of exporting C++ class declarations for native classes in the specified package
@@ -740,6 +756,30 @@ private:
 
 	// Verifies the target function meets the criteria for a replication notify callback
 	void VerifyRepNotifyCallback(UProperty* Property, UFunction* TargetFunction);
+
+	// Constructs the policy from a string
+	static FDocumentationPolicy GetDocumentationPolicyFromName(const FString& PolicyName);
+
+	// Constructs the policy for documentation checks for a given struct
+	static FDocumentationPolicy GetDocumentationPolicyForStruct(UStruct* Struct);
+
+	// Property types to provide UI Min and Max ranges
+	static TArray<FString> PropertyCPPTypesRequiringUIRanges;
+
+	// Returns true if a given CPP types required ui checking
+	static bool DoesCPPTypeRequireDocumentation(const FString& CPPType);
+
+	// Validates the documentation for a given enum
+	void CheckDocumentationPolicyForEnum(UEnum* Enum, const TMap<FName, FString>& MetaData, TArray<TMap<FName, FString>> Entries);
+
+	// Validates the documentation for a given struct
+	void CheckDocumentationPolicyForStruct(UStruct* Struct, const TMap<FName, FString>& MetaData);
+
+	// Validates the documentation for a given method
+	void CheckDocumentationPolicyForFunc(UClass* Class, UFunction* Func, const TMap<FName, FString>& MetaData);
+
+	// Checks if a valid range has been found on the provided metadata
+	bool CheckUIMinMaxRangeFromMetaData(const FString& UIMin, const FString& UIMax);
 };
 
 /////////////////////////////////////////////////////
diff --git a/Engine/Source/Programs/UnrealHeaderTool/Private/Specifiers/CheckedMetadataSpecifiers.def b/Engine/Source/Programs/UnrealHeaderTool/Private/Specifiers/CheckedMetadataSpecifiers.def
--- a/Engine/Source/Programs/UnrealHeaderTool/Private/Specifiers/CheckedMetadataSpecifiers.def
+++ b/Engine/Source/Programs/UnrealHeaderTool/Private/Specifiers/CheckedMetadataSpecifiers.def
@@ -7,6 +7,7 @@ CHECKED_METADATA_SPECIFIER(ClampMax)
 CHECKED_METADATA_SPECIFIER(ClampMin)
 CHECKED_METADATA_SPECIFIER(CommutativeAssociativeBinaryOperator)
 CHECKED_METADATA_SPECIFIER(DevelopmentStatus)
+CHECKED_METADATA_SPECIFIER(DocumentationPolicy)
 CHECKED_METADATA_SPECIFIER(ExpandEnumAsExecs)
 CHECKED_METADATA_SPECIFIER(UIMax)
 CHECKED_METADATA_SPECIFIER(UIMin)
