Follow up on Document Policy changes

diff --git a/Engine/Source/Programs/UnrealHeaderTool/Private/HeaderParser.cpp b/Engine/Source/Programs/UnrealHeaderTool/Private/HeaderParser.cpp
--- a/Engine/Source/Programs/UnrealHeaderTool/Private/HeaderParser.cpp
+++ b/Engine/Source/Programs/UnrealHeaderTool/Private/HeaderParser.cpp
@@ -74,7 +74,7 @@ TArray<FString> FHeaderParser::StructsWithNoPrefix;
 TArray<FString> FHeaderParser::StructsWithTPrefix;
 TArray<FString> FHeaderParser::DelegateParameterCountStrings;
 TMap<FString, FString> FHeaderParser::TypeRedirectMap;
-TArray<FString> FHeaderParser::PropertyCPPTypesRequiringUIRanges;
+TArray<FString> FHeaderParser::PropertyCPPTypesRequiringUIRanges = { TEXT("float"), TEXT("double") };
 TMap<UClass*, ClassDefinitionRange> ClassDefinitionRanges;
 
 /**
@@ -1087,7 +1087,7 @@ namespace
 
 			case ECheckedMetadataSpecifier::DocumentationPolicy:
 			{
-				const TCHAR* StrictValue = L"Strict";
+				const TCHAR* StrictValue = TEXT("Strict");
 				if (InValue != StrictValue)
 				{
 					FError::Throwf(TEXT("'%s' metadata was '%s' but it must be %s"), *InKey, *InValue, *StrictValue);
@@ -2097,7 +2097,7 @@ TMap<FName, FString> FHeaderParser::GetParameterToolTipsFromFunctionComment(cons
 
 		FString ParamName = Param.Left(FirstSpaceIndex);
 		FString ParamToolTip = Param.Mid(FirstSpaceIndex + 1);
-		while (ParamToolTip.RemoveFromStart(" ")) {}
+		ParamToolTip.TrimStartInline();
 
 		Map.Add(*ParamName, ParamToolTip);
 	}
@@ -9589,7 +9589,7 @@ FDocumentationPolicy FHeaderParser::GetDocumentationPolicyForStruct(UStruct* Str
 	return DocumentationPolicy;
 }
 
-void FHeaderParser::CheckDocumentationPolicyForEnum(UEnum* Enum, const TMap<FName, FString>& MetaData, TArray<TMap<FName, FString>> Entries)
+void FHeaderParser::CheckDocumentationPolicyForEnum(UEnum* Enum, const TMap<FName, FString>& MetaData, const TArray<TMap<FName, FString>>& Entries)
 {
 	check(Enum != nullptr);
 
@@ -9716,11 +9716,6 @@ void FHeaderParser::CheckDocumentationPolicyForStruct(UStruct* Struct, const TMa
 
 bool FHeaderParser::DoesCPPTypeRequireDocumentation(const FString& CPPType)
 {
-	if (PropertyCPPTypesRequiringUIRanges.Num() == 0)
-	{
-		PropertyCPPTypesRequiringUIRanges.Add(TEXT("float"));
-		PropertyCPPTypesRequiringUIRanges.Add(TEXT("double"));
-	}
 	return PropertyCPPTypesRequiringUIRanges.Find(CPPType) != INDEX_NONE;
 }
 
@@ -9751,9 +9746,9 @@ void FHeaderParser::CheckDocumentationPolicyForFunc(UClass* Class, UFunction* Fu
 
 		// ensure each parameter has a tooltip
 		TSet<FName> ExistingFields;
-		for (TFieldIterator<UProperty> FieldIt(Func); !!FieldIt; ++FieldIt)
+		for (UProperty* Property : TFieldRange<UProperty>(Func))
 		{
-			FName ParamName = FieldIt->GetFName();
+			FName ParamName = Property->GetFName();
 			const FString* ParamToolTip = ParamToolTips.Find(ParamName);
 			if (ParamToolTip == nullptr)
 			{
@@ -9786,18 +9781,18 @@ void FHeaderParser::CheckDocumentationPolicyForFunc(UClass* Class, UFunction* Fu
 
 	if (DocumentationPolicy.bFloatRangesRequired)
 	{
-		for (TFieldIterator<UProperty> FieldIt(Func); !!FieldIt; ++FieldIt)
+		for (UProperty* Property : TFieldRange<UProperty>(Func))
 		{
-			if (DoesCPPTypeRequireDocumentation(FieldIt->GetCPPType()))
+			if (DoesCPPTypeRequireDocumentation(Property->GetCPPType()))
 			{
-				FString ParamName = FieldIt->GetName();
+				FString ParamName = Property->GetName();
 				if (ParamName.Equals(TEXT("ReturnValue")))
 				{
 					continue;
 				}
 
-				const FString& UIMin = FieldIt->GetMetaData(TEXT("UIMin"));
-				const FString& UIMax = FieldIt->GetMetaData(TEXT("UIMax"));
+				const FString& UIMin = Property->GetMetaData(TEXT("UIMin"));
+				const FString& UIMax = Property->GetMetaData(TEXT("UIMax"));
 				if (!CheckUIMinMaxRangeFromMetaData(UIMin, UIMax))
 				{
 					FError::Throwf(TEXT("Function's '%s::%s' parameter '%s' does not provide valid UIMin / UIMax (DocumentationPolicy)."), *Class->GetName(), *Func->GetName(), *ParamName);
@@ -9815,8 +9810,8 @@ bool FHeaderParser::CheckUIMinMaxRangeFromMetaData(const FString& UIMin, const F
 		return false;
 	}
 
-	float UIMinValue = FCString::Atof(*UIMin);
-	float UIMaxValue = FCString::Atof(*UIMax);
+	double UIMinValue = FCString::Atod(*UIMin);
+	double UIMaxValue = FCString::Atod(*UIMax);
 	if (UIMin > UIMax) // note that we actually allow UIMin == UIMax to disable the range manually.
 	{
 		return false;
diff --git a/Engine/Source/Programs/UnrealHeaderTool/Private/HeaderParser.h b/Engine/Source/Programs/UnrealHeaderTool/Private/HeaderParser.h
--- a/Engine/Source/Programs/UnrealHeaderTool/Private/HeaderParser.h
+++ b/Engine/Source/Programs/UnrealHeaderTool/Private/HeaderParser.h
@@ -770,7 +770,7 @@ private:
 	static bool DoesCPPTypeRequireDocumentation(const FString& CPPType);
 
 	// Validates the documentation for a given enum
-	void CheckDocumentationPolicyForEnum(UEnum* Enum, const TMap<FName, FString>& MetaData, TArray<TMap<FName, FString>> Entries);
+	void CheckDocumentationPolicyForEnum(UEnum* Enum, const TMap<FName, FString>& MetaData, const TArray<TMap<FName, FString>>& Entries);
 
 	// Validates the documentation for a given struct
 	void CheckDocumentationPolicyForStruct(UStruct* Struct, const TMap<FName, FString>& MetaData);
