Follow up on UHT changes based on Core-Team feedback.

* Remove the requirement of a [REDACTED_HANDLE] statement
* Only value the [REDACTED_HANDLE] statements if ANY [REDACTED_HANDLE] statement exist
* Remove the requirement of UPARAM range definitions for float / double function parameters

diff --git a/Engine/Source/Programs/UnrealHeaderTool/Private/HeaderParser.cpp b/Engine/Source/Programs/UnrealHeaderTool/Private/HeaderParser.cpp
--- a/Engine/Source/Programs/UnrealHeaderTool/Private/HeaderParser.cpp
+++ b/Engine/Source/Programs/UnrealHeaderTool/Private/HeaderParser.cpp
@@ -9744,65 +9744,70 @@ void FHeaderParser::CheckDocumentationPolicyForFunc(UClass* Class, UFunction* Fu
 		{
 			FError::Throwf(TEXT("Function '%s::%s' does not provide a comment (DocumentationPolicy)."), *Class->GetName(), *Func->GetName());
 		}
+		
 		TMap<FName, FString> ParamToolTips = GetParameterToolTipsFromFunctionComment(*FunctionComment);
-
-		// ensure each parameter has a tooltip
-		TSet<FName> ExistingFields;
-		for (UProperty* Property : TFieldRange<UProperty>(Func))
+		bool HasAnyParamToolTips = ParamToolTips.Num() > 0;
+		if (ParamToolTips.Num() == 0)
 		{
-			FName ParamName = Property->GetFName();
-			const FString* ParamToolTip = ParamToolTips.Find(ParamName);
-			if (ParamToolTip == nullptr)
+			const FString* ReturnValueToolTip = ParamToolTips.Find(TEXT("ReturnValue"));
+			if (ReturnValueToolTip != nullptr)
 			{
-				FError::Throwf(TEXT("Function '%s::%s' doesn't provide a tooltip for parameter '%s' (DocumentationPolicy)."), *Class->GetName(), *Func->GetName(), *ParamName.ToString());
+				HasAnyParamToolTips = false;
 			}
-			ExistingFields.Add(ParamName);
 		}
 
-		// ensure we don't have parameter tooltips for parameters that don't exist
-		for (TPair<FName, FString>& Pair : ParamToolTips)
+		// only apply the validation for parameter tooltips if a function has any @param statements at all.
+		if (HasAnyParamToolTips)
 		{
-			if (!ExistingFields.Contains(Pair.Key))
+			// ensure each parameter has a tooltip
+			TSet<FName> ExistingFields;
+			for (UProperty* Property : TFieldRange<UProperty>(Func))
 			{
-				FError::Throwf(TEXT("Function '%s::%s' provides a tooltip for an unknown parameter '%s' (DocumentationPolicy)."), *Class->GetName(), *Func->GetName(), *Pair.Key.ToString());
+				FName ParamName = Property->GetFName();
+				if (ParamName.IsEqual(TEXT("ReturnValue")))
+				{
+					continue;
+				}
+				const FString* ParamToolTip = ParamToolTips.Find(ParamName);
+				if (ParamToolTip == nullptr)
+				{
+					FError::Throwf(TEXT("Function '%s::%s' doesn't provide a tooltip for parameter '%s' (DocumentationPolicy)."), *Class->GetName(), *Func->GetName(), *ParamName.ToString());
+				}
+				ExistingFields.Add(ParamName);
 			}
-		}
 
-		// check for duplicate tooltips
-		TMap<FString, FName> ToolTipToParam;
-		for (TPair<FName, FString>& Pair : ParamToolTips)
-		{
-			const FName* ExistingParam = ToolTipToParam.Find(Pair.Value);
-			if (ExistingParam != nullptr)
+			// ensure we don't have parameter tooltips for parameters that don't exist
+			for (TPair<FName, FString>& Pair : ParamToolTips)
 			{
-				FError::Throwf(TEXT("Function '%s::%s' uses identical tooltips for parameters '%s' and '%s' (DocumentationPolicy)."), *Class->GetName(), *Func->GetName(), *ExistingParam->ToString(), *Pair.Key.ToString());
+				const FName& ParamName = Pair.Key;
+				if (ParamName.IsEqual(TEXT("ReturnValue")))
+				{
+					continue;
+				}
+				if (!ExistingFields.Contains(ParamName))
+				{
+					FError::Throwf(TEXT("Function '%s::%s' provides a tooltip for an unknown parameter '%s' (DocumentationPolicy)."), *Class->GetName(), *Func->GetName(), *Pair.Key.ToString());
+				}
 			}
-			ToolTipToParam.Add(Pair.Value, Pair.Key);
-		}
-	}
 
-	if (DocumentationPolicy.bFloatRangesRequired)
-	{
-		for (UProperty* Property : TFieldRange<UProperty>(Func))
-		{
-			if (DoesCPPTypeRequireDocumentation(Property->GetCPPType()))
+			// check for duplicate tooltips
+			TMap<FString, FName> ToolTipToParam;
+			for (TPair<FName, FString>& Pair : ParamToolTips)
 			{
-				FString ParamName = Property->GetName();
-				if (ParamName.Equals(TEXT("ReturnValue")))
+				const FName& ParamName = Pair.Key;
+				if (ParamName.IsEqual(TEXT("ReturnValue")))
 				{
 					continue;
 				}
-
-				const FString& UIMin = Property->GetMetaData(TEXT("UIMin"));
-				const FString& UIMax = Property->GetMetaData(TEXT("UIMax"));
-				if (!CheckUIMinMaxRangeFromMetaData(UIMin, UIMax))
+				const FName* ExistingParam = ToolTipToParam.Find(Pair.Value);
+				if (ExistingParam != nullptr)
 				{
-					FError::Throwf(TEXT("Function's '%s::%s' parameter '%s' does not provide valid UIMin / UIMax (DocumentationPolicy)."), *Class->GetName(), *Func->GetName(), *ParamName);
+					FError::Throwf(TEXT("Function '%s::%s' uses identical tooltips for parameters '%s' and '%s' (DocumentationPolicy)."), *Class->GetName(), *Func->GetName(), *ExistingParam->ToString(), *Pair.Key.ToString());
 				}
+				ToolTipToParam.Add(Pair.Value, Pair.Key);
 			}
 		}
 	}
-
 }
 
 bool FHeaderParser::CheckUIMinMaxRangeFromMetaData(const FString& UIMin, const FString& UIMax)
