Made server rpc validation optional.

diff --git a/Engine/Source/Programs/UnrealHeaderTool/Private/HeaderParser.cpp b/Engine/Source/Programs/UnrealHeaderTool/Private/HeaderParser.cpp
--- a/Engine/Source/Programs/UnrealHeaderTool/Private/HeaderParser.cpp
+++ b/Engine/Source/Programs/UnrealHeaderTool/Private/HeaderParser.cpp
@@ -6700,10 +6700,6 @@ void FHeaderParser::CompileFunctionDeclaration(FClasses& AllClasses)
 	ProcessFunctionSpecifiers(FuncInfo, SpecifiersFound, MetaData);
 
 	const bool bClassGeneratedFromBP = FClass::IsDynamic(GetCurrentClass());
-	if ((FuncInfo.FunctionFlags & FUNC_NetServer) && !(FuncInfo.FunctionFlags & FUNC_NetValidate) && !bClassGeneratedFromBP)
-	{
-		FError::Throwf(TEXT("Server RPC missing 'WithValidation' keyword in the UFUNCTION() declaration statement.  Required for security purposes."));
-	}
 
 	if ((0 != (FuncInfo.FunctionExportFlags & FUNCEXPORT_CustomThunk)) && !MetaData.Contains("CustomThunk"))
 	{
