Add a warning to UHT to check enums and ensure that they have an entry that is equal to 0. This will fix BP compiler warnings with those enums and conform to the memset to 0 contract of the reflection system.

diff --git a/Engine/Source/Programs/UnrealHeaderTool/Private/HeaderParser.cpp b/Engine/Source/Programs/UnrealHeaderTool/Private/HeaderParser.cpp
--- a/Engine/Source/Programs/UnrealHeaderTool/Private/HeaderParser.cpp
+++ b/Engine/Source/Programs/UnrealHeaderTool/Private/HeaderParser.cpp
@@ -1784,6 +1784,12 @@ UEnum* FHeaderParser::CompileEnum()
 	}
 
 	CheckDocumentationPolicyForEnum(Enum, EnumValueMetaData, EntryMetaData);
+
+	if (!Enum->IsValidEnumValue(0) && EnumToken.MetaData.Contains(BlueprintTypeName))
+	{
+		UE_LOG_WARNING_UHT(TEXT("'%s' does not have a 0 entry! (This is a problem when the enum is initalized by default)"), *Enum->GetName());
+	}
+
 	return Enum;
 }
 
