* Added more stats to the Unreal Header Tool
* Fixed a nullptr bug in the PropertyDelegate
* Added PropertyDelegate test to the UnrealHeaderTool

diff --git a/Engine/Source/Programs/UnrealHeaderTool/Private/CodeGenerator.cpp b/Engine/Source/Programs/UnrealHeaderTool/Private/CodeGenerator.cpp
--- a/Engine/Source/Programs/UnrealHeaderTool/Private/CodeGenerator.cpp
+++ b/Engine/Source/Programs/UnrealHeaderTool/Private/CodeGenerator.cpp
@@ -5471,15 +5471,18 @@ bool FNativeClassHeaderGenerator::SaveHeaderIfChanged(const TCHAR* HeaderPath, c
 			check(i<10);
 			FString RefHeader;
 			FString Message;
-			if (!FFileHelper::LoadFileToString(RefHeader, *Ref))
 			{
-				Message = FString::Printf(TEXT("********************************* %s appears to be a new generated file."), *FPaths::GetCleanFilename(HeaderPath));
-			}
-			else
-			{
-				if (FCString::Strcmp(NewHeaderContents, *RefHeader) != 0)
+				SCOPE_SECONDS_COUNTER_UHT(LoadHeaderContentFromFile);
+				if (!FFileHelper::LoadFileToString(RefHeader, *Ref))
+				{
+					Message = FString::Printf(TEXT("********************************* %s appears to be a new generated file."), *FPaths::GetCleanFilename(HeaderPath));
+				}
+				else
 				{
-					Message = FString::Printf(TEXT("********************************* %s has changed."), *FPaths::GetCleanFilename(HeaderPath));
+					if (FCString::Strcmp(NewHeaderContents, *RefHeader) != 0)
+					{
+						Message = FString::Printf(TEXT("********************************* %s has changed."), *FPaths::GetCleanFilename(HeaderPath));
+					}
 				}
 			}
 			if (Message.Len())
@@ -5492,7 +5495,10 @@ bool FNativeClassHeaderGenerator::SaveHeaderIfChanged(const TCHAR* HeaderPath, c
 
 
 	FString OriginalHeaderLocal;
-	FFileHelper::LoadFileToString(OriginalHeaderLocal, HeaderPath);
+	{
+		SCOPE_SECONDS_COUNTER_UHT(LoadHeaderContentFromFile);
+		FFileHelper::LoadFileToString(OriginalHeaderLocal, HeaderPath);
+	}
 
 	const bool bHasChanged = OriginalHeaderLocal.Len() == 0 || FCString::Strcmp(*OriginalHeaderLocal, NewHeaderContents);
 	if (bHasChanged)
@@ -5770,9 +5776,12 @@ ECompilationResult::Type PreparseModules(const FString& ModuleInfoPath, int32& N
 					const FString FullFilename = FPaths::ConvertRelativePathToFull(ModuleInfoPath, RawFilename);
 
 					FString HeaderFile;
-					if (!FFileHelper::LoadFileToString(HeaderFile, *FullFilename))
 					{
-						FError::Throwf(TEXT("UnrealHeaderTool was unable to load source file '%s'"), *FullFilename);
+						SCOPE_SECONDS_COUNTER_UHT(LoadHeaderContentFromFile);
+						if (!FFileHelper::LoadFileToString(HeaderFile, *FullFilename))
+						{
+							FError::Throwf(TEXT("UnrealHeaderTool was unable to load source file '%s'"), *FullFilename);
+						}
 					}
 
 					TSharedRef<FUnrealSourceFile> UnrealSourceFile = PerformInitialParseOnHeader(Package, *RawFilename, RF_Public | RF_Standalone, *HeaderFile);
@@ -6033,6 +6042,14 @@ ECompilationResult::Type UnrealHeaderTool_Main(const FString& ModuleInfoFilename
 	UE_LOG(LogCompile, Log, TEXT("Code generation took %.2f seconds"), GHeaderCodeGenTime);
 	UE_LOG(LogCompile, Log, TEXT("ScriptPlugin overhead was %.2f seconds"), GPluginOverheadTime);
 	UE_LOG(LogCompile, Log, TEXT("Macroize time was %.2f seconds"), GMacroizeTime);
+
+	FUnrealHeaderToolStats& Stats = FUnrealHeaderToolStats::Get();
+	for (const TPair<FName, double>& Pair : Stats.Counters)
+	{
+		FString CounterName = Pair.Key.ToString();
+		UE_LOG(LogCompile, Log, TEXT("%s timer was %.3f seconds"), *CounterName, Pair.Value);
+	}
+
 	UE_LOG(LogCompile, Log, TEXT("Total time was %.2f seconds"), MainTime);
 
 	if (bWriteContents)
diff --git a/Engine/Source/Programs/UnrealHeaderTool/Private/HeaderParser.cpp b/Engine/Source/Programs/UnrealHeaderTool/Private/HeaderParser.cpp
--- a/Engine/Source/Programs/UnrealHeaderTool/Private/HeaderParser.cpp
+++ b/Engine/Source/Programs/UnrealHeaderTool/Private/HeaderParser.cpp
@@ -2045,54 +2045,59 @@ FString FHeaderParser::FormatCommentForToolTip(const FString& Input)
 
 TMap<FName, FString> FHeaderParser::GetParameterToolTipsFromFunctionComment(const FString& Input)
 {
+	SCOPE_SECONDS_COUNTER_UHT(DocumentationPolicy);
+
 	TMap<FName, FString> Map;
+	if (Input.IsEmpty())
+	{
+		return Map;
+	}
+	
 	TArray<FString> Params;
+	static const TCHAR ParamTag[] = TEXT("@param");
+	static const TCHAR ReturnTag[] = TEXT("@return");
+	static const TCHAR ReturnParamPrefix[] = TEXT("ReturnValue ");
 
+	/**
+	 * Search for @param / @return followed by a section until a line break.
+	 * For example: "@param Test MyTest Variable" becomes "Test", "MyTest Variable"
+	 * These pairs are then later split and stored as the parameter tooltips.
+	 * Once we don't find either @param or @return we break from the loop.
+	 */
 	int32 Offset = 0;
-	int32 ParamStart = -1;
-	int32 ParamLength = 0;
-	FString ParamPrefix;
 	while (Offset < Input.Len())
 	{
-		if (ParamStart == -1)
+		const TCHAR* ParamPrefix = TEXT("");
+		int32 ParamStart = Input.Find(ParamTag, ESearchCase::CaseSensitive, ESearchDir::FromStart, Offset);
+		if(ParamStart != INDEX_NONE)
+		{
+			ParamStart = ParamStart + ARRAY_COUNT(ParamTag);
+			Offset = ParamStart;
+		}
+		else
 		{
-			if (Input.Find(TEXT("@param"), ESearchCase::CaseSensitive, ESearchDir::FromStart, Offset) == Offset)
+			ParamStart = Input.Find(ReturnTag, ESearchCase::CaseSensitive, ESearchDir::FromStart, Offset);
+			if (ParamStart != INDEX_NONE)
 			{
-				ParamStart = Offset + 6;
-				Offset += 6;
-				continue;
+				ParamStart = ParamStart + ARRAY_COUNT(ReturnTag);
+				Offset = ParamStart;
+				ParamPrefix = ReturnParamPrefix;
 			}
-			else if (Input.Find(TEXT("@return"), ESearchCase::CaseSensitive, ESearchDir::FromStart, Offset) == Offset)
+			else
 			{
-				ParamStart = Offset + 7;
-				Offset += 7;
-				ParamPrefix = TEXT("ReturnValue ");
-				continue;
+				// no @param, no @return?
+				break;
 			}
-			Offset++;
-			continue;
 		}
 
-		TCHAR CurrentChar = Input[Offset];
-		if ((CurrentChar == '\n') || (CurrentChar == '\r'))
+		int32 ParamEnd = Input.Find(TEXT("\n"), ESearchCase::CaseSensitive, ESearchDir::FromStart, ParamStart);
+		if (ParamEnd == INDEX_NONE)
 		{
-			if (ParamLength > 0)
-			{
-				Params.Add(ParamPrefix + Input.Mid(ParamStart, ParamLength));
-				ParamStart = -1;
-				ParamLength = 0;
-			}
-			ParamPrefix.Reset();
+			ParamEnd = Input.Len();
 		}
-		else
-		{
-			ParamLength++;
-		}
-		Offset++;
-	}
-	if (ParamLength > 0)
-	{
-		Params.Add(ParamPrefix + Input.Mid(ParamStart, ParamLength));
+		Offset = ParamEnd;
+
+		Params.Add(ParamPrefix + Input.Mid(ParamStart, ParamEnd - ParamStart - 1));
 	}
 
 	for (FString Param : Params)
@@ -8149,6 +8154,8 @@ ECompilationResult::Type FHeaderParser::ParseAllHeadersInside(
 	TArray<IScriptGeneratorPluginInterface*>& ScriptPlugins
 	)
 {
+	SCOPE_SECONDS_COUNTER_UHT(ParseAllHeaders);
+
 	// Disable loading of objects outside of this package (or more exactly, objects which aren't UFields, CDO, or templates)
 	TGuardValue<bool> AutoRestoreVerifyObjectRefsFlag(GVerifyObjectReferencesOnly, true);
 	// Create the header parser and register it as the warning context.
@@ -9624,6 +9631,8 @@ FDocumentationPolicy FHeaderParser::GetDocumentationPolicyFromName(const FString
 
 FDocumentationPolicy FHeaderParser::GetDocumentationPolicyForStruct(UStruct* Struct)
 {
+	SCOPE_SECONDS_COUNTER_UHT(DocumentationPolicy);
+
 	check(Struct!= nullptr);
 
 	FDocumentationPolicy DocumentationPolicy;
@@ -9637,6 +9646,8 @@ FDocumentationPolicy FHeaderParser::GetDocumentationPolicyForStruct(UStruct* Str
 
 void FHeaderParser::CheckDocumentationPolicyForEnum(UEnum* Enum, const TMap<FName, FString>& MetaData, const TArray<TMap<FName, FString>>& Entries)
 {
+	SCOPE_SECONDS_COUNTER_UHT(DocumentationPolicy);
+
 	check(Enum != nullptr);
 
 	const FString* DocumentationPolicyName = MetaData.Find(NAME_DocumentationPolicy);
@@ -9653,7 +9664,7 @@ void FHeaderParser::CheckDocumentationPolicyForEnum(UEnum* Enum, const TMap<FNam
 		const FString* EnumToolTip = MetaData.Find(NAME_ToolTip);
 		if (EnumToolTip == nullptr)
 		{
-			FError::Throwf(TEXT("Enum '%s' does not provide a tooltip / comment (DocumentationPolicy)."), *Enum->GetName());
+			UE_LOG_ERROR_UHT(TEXT("Enum '%s' does not provide a tooltip / comment (DocumentationPolicy)."), *Enum->GetName());
 		}
 	}
 
@@ -9669,21 +9680,23 @@ void FHeaderParser::CheckDocumentationPolicyForEnum(UEnum* Enum, const TMap<FNam
 		const FString* ToolTip = Entry.Find(NAME_ToolTip);
 		if (ToolTip == nullptr)
 		{
-			FError::Throwf(TEXT("Enum entry '%s::%s' does not provide a tooltip / comment (DocumentationPolicy)."), *Enum->GetName(), **EntryName);
+			UE_LOG_ERROR_UHT(TEXT("Enum entry '%s::%s' does not provide a tooltip / comment (DocumentationPolicy)."), *Enum->GetName(), *(*EntryName));
+			continue;
 		}
 
 		const FString* ExistingEntry = ToolTipToEntry.Find(*ToolTip);
 		if (ExistingEntry != nullptr)
 		{
-			FError::Throwf(TEXT("Enum entries '%s::%s' and '%s::%s' have identical tooltips / comments (DocumentationPolicy)."), *Enum->GetName(), **ExistingEntry, *Enum->GetName(), **EntryName);
+			UE_LOG_ERROR_UHT(TEXT("Enum entries '%s::%s' and '%s::%s' have identical tooltips / comments (DocumentationPolicy)."), *Enum->GetName(), *(*ExistingEntry), *Enum->GetName(), *(*EntryName));
 		}
 		ToolTipToEntry.Add(*ToolTip, *EntryName);
 	}
-	
 }
 
 void FHeaderParser::CheckDocumentationPolicyForStruct(UStruct* Struct, const TMap<FName, FString>& MetaData)
 {
+	SCOPE_SECONDS_COUNTER_UHT(DocumentationPolicy);
+
 	check(Struct != nullptr);
 
 	FDocumentationPolicy DocumentationPolicy = GetDocumentationPolicyForStruct(Struct);
@@ -9692,7 +9705,7 @@ void FHeaderParser::CheckDocumentationPolicyForStruct(UStruct* Struct, const TMa
 		const FString* ClassTooltip = MetaData.Find(NAME_ToolTip);
 		if (ClassTooltip == nullptr)
 		{
-			FError::Throwf(TEXT("Struct '%s' does not provide a tooltip / comment (DocumentationPolicy)."), *Struct->GetName());
+			UE_LOG_ERROR_UHT(TEXT("Struct '%s' does not provide a tooltip / comment (DocumentationPolicy)."), *Struct->GetName());
 		}
 	}
 
@@ -9704,12 +9717,13 @@ void FHeaderParser::CheckDocumentationPolicyForStruct(UStruct* Struct, const TMa
 			FString ToolTip = Property->GetToolTipText().ToString();
 			if (ToolTip.IsEmpty() || ToolTip.Equals(Property->GetName()))
 			{
-				FError::Throwf(TEXT("Property '%s::%s' does not provide a tooltip / comment (DocumentationPolicy)."), *Struct->GetName(), *Property->GetName());
+				UE_LOG_ERROR_UHT(TEXT("Property '%s::%s' does not provide a tooltip / comment (DocumentationPolicy)."), *Struct->GetName(), *Property->GetName());
+				continue;
 			}
 			const FName* ExistingPropertyName = ToolTipToPropertyName.Find(ToolTip);
 			if (ExistingPropertyName != nullptr)
 			{
-				FError::Throwf(TEXT("Property '%s::%s' and '%s::%s' are using identical tooltips (DocumentationPolicy)."), *Struct->GetName(), *ExistingPropertyName->ToString(), *Struct->GetName(), *Property->GetName());
+				UE_LOG_ERROR_UHT(TEXT("Property '%s::%s' and '%s::%s' are using identical tooltips (DocumentationPolicy)."), *Struct->GetName(), *ExistingPropertyName->ToString(), *Struct->GetName(), *Property->GetName());
 			}
 			ToolTipToPropertyName.Add(ToolTip, Property->GetFName());
 		}
@@ -9723,9 +9737,10 @@ void FHeaderParser::CheckDocumentationPolicyForStruct(UStruct* Struct, const TMa
 			{
 				const FString& UIMin = Property->GetMetaData(TEXT("UIMin"));
 				const FString& UIMax = Property->GetMetaData(TEXT("UIMax"));
+
 				if(!CheckUIMinMaxRangeFromMetaData(UIMin, UIMax))
 				{
-					FError::Throwf(TEXT("Property '%s::%s' does not provide a valid UIMin / UIMax (DocumentationPolicy)."), *Struct->GetName(), *Property->GetName());
+					UE_LOG_ERROR_UHT(TEXT("Property '%s::%s' does not provide a valid UIMin / UIMax (DocumentationPolicy)."), *Struct->GetName(), *Property->GetName());
 				}
 			}
 		}
@@ -9746,12 +9761,13 @@ void FHeaderParser::CheckDocumentationPolicyForStruct(UStruct* Struct, const TMa
 					FString ToolTip = Func->GetToolTipText().ToString();
 					if (ToolTip.IsEmpty())
 					{
-						FError::Throwf(TEXT("Function '%s::%s' does not provide a tooltip / comment (DocumentationPolicy)."), *Class->GetName(), *Func->GetName());
+						UE_LOG_ERROR_UHT(TEXT("Function '%s::%s' does not provide a tooltip / comment (DocumentationPolicy)."), *Class->GetName(), *Func->GetName());
+						continue;
 					}
 					const FString* ExistingFuncName = ToolTipToFunc.Find(ToolTip);
 					if (ExistingFuncName != nullptr)
 					{
-						FError::Throwf(TEXT("Functions '%s::%s' and '%s::%s' uses identical tooltips / comments (DocumentationPolicy)."), *Class->GetName(), **ExistingFuncName, *Class->GetName(), *Func->GetName());
+						UE_LOG_ERROR_UHT(TEXT("Functions '%s::%s' and '%s::%s' uses identical tooltips / comments (DocumentationPolicy)."), *Class->GetName(), *(*ExistingFuncName), *Class->GetName(), *Func->GetName());
 					}
 					ToolTipToFunc.Add(ToolTip, Func->GetName());
 				}
@@ -9762,12 +9778,14 @@ void FHeaderParser::CheckDocumentationPolicyForStruct(UStruct* Struct, const TMa
 
 bool FHeaderParser::DoesCPPTypeRequireDocumentation(const FString& CPPType)
 {
-	return PropertyCPPTypesRequiringUIRanges.Find(CPPType) != INDEX_NONE;
+	return PropertyCPPTypesRequiringUIRanges.Contains(CPPType);
 }
 
 // Validates the documentation for a given method
 void FHeaderParser::CheckDocumentationPolicyForFunc(UClass* Class, UFunction* Func, const TMap<FName, FString>& MetaData)
 {
+	SCOPE_SECONDS_COUNTER_UHT(DocumentationPolicy);
+
 	check(Class != nullptr);
 	check(Func != nullptr);
 
@@ -9777,7 +9795,7 @@ void FHeaderParser::CheckDocumentationPolicyForFunc(UClass* Class, UFunction* Fu
 		const FString* FunctionTooltip = MetaData.Find(NAME_ToolTip);
 		if (FunctionTooltip == nullptr)
 		{
-			FError::Throwf(TEXT("Function '%s::%s' does not provide a tooltip / comment (DocumentationPolicy)."), *Class->GetName(), *Func->GetName());
+			UE_LOG_ERROR_UHT(TEXT("Function '%s::%s' does not provide a tooltip / comment (DocumentationPolicy)."), *Class->GetName(), *Func->GetName());
 		}
 	}
 
@@ -9786,7 +9804,8 @@ void FHeaderParser::CheckDocumentationPolicyForFunc(UClass* Class, UFunction* Fu
 		const FString* FunctionComment = MetaData.Find(NAME_Comment);
 		if (FunctionComment == nullptr)
 		{
-			FError::Throwf(TEXT("Function '%s::%s' does not provide a comment (DocumentationPolicy)."), *Class->GetName(), *Func->GetName());
+			UE_LOG_ERROR_UHT(TEXT("Function '%s::%s' does not provide a comment (DocumentationPolicy)."), *Class->GetName(), *Func->GetName());
+			return;
 		}
 		
 		TMap<FName, FString> ParamToolTips = GetParameterToolTipsFromFunctionComment(*FunctionComment);
@@ -9815,7 +9834,7 @@ void FHeaderParser::CheckDocumentationPolicyForFunc(UClass* Class, UFunction* Fu
 				const FString* ParamToolTip = ParamToolTips.Find(ParamName);
 				if (ParamToolTip == nullptr)
 				{
-					FError::Throwf(TEXT("Function '%s::%s' doesn't provide a tooltip for parameter '%s' (DocumentationPolicy)."), *Class->GetName(), *Func->GetName(), *ParamName.ToString());
+					UE_LOG_ERROR_UHT(TEXT("Function '%s::%s' doesn't provide a tooltip for parameter '%s' (DocumentationPolicy)."), *Class->GetName(), *Func->GetName(), *ParamName.ToString());
 				}
 				ExistingFields.Add(ParamName);
 			}
@@ -9830,7 +9849,7 @@ void FHeaderParser::CheckDocumentationPolicyForFunc(UClass* Class, UFunction* Fu
 				}
 				if (!ExistingFields.Contains(ParamName))
 				{
-					FError::Throwf(TEXT("Function '%s::%s' provides a tooltip for an unknown parameter '%s' (DocumentationPolicy)."), *Class->GetName(), *Func->GetName(), *Pair.Key.ToString());
+					UE_LOG_ERROR_UHT(TEXT("Function '%s::%s' provides a tooltip for an unknown parameter '%s' (DocumentationPolicy)."), *Class->GetName(), *Func->GetName(), *Pair.Key.ToString());
 				}
 			}
 
@@ -9846,7 +9865,7 @@ void FHeaderParser::CheckDocumentationPolicyForFunc(UClass* Class, UFunction* Fu
 				const FName* ExistingParam = ToolTipToParam.Find(Pair.Value);
 				if (ExistingParam != nullptr)
 				{
-					FError::Throwf(TEXT("Function '%s::%s' uses identical tooltips for parameters '%s' and '%s' (DocumentationPolicy)."), *Class->GetName(), *Func->GetName(), *ExistingParam->ToString(), *Pair.Key.ToString());
+					UE_LOG_ERROR_UHT(TEXT("Function '%s::%s' uses identical tooltips for parameters '%s' and '%s' (DocumentationPolicy)."), *Class->GetName(), *Func->GetName(), *ExistingParam->ToString(), *Pair.Key.ToString());
 				}
 				ToolTipToParam.Add(Pair.Value, Pair.Key);
 			}
diff --git a/Engine/Source/Programs/UnrealHeaderTool/Private/HeaderParser.h b/Engine/Source/Programs/UnrealHeaderTool/Private/HeaderParser.h
--- a/Engine/Source/Programs/UnrealHeaderTool/Private/HeaderParser.h
+++ b/Engine/Source/Programs/UnrealHeaderTool/Private/HeaderParser.h
@@ -147,13 +147,17 @@ struct ClassDefinitionRange
 
 extern TMap<UClass*, ClassDefinitionRange> ClassDefinitionRanges;
 
+#ifndef UHT_DOCUMENTATION_POLICY_DEFAULT
+#define UHT_DOCUMENTATION_POLICY_DEFAULT false
+#endif
+
 struct FDocumentationPolicy
 {
-	bool bClassOrStructCommentRequired = false;
-	bool bFunctionToolTipsRequired = false;
-	bool bMemberToolTipsRequired = false;
-	bool bParameterToolTipsRequired = false;
-	bool bFloatRangesRequired = false;
+	bool bClassOrStructCommentRequired = UHT_DOCUMENTATION_POLICY_DEFAULT;
+	bool bFunctionToolTipsRequired = UHT_DOCUMENTATION_POLICY_DEFAULT;
+	bool bMemberToolTipsRequired = UHT_DOCUMENTATION_POLICY_DEFAULT;
+	bool bParameterToolTipsRequired = UHT_DOCUMENTATION_POLICY_DEFAULT;
+	bool bFloatRangesRequired = UHT_DOCUMENTATION_POLICY_DEFAULT;
 };
 
 /////////////////////////////////////////////////////
diff --git a/Engine/Source/Programs/UnrealHeaderTool/Private/UnrealHeaderToolGlobals.h b/Engine/Source/Programs/UnrealHeaderTool/Private/UnrealHeaderToolGlobals.h
--- a/Engine/Source/Programs/UnrealHeaderTool/Private/UnrealHeaderToolGlobals.h
+++ b/Engine/Source/Programs/UnrealHeaderTool/Private/UnrealHeaderToolGlobals.h
@@ -5,6 +5,7 @@
 #include "CoreMinimal.h"
 #include "Logging/LogCategory.h"
 #include "Logging/LogMacros.h"
+#include "UnrealHeaderToolStats.h"
 
 DECLARE_LOG_CATEGORY_EXTERN(LogCompile, Log, All);
 
diff --git a/Engine/Source/Programs/UnrealHeaderTool/Private/UnrealHeaderToolStats.cpp b/Engine/Source/Programs/UnrealHeaderTool/Private/UnrealHeaderToolStats.cpp
new file mode 100644
--- /dev/null
+++ b/Engine/Source/Programs/UnrealHeaderTool/Private/UnrealHeaderToolStats.cpp
@@ -0,0 +1,15 @@
+// Copyright 1998-2019 Epic Games, Inc. All Rights Reserved.
+
+#include "UnrealHeaderToolStats.h"
+
+FUnrealHeaderToolStats GUnrealHeaderToolStats;
+
+FUnrealHeaderToolStats& FUnrealHeaderToolStats::Get()
+{
+	return GUnrealHeaderToolStats;
+}
+
+double& FUnrealHeaderToolStats::GetCounter(const FName& Key)
+{
+	return GUnrealHeaderToolStats.Counters.FindOrAdd(Key);
+}
diff --git a/Engine/Source/Programs/UnrealHeaderTool/Private/UnrealHeaderToolStats.h b/Engine/Source/Programs/UnrealHeaderTool/Private/UnrealHeaderToolStats.h
new file mode 100644
--- /dev/null
+++ b/Engine/Source/Programs/UnrealHeaderTool/Private/UnrealHeaderToolStats.h
@@ -0,0 +1,21 @@
+// Copyright 1998-2019 Epic Games, Inc. All Rights Reserved.
+
+#pragma once
+
+#include "CoreMinimal.h"
+#include "Stats/StatsMisc.h"
+
+struct FUnrealHeaderToolStats
+{
+	TMap<FName, double> Counters;
+
+	static FUnrealHeaderToolStats& Get();
+	static double& GetCounter(const FName& Key);
+};
+
+#if STATS
+#define SCOPE_SECONDS_COUNTER_UHT(Name) \
+	FSimpleScopeSecondsCounter ANONYMOUS_VARIABLE(Name)(FUnrealHeaderToolStats::GetCounter(TEXT(#Name)))
+#else
+#define SCOPE_SECONDS_COUNTER_UHT(Name)
+#endif
diff --git a/Engine/Source/Programs/UnrealHeaderTool/Resources/UHTDebugging/StructOnlyHeader.h b/Engine/Source/Programs/UnrealHeaderTool/Resources/UHTDebugging/StructOnlyHeader.h
--- a/Engine/Source/Programs/UnrealHeaderTool/Resources/UHTDebugging/StructOnlyHeader.h
+++ b/Engine/Source/Programs/UnrealHeaderTool/Resources/UHTDebugging/StructOnlyHeader.h
@@ -4,6 +4,8 @@
 
 #include "StructOnlyHeader.generated.h"
 
+DECLARE_DYNAMIC_DELEGATE(FSimpleStructDelegate);
+
 USTRUCT()
 struct FSomeStruct
 {
@@ -18,4 +20,7 @@ USTRUCT()
 struct alignas(8) FAlignedStruct
 {
 	GENERATED_BODY()
+
+	UPROPERTY()
+	FSimpleStructDelegate DelegateProp;
 };
diff --git a/Engine/Source/Programs/UnrealHeaderTool/Resources/UHTDebugging/TestObject.h b/Engine/Source/Programs/UnrealHeaderTool/Resources/UHTDebugging/TestObject.h
--- a/Engine/Source/Programs/UnrealHeaderTool/Resources/UHTDebugging/TestObject.h
+++ b/Engine/Source/Programs/UnrealHeaderTool/Resources/UHTDebugging/TestObject.h
@@ -7,6 +7,7 @@
 
 PRAGMA_DISABLE_DEPRECATION_WARNINGS
 
+DECLARE_DYNAMIC_DELEGATE(FSimpleClassDelegate);
 DECLARE_DYNAMIC_DELEGATE_OneParam(FRegularDelegate, int32, SomeArgument);
 DECLARE_DYNAMIC_DELEGATE_OneParam(FDelegateWithDelegateParam, FRegularDelegate const &, RegularDelegate);
 
@@ -88,6 +89,9 @@ public:
 	UPROPERTY()
 	UObject* const ConstPointerProperty;
 
+	UPROPERTY()
+	FSimpleClassDelegate DelegateProperty;
+
 	UFUNCTION()
 	void CodeGenTestForEnumClasses(ECppEnum Val);
 
@@ -137,4 +141,29 @@ public:
 	UEnum* SomeFunc() const;
 };
 
+/**
+ * A test objec to validate the DocumentationPolicy validation.
+ */
+UCLASS(meta = (DocumentationPolicy = "Strict"))
+class UDocumentedTestObject : public UObject
+{
+	GENERATED_BODY()
+
+public:
+
+	/**
+	 * A float member on this class.
+	 */
+	UPROPERTY(meta = (UIMin = "0.0", UIMax = "1.0"))
+	float Member;
+
+	/**
+	 * Tests the documentation policy
+	 * @param bFlag If set to true, a flag is set to true
+	 * @param Range The range of the results.
+	 */
+	UFUNCTION()
+	void TestFunction(bool bFlag, float Range);
+}
+
 PRAGMA_ENABLE_DEPRECATION_WARNINGS
diff --git a/Engine/Source/Runtime/CoreUObject/Private/UObject/PropertyDelegate.cpp b/Engine/Source/Runtime/CoreUObject/Private/UObject/PropertyDelegate.cpp
--- a/Engine/Source/Runtime/CoreUObject/Private/UObject/PropertyDelegate.cpp
+++ b/Engine/Source/Runtime/CoreUObject/Private/UObject/PropertyDelegate.cpp
@@ -74,6 +74,8 @@ bool UDelegateProperty::NetSerializeItem( FArchive& Ar, UPackageMap* Map, void*
 
 FString UDelegateProperty::GetCPPType( FString* ExtendedTypeText/*=NULL*/, uint32 CPPExportFlags/*=0*/ ) const
 {
+	check(SignatureFunction);
+
 	FString UnmangledFunctionName = SignatureFunction->GetName().LeftChop( FString( HEADER_GENERATED_DELEGATE_SIGNATURE_SUFFIX ).Len() );
 	const bool bBlueprintCppBackend = (0 != (CPPExportFlags & EPropertyExportCPPFlags::CPPF_BlueprintCppBackend));
 	const bool bNative = SignatureFunction->IsNative();
