Ensure to only parse the direct fields below a struct - not the super fields.

diff --git a/Engine/Source/Programs/UnrealHeaderTool/Private/HeaderParser.cpp b/Engine/Source/Programs/UnrealHeaderTool/Private/HeaderParser.cpp
--- a/Engine/Source/Programs/UnrealHeaderTool/Private/HeaderParser.cpp
+++ b/Engine/Source/Programs/UnrealHeaderTool/Private/HeaderParser.cpp
@@ -9712,7 +9712,7 @@ void FHeaderParser::CheckDocumentationPolicyForStruct(UStruct* Struct, const TMa
 	if (DocumentationPolicy.bMemberToolTipsRequired)
 	{
 		TMap<FString, FName> ToolTipToPropertyName;
-		for (UProperty* Property : TFieldRange<UProperty>(Struct))
+		for (UProperty* Property : TFieldRange<UProperty>(Struct, EFieldIteratorFlags::ExcludeSuper))
 		{
 			FString ToolTip = Property->GetToolTipText().ToString();
 			if (ToolTip.IsEmpty() || ToolTip.Equals(Property->GetName()))
@@ -9731,7 +9731,7 @@ void FHeaderParser::CheckDocumentationPolicyForStruct(UStruct* Struct, const TMa
 
 	if (DocumentationPolicy.bFloatRangesRequired)
 	{
-		for (UProperty* Property : TFieldRange<UProperty>(Struct))
+		for (UProperty* Property : TFieldRange<UProperty>(Struct, EFieldIteratorFlags::ExcludeSuper))
 		{
 			if(DoesCPPTypeRequireDocumentation(Property->GetCPPType()))
 			{
diff --git a/Engine/Source/Programs/UnrealHeaderTool/Resources/UHTDebugging/TestObject.h b/Engine/Source/Programs/UnrealHeaderTool/Resources/UHTDebugging/TestObject.h
--- a/Engine/Source/Programs/UnrealHeaderTool/Resources/UHTDebugging/TestObject.h
+++ b/Engine/Source/Programs/UnrealHeaderTool/Resources/UHTDebugging/TestObject.h
@@ -141,11 +141,27 @@ public:
 	UEnum* SomeFunc() const;
 };
 
+UCLASS()
+class UDocumentedBaseObject : public UObject
+{
+	GENERATED_BODY()
+
+	virtual void UndocumentedMethod();
+
+	// Duplicate tooltip
+	UPROPERTY()
+	bool bFlagA;
+
+	// Duplicate tooltip
+	UPROPERTY()
+	bool bFlagB;
+};
+
 /**
- * A test objec to validate the DocumentationPolicy validation.
+ * A test object to validate the DocumentationPolicy validation.
  */
 UCLASS(meta = (DocumentationPolicy = "Strict"))
-class UDocumentedTestObject : public UObject
+class UDocumentedTestObject : public UDocumentedBaseObject
 {
 	GENERATED_BODY()
 
@@ -164,6 +180,6 @@ public:
 	 */
 	UFUNCTION()
 	void TestFunction(bool bFlag, float Range);
-}
+};
 
 PRAGMA_ENABLE_DEPRECATION_WARNINGS
