Make sure to perform the documentation checking at the right time (after all statements have been parsed)

diff --git a/Engine/Source/Programs/UnrealHeaderTool/Private/HeaderParser.cpp b/Engine/Source/Programs/UnrealHeaderTool/Private/HeaderParser.cpp
--- a/Engine/Source/Programs/UnrealHeaderTool/Private/HeaderParser.cpp
+++ b/Engine/Source/Programs/UnrealHeaderTool/Private/HeaderParser.cpp
@@ -2638,7 +2638,6 @@ UScriptStruct* FHeaderParser::CompileStructDeclaration(FClasses& AllClasses)
 	// Link the properties within the struct
 	Struct->StaticLink(true);
 
-	CheckDocumentationPolicyForStruct(Struct, MetaData);
 	return Struct;
 }
 
@@ -2998,6 +2997,10 @@ void FHeaderParser::FixupDelegateProperties( FClasses& AllClasses, UStruct* Stru
 			}
 		}
 	}
+
+	TMap<FName, FString> MetaData;
+	MetaData.Add(NAME_ToolTip, Struct->GetMetaData(NAME_ToolTip));
+	CheckDocumentationPolicyForStruct(Struct, MetaData);
 }
 
 void FHeaderParser::VerifyBlueprintPropertyGetter(UProperty* Prop, UFunction* TargetFunc)
@@ -5961,7 +5964,6 @@ UClass* FHeaderParser::CompileClassDeclaration(FClasses& AllClasses)
 		}
 	}
 
-	CheckDocumentationPolicyForStruct(Class, MetaData);
 	return Class;
 }
 
@@ -9753,24 +9755,20 @@ void FHeaderParser::CheckDocumentationPolicyForStruct(UStruct* Struct, const TMa
 		if (Class != nullptr)
 		{
 			TMap<FString, FString> ToolTipToFunc;
-			for (UField* Field = Class->Children; Field; Field = Field->Next)
+			for (UFunction* Func : TFieldRange<UFunction>(Class, EFieldIteratorFlags::ExcludeSuper))
 			{
-				UFunction* Func = Cast<UFunction>(Field);
-				if (Func != NULL)
+				FString ToolTip = Func->GetToolTipText().ToString();
+				if (ToolTip.IsEmpty())
 				{
-					FString ToolTip = Func->GetToolTipText().ToString();
-					if (ToolTip.IsEmpty())
-					{
-						UE_LOG_ERROR_UHT(TEXT("Function '%s::%s' does not provide a tooltip / comment (DocumentationPolicy)."), *Class->GetName(), *Func->GetName());
-						continue;
-					}
-					const FString* ExistingFuncName = ToolTipToFunc.Find(ToolTip);
-					if (ExistingFuncName != nullptr)
-					{
-						UE_LOG_ERROR_UHT(TEXT("Functions '%s::%s' and '%s::%s' uses identical tooltips / comments (DocumentationPolicy)."), *Class->GetName(), *(*ExistingFuncName), *Class->GetName(), *Func->GetName());
-					}
-					ToolTipToFunc.Add(ToolTip, Func->GetName());
+					UE_LOG_ERROR_UHT(TEXT("Function '%s::%s' does not provide a tooltip / comment (DocumentationPolicy)."), *Class->GetName(), *Func->GetName());
+					continue;
+				}
+				const FString* ExistingFuncName = ToolTipToFunc.Find(ToolTip);
+				if (ExistingFuncName != nullptr)
+				{
+					UE_LOG_ERROR_UHT(TEXT("Functions '%s::%s' and '%s::%s' uses identical tooltips / comments (DocumentationPolicy)."), *Class->GetName(), *(*ExistingFuncName), *Class->GetName(), *Func->GetName());
 				}
+				ToolTipToFunc.Add(ToolTip, Func->GetName());
 			}
 		}
 	}
diff --git a/Engine/Source/Programs/UnrealHeaderTool/Resources/UHTDebugging/TestObject.h b/Engine/Source/Programs/UnrealHeaderTool/Resources/UHTDebugging/TestObject.h
--- a/Engine/Source/Programs/UnrealHeaderTool/Resources/UHTDebugging/TestObject.h
+++ b/Engine/Source/Programs/UnrealHeaderTool/Resources/UHTDebugging/TestObject.h
@@ -146,6 +146,9 @@ class UDocumentedBaseObject : public UObject
 {
 	GENERATED_BODY()
 
+public:
+
+	UFUNCTION()
 	virtual void UndocumentedMethod();
 
 	// Duplicate tooltip
@@ -180,6 +183,14 @@ public:
 	 */
 	UFUNCTION()
 	void TestFunction(bool bFlag, float Range);
+
+	/**
+	 * Tests the documentation policy (2)
+	 * @param bFlag If set to true, a flag is set to true
+	 * @param Range The range of the results.
+	 */
+	UFUNCTION()
+	void TestFunction2(bool bFlag, float Range);
 };
 
 PRAGMA_ENABLE_DEPRECATION_WARNINGS
