Improved data validation for sparse class data types.
UHT now checks for the following:
 - max of one sparse class data struct per class
 - sparse class data struct must inherit from the parent classes sparse class data struct
 - all sparse properties must be BlueprintReadOnly and EditDefaultsOnly
 - BP assignable delegates are not allowed as sparse properties

diff --git a/Engine/Source/Programs/UnrealHeaderTool/Private/HeaderParser.cpp b/Engine/Source/Programs/UnrealHeaderTool/Private/HeaderParser.cpp
--- a/Engine/Source/Programs/UnrealHeaderTool/Private/HeaderParser.cpp
+++ b/Engine/Source/Programs/UnrealHeaderTool/Private/HeaderParser.cpp
@@ -3016,6 +3016,87 @@ void FHeaderParser::FixupDelegateProperties( FClasses& AllClasses, UStruct* Stru
 	ParseRigVMMethodParameters(Struct);
 }
 
+void FHeaderParser::CheckSparseClassData(const UStruct* StructToCheck)
+{
+	// we're looking for classes that have sparse class data structures
+	const UClass* ClassToCheck = Cast<const UClass>(StructToCheck);
+	if (!ClassToCheck)
+	{
+		// make sure we don't try to have sparse class data inside of a struct instead of a class
+		if (StructToCheck->HasMetaData(TEXT("SparseClassDataTypes")))
+		{
+			UE_LOG_ERROR_UHT(TEXT("%s contains sparse class data but is not a class."), *StructToCheck->GetName());
+		}
+		return;
+	}
+
+	if (!ClassToCheck->HasMetaData(TEXT("SparseClassDataTypes")))
+	{
+		return;
+	}
+
+	TArray<FString> SparseClassDataTypes;
+	((FClass*)ClassToCheck)->GetSparseClassDataTypes(SparseClassDataTypes);
+
+	// for now we only support one sparse class data structure per class
+	if (SparseClassDataTypes.Num() > 1)
+	{
+		UE_LOG_ERROR_UHT(TEXT("Class %s contains multiple sparse class data types."), *ClassToCheck->GetName());
+		return;
+	}
+	if (SparseClassDataTypes.Num() == 0)
+	{
+		UE_LOG_ERROR_UHT(TEXT("Class %s has sparse class metadata but does not specify a type."), *ClassToCheck->GetName());
+		return;
+	}
+
+	for (const FString& SparseClassDataTypeName : SparseClassDataTypes)
+	{
+		UScriptStruct* SparseClassDataStruct = FindObjectSafe<UScriptStruct>(ANY_PACKAGE, *SparseClassDataTypeName);
+
+		// make sure the sparse class data struct actually exists
+		if (!SparseClassDataStruct)
+		{
+			UE_LOG_ERROR_UHT(TEXT("Unable to find sparse data type %s for class %s."), *SparseClassDataTypeName, *ClassToCheck->GetName());
+			return;
+		}
+
+		// check the data struct for invalid properties
+		for (TFieldIterator<UProperty> Property(SparseClassDataStruct); Property; ++Property)
+		{
+			if (Property->HasAnyPropertyFlags(CPF_BlueprintAssignable))
+			{
+				UE_LOG_ERROR_UHT(TEXT("Sparse class data types can not contain blueprint assignable delegates. Type '%s' Delegate '%s'"), *SparseClassDataStruct->GetName(), *Property->GetName());
+			}
+
+			// all sparse properties should have EditDefaultsOnly
+			if (!Property->HasAllPropertyFlags(CPF_Edit | CPF_DisableEditOnInstance))
+			{
+				UE_LOG_ERROR_UHT(TEXT("Sparse class data types must be EditDefaultsOnly. Type '%s' Property '%s'"), *SparseClassDataStruct->GetName(), *Property->GetName());
+			}
+
+			// no sparse properties should have BlueprintReadWrite
+			if (Property->HasAllPropertyFlags(CPF_BlueprintVisible) && !Property->HasAllPropertyFlags(CPF_BlueprintReadOnly))
+			{
+				UE_LOG_ERROR_UHT(TEXT("Sparse class data types must not be BlueprintReadWrite. Type '%s' Property '%s'"), *SparseClassDataStruct->GetName(), *Property->GetName());
+			}
+		}
+
+		// if the class's parent has a sparse class data struct then the current class must also use the same struct or one that inherits from it
+		const UClass* ParentClass = ClassToCheck->GetSuperClass();
+		TArray<FString> ParentSparseClassDataTypeNames;
+		((FClass*)ParentClass)->GetSparseClassDataTypes(ParentSparseClassDataTypeNames);
+		for (FString& ParentSparseClassDataTypeName : ParentSparseClassDataTypeNames)
+		{
+			UScriptStruct* ParentSparseClassDataStruct = FindObjectSafe<UScriptStruct>(ANY_PACKAGE, *ParentSparseClassDataTypeName);
+			if (ParentSparseClassDataStruct && !SparseClassDataStruct->IsChildOf(ParentSparseClassDataStruct))
+			{
+				UE_LOG_ERROR_UHT(TEXT("Class %s is a child of %s but its sparse class data struct, %s, does not inherit from %s."), *ClassToCheck->GetName(), *ParentClass->GetName(), *SparseClassDataStruct->GetName(), *ParentSparseClassDataStruct->GetName());
+			}
+		}
+	}
+}
+
 void FHeaderParser::VerifyBlueprintPropertyGetter(UProperty* Prop, UFunction* TargetFunc)
 {
 	check(TargetFunc);
@@ -8420,9 +8501,16 @@ ECompilationResult::Type FHeaderParser::ParseAllHeadersInside(
 			Result = FHeaderParser::ParseRestOfModulesSourceFiles(ModuleClasses, CurrentPackage, HeaderParser);
 		}
 
-		// Export the autogenerated code wrappers
 		if (Result == ECompilationResult::Succeeded)
 		{
+			// Validate the sparse class data for all classes in the current package
+			for (const FClass* Class : ModuleClasses.GetClassesInPackage(CurrentPackage))
+			{
+				CheckSparseClassData(Class);
+			}
+
+			// Export the autogenerated code wrappers
+
 			// At this point all headers have been parsed and the header parser will
 			// no longer have up to date info about what's being done so unregister it 
 			// from the feedback context.
diff --git a/Engine/Source/Programs/UnrealHeaderTool/Private/HeaderParser.h b/Engine/Source/Programs/UnrealHeaderTool/Private/HeaderParser.h
--- a/Engine/Source/Programs/UnrealHeaderTool/Private/HeaderParser.h
+++ b/Engine/Source/Programs/UnrealHeaderTool/Private/HeaderParser.h
@@ -698,6 +698,9 @@ public:
 	*/
 	static UField* FindField( UStruct* InScope, const TCHAR* InIdentifier, bool bIncludeParents=true, UClass* FieldClass=UField::StaticClass(), const TCHAR* Thing=nullptr );
 
+	// Checks ToValidate to make sure that its associated sparse class data struct, if one exists, is a valid structure to use for storing sparse class data.
+	static void CheckSparseClassData(const UStruct* ToValidate);
+
 protected:
 
 	/**
