Remove a redundant ensure.

diff --git a/Engine/Source/Programs/UnrealHeaderTool/Private/CodeGenerator.cpp b/Engine/Source/Programs/UnrealHeaderTool/Private/CodeGenerator.cpp
--- a/Engine/Source/Programs/UnrealHeaderTool/Private/CodeGenerator.cpp
+++ b/Engine/Source/Programs/UnrealHeaderTool/Private/CodeGenerator.cpp
@@ -2285,7 +2285,7 @@ void FNativeClassHeaderGenerator::ExportNativeGeneratedInitCode(FOutputDevice& O
 
 		TArray<FString> SparseClassDataTypes;
 		((FClass*)Class)->GetSparseClassDataTypes(SparseClassDataTypes);
-		ensureMsgf(SparseClassDataTypes.Num() <= 1, TEXT("We don't currently support multiple sparse class data structures on a single object."));
+		
 		for (const FString& SparseClassDataString : SparseClassDataTypes)
 		{
 			GeneratedClassRegisterFunctionText.Logf(TEXT("\t\t\tOuterClass->SetSparseClassDataStruct(F%s::StaticStruct());\r\n"), *SparseClassDataString);
