Added info about the sparse class data structure into the class info hash used to detect changes during hot reloads.

diff --git a/Engine/Source/Programs/UnrealHeaderTool/Private/CodeGenerator.cpp b/Engine/Source/Programs/UnrealHeaderTool/Private/CodeGenerator.cpp
--- a/Engine/Source/Programs/UnrealHeaderTool/Private/CodeGenerator.cpp
+++ b/Engine/Source/Programs/UnrealHeaderTool/Private/CodeGenerator.cpp
@@ -2334,6 +2334,24 @@ void FNativeClassHeaderGenerator::ExportNativeGeneratedInitCode(FOutputDevice& O
 	}
 	GeneratedClassRegisterFunctionText.Logf(TEXT("\r\n// %u\r\n"), BaseClassHash);
 
+	// Append info for the sparse class data struct onto the text to be hashed
+	TArray<FString> SparseClassDataTypes;
+	((FClass*)Class)->GetSparseClassDataTypes(SparseClassDataTypes);
+
+	for (const FString& SparseClassDataString : SparseClassDataTypes)
+	{
+		UScriptStruct* SparseClassDataStruct = FindObjectSafe<UScriptStruct>(ANY_PACKAGE, *SparseClassDataString);
+		if (!SparseClassDataStruct)
+		{
+			continue;
+		}
+		GeneratedClassRegisterFunctionText.Logf(TEXT("%s\r\n"), *SparseClassDataStruct->GetName());
+		for (UProperty* Child : TFieldRange<UProperty>(SparseClassDataStruct))
+		{
+			GeneratedClassRegisterFunctionText.Logf(TEXT("%s %s\r\n"), *Child->GetCPPType(), *Child->GetNameCPP());
+		}
+	}
+
 	// Calculate generated class initialization code hash so that we know when it changes after hot-reload
 	uint32 ClassHash = GenerateTextHash(*GeneratedClassRegisterFunctionText);
 	GGeneratedCodeHashes.Add(Class, ClassHash);
