Cleanup GetCPPName to properly reuse interface names rather than creating a new one each time.
Reduce some unnecessary string allocations
Do not export __PPO__ private access functions for properties that are not BlueprintVisible
#rnx

diff --git a/Engine/Source/Programs/UnrealHeaderTool/Private/CodeGenerator.cpp b/Engine/Source/Programs/UnrealHeaderTool/Private/CodeGenerator.cpp
--- a/Engine/Source/Programs/UnrealHeaderTool/Private/CodeGenerator.cpp
+++ b/Engine/Source/Programs/UnrealHeaderTool/Private/CodeGenerator.cpp
@@ -888,15 +888,17 @@ private:
 	{
 		check(Item);
 
-		FString Suffix;
+		bool bNoRegister = false;
 		if (UClass* ItemClass = Cast<UClass>(Item))
 		{
 			if (!bRequiresValidObject && !ItemClass->HasAllClassFlags(CLASS_Intrinsic))
 			{
-				Suffix = TEXT("_NoRegister");
+				bNoRegister = true;
 			}
 		}
 
+		const TCHAR* Suffix = (bNoRegister ? TEXT("_NoRegister") : TEXT(""));
+
 		FString Result;
 		for (UObject* Outer = Item; Outer; Outer = Outer->GetOuter())
 		{
@@ -907,8 +909,8 @@ private:
 
 			if (Cast<UClass>(Outer) || Cast<UScriptStruct>(Outer))
 			{
-				FString OuterName = NameLookupCPP.GetNameCPP(Cast<UStruct>(Outer));
-				Result = OuterName + Result;
+				const TCHAR* OuterName = NameLookupCPP.GetNameCPP(Cast<UStruct>(Outer));
+				Result = FString::Printf(TEXT("%s%s"), OuterName, *Result);
 
 				// Structs can also have UPackage outer.
 				if (Cast<UClass>(Outer) || Cast<UPackage>(Outer->GetOuter()))
@@ -928,8 +930,8 @@ private:
 			Result = FPackageName::GetShortName(Result);
 		}
 
-		FString ClassString = NameLookupCPP.GetNameCPP(Item->GetClass());
-		return FString(TEXT("Z_Construct_")) + ClassString + TEXT("_") + Result + Suffix + TEXT("()");
+		const TCHAR* ClassString = NameLookupCPP.GetNameCPP(Item->GetClass());
+		return FString::Printf(TEXT("Z_Construct_%s_%s%s()"), ClassString, *Result, Suffix);
 	}
 };
 
@@ -1948,9 +1950,9 @@ static TArray<UScriptStruct*> FindNoExportStructs(UStruct* Start)
 
 FString FNativeClassHeaderGenerator::GetPackageSingletonName(const UPackage* InPackage)
 {
-	static FString ClassString = NameLookupCPP.GetNameCPP(UPackage::StaticClass());
+	const TCHAR* ClassString = NameLookupCPP.GetNameCPP(UPackage::StaticClass());
 
-	FString Result = TEXT("Z_Construct_") + ClassString + TEXT("_") + InPackage->GetName().Replace(TEXT("/"), TEXT("_")) + TEXT("()");
+	FString Result = FString::Printf(TEXT("Z_Construct_%s_%s()"), ClassString, *InPackage->GetName().Replace(TEXT("/"), TEXT("_")));
 
 	if (UniqueCrossModuleReferences)
 	{
@@ -2751,7 +2753,7 @@ static FString PrivatePropertiesOffsetGetters(const UStruct* Struct, const FStri
 	FUHTStringBuilder Result;
 	for (const UProperty* Property : TFieldRange<UProperty>(Struct, EFieldIteratorFlags::ExcludeSuper))
 	{
-		if (Property && Property->HasAnyPropertyFlags(CPF_NativeAccessSpecifierPrivate | CPF_NativeAccessSpecifierProtected) && !Property->HasAnyPropertyFlags(CPF_EditorOnly))
+		if (Property && Property->HasAnyPropertyFlags(CPF_BlueprintVisible) && Property->HasAnyPropertyFlags(CPF_NativeAccessSpecifierPrivate | CPF_NativeAccessSpecifierProtected) && !Property->HasAnyPropertyFlags(CPF_EditorOnly))
 		{
 			const UBoolProperty* BoolProperty = Cast<const UBoolProperty>(Property);
 			if (BoolProperty && !BoolProperty->IsNativeBool()) // if it's a bitfield
@@ -4540,7 +4542,7 @@ void FNativeClassHeaderGenerator::ExportNativeFunctionHeader(
 	FString FunctionName;
 	if (FunctionHeaderStyle == EExportFunctionHeaderStyle::Definition)
 	{
-		FunctionName = FString(NameLookupCPP.GetNameCPP(CastChecked<UClass>(Function->GetOuter()), bIsInterface || FunctionType == EExportFunctionType::Interface)) + TEXT("::");
+		FunctionName = FString::Printf(TEXT("%s::"), NameLookupCPP.GetNameCPP(CastChecked<UClass>(Function->GetOuter()), bIsInterface || FunctionType == EExportFunctionType::Interface));
 	}
 
 	if (FunctionType == EExportFunctionType::Interface)
@@ -4896,7 +4898,7 @@ void FNativeClassHeaderGenerator::ExportNativeFunctions(FOutputDevice& OutGenera
 	FNativeFunctionStringBuilder RuntimeStringBuilders;
 	FNativeFunctionStringBuilder EditorStringBuilders;
 
-	FString ClassName = Class->GetName();
+	const TCHAR* ClassCPPName = NameLookupCPP.GetNameCPP(Class, Class->HasAnyClassFlags(CLASS_Interface));
 
 	ClassDefinitionRange ClassRange;
 	if (ClassDefinitionRanges.Contains(Class))
@@ -5057,8 +5059,7 @@ void FNativeClassHeaderGenerator::ExportNativeFunctions(FOutputDevice& OutGenera
 			// Versions that skip function autodeclaration throw an error when a function is missing.
 			if (ClassRange.bHasGeneratedBody && (SourceFile.GetGeneratedCodeVersionForStruct(Class) > EGeneratedCodeVersion::V1))
 			{
-				FString Name = Class->HasAnyClassFlags(CLASS_Interface) ? TEXT("I") + ClassName : FString(NameLookupCPP.GetNameCPP(Class));
-				CheckRPCFunctions(FunctionData, *Name, ImplementationPosition, ValidatePosition, SourceFile);
+				CheckRPCFunctions(FunctionData, ClassCPPName, ImplementationPosition, ValidatePosition, SourceFile);
 			}
 		}
 
diff --git a/Engine/Source/Programs/UnrealHeaderTool/Private/ParserHelper.cpp b/Engine/Source/Programs/UnrealHeaderTool/Private/ParserHelper.cpp
--- a/Engine/Source/Programs/UnrealHeaderTool/Private/ParserHelper.cpp
+++ b/Engine/Source/Programs/UnrealHeaderTool/Private/ParserHelper.cpp
@@ -281,28 +281,22 @@ FTokenData* FPropertyData::Set(UProperty* InKey, const FTokenData& InValue, FUnr
 
 const TCHAR* FNameLookupCPP::GetNameCPP(UStruct* Struct, bool bForceInterface /*= false */)
 {
-	TCHAR* NameCPP = StructNameMap.FindRef(Struct);
-	if (NameCPP && !bForceInterface)
+	TMap<UStruct*, TCHAR*>& NameMap = (bForceInterface ? InterfaceNameMap : StructNameMap);
+
+	TCHAR* NameCPP = NameMap.FindRef(Struct);
+	if (NameCPP)
 	{
 		return NameCPP;
 	}
 
-	FString DesiredStructName = Struct->GetName();
-	FString	TempName = FString(bForceInterface ? TEXT("I") : Struct->GetPrefixCPP()) + DesiredStructName;
-	int32 StringLength = TempName.Len();
+	const FString TempName = FString::Printf(TEXT("%s%s"), (bForceInterface ? TEXT("I") : Struct->GetPrefixCPP()), *Struct->GetName());
+	const int32 StringLength = TempName.Len();
 
 	NameCPP = new TCHAR[StringLength + 1];
 	FCString::Strcpy(NameCPP, StringLength + 1, *TempName);
 	NameCPP[StringLength] = 0;
 
-	if (bForceInterface)
-	{
-		InterfaceAllocations.Add(NameCPP);
-	}
-	else
-	{
-		StructNameMap.Add(Struct, NameCPP);
-	}
+	NameMap.Add(Struct, NameCPP);
 
 	return NameCPP;
 }
diff --git a/Engine/Source/Programs/UnrealHeaderTool/Private/ParserHelper.h b/Engine/Source/Programs/UnrealHeaderTool/Private/ParserHelper.h
--- a/Engine/Source/Programs/UnrealHeaderTool/Private/ParserHelper.h
+++ b/Engine/Source/Programs/UnrealHeaderTool/Private/ParserHelper.h
@@ -1828,9 +1828,9 @@ struct FNameLookupCPP
 			delete [] Name;
 		}
 	
-		for (TArray<TCHAR*>::TIterator It(InterfaceAllocations); It; ++It)
+		for (TMap<UStruct*, TCHAR*>::TIterator It(InterfaceNameMap); It; ++It)
 		{
-			TCHAR* Name = *It;
+			TCHAR* Name = It.Value();
 			delete [] Name;
 		}
 	}
@@ -1851,7 +1851,7 @@ struct FNameLookupCPP
 private:
 	/** Map of UStruct pointers to C++ names */
 	TMap<UStruct*,TCHAR*> StructNameMap;
-	TArray<TCHAR*> InterfaceAllocations;
+	TMap<UStruct*, TCHAR*> InterfaceNameMap;
 
 	FUnrealSourceFile* UnrealSourceFile;
 };
