Move implementation of RPCWrapper functions in to generated cpp instead of generated h
#rnx

diff --git a/Engine/Source/Programs/UnrealHeaderTool/Private/CodeGenerator.cpp b/Engine/Source/Programs/UnrealHeaderTool/Private/CodeGenerator.cpp
--- a/Engine/Source/Programs/UnrealHeaderTool/Private/CodeGenerator.cpp
+++ b/Engine/Source/Programs/UnrealHeaderTool/Private/CodeGenerator.cpp
@@ -2791,7 +2791,7 @@ void FNativeClassHeaderGenerator::ExportClassFromSourceFileInner(
 	// C++ -> VM stubs (native function execs)
 	FUHTStringBuilder ClassMacroCalls;
 	FUHTStringBuilder ClassNoPureDeclsMacroCalls;
-	ExportNativeFunctions(OutGeneratedHeaderText, ClassMacroCalls, ClassNoPureDeclsMacroCalls, SourceFile, Class, ClassData);
+	ExportNativeFunctions(OutGeneratedHeaderText, OutCpp, ClassMacroCalls, ClassNoPureDeclsMacroCalls, SourceFile, Class, ClassData);
 
 	// Get Callback functions
 	TArray<UFunction*> CallbackFunctions;
@@ -4887,13 +4887,14 @@ FString FNativeClassHeaderGenerator::GetFunctionParameterString(UFunction* Funct
 struct FNativeFunctionStringBuilder
 {
 	FUHTStringBuilder RPCWrappers;
+	FUHTStringBuilder RPCImplementations;
 	FUHTStringBuilder AutogeneratedBlueprintFunctionDeclarations;
 	FUHTStringBuilder AutogeneratedBlueprintFunctionDeclarationsOnlyNotDeclared;
 	FUHTStringBuilder AutogeneratedStaticData;
 	FUHTStringBuilder AutogeneratedStaticDataFuncs;
 };
 
-void FNativeClassHeaderGenerator::ExportNativeFunctions(FOutputDevice& OutGeneratedHeaderText, FOutputDevice& OutMacroCalls, FOutputDevice& OutNoPureDeclsMacroCalls, const FUnrealSourceFile& SourceFile, UClass* Class, FClassMetaData* ClassData)
+void FNativeClassHeaderGenerator::ExportNativeFunctions(FOutputDevice& OutGeneratedHeaderText, FOutputDevice& OutGeneratedCPPText, FOutputDevice& OutMacroCalls, FOutputDevice& OutNoPureDeclsMacroCalls, const FUnrealSourceFile& SourceFile, UClass* Class, FClassMetaData* ClassData)
 {
 	FNativeFunctionStringBuilder RuntimeStringBuilders;
 	FNativeFunctionStringBuilder EditorStringBuilders;
@@ -5073,13 +5074,14 @@ void FNativeClassHeaderGenerator::ExportNativeFunctions(FOutputDevice& OutGenera
 		}
 
 		// export the script wrappers
-		FuncStringBuilders.RPCWrappers.Logf(TEXT("\tDECLARE_FUNCTION(%s)"), *FunctionData.UnMarshallAndCallName);
-		FuncStringBuilders.RPCWrappers += LINE_TERMINATOR TEXT("\t{") LINE_TERMINATOR;
+		FuncStringBuilders.RPCWrappers.Logf(TEXT("\tDECLARE_FUNCTION(%s);"), *FunctionData.UnMarshallAndCallName);
+		FuncStringBuilders.RPCImplementations.Logf(TEXT("\tDEFINE_FUNCTION(%s::%s)"), ClassCPPName, *FunctionData.UnMarshallAndCallName);
+		FuncStringBuilders.RPCImplementations += LINE_TERMINATOR TEXT("\t{") LINE_TERMINATOR;
 
 		FParmsAndReturnProperties Parameters = GetFunctionParmsAndReturn(FunctionData.FunctionReference);
-		ExportFunctionThunk(FuncStringBuilders.RPCWrappers, Function, FunctionData, Parameters.Parms, Parameters.Return);
+		ExportFunctionThunk(FuncStringBuilders.RPCImplementations, Function, FunctionData, Parameters.Parms, Parameters.Return);
 
-		FuncStringBuilders.RPCWrappers += TEXT("\t}") LINE_TERMINATOR;
+		FuncStringBuilders.RPCImplementations += TEXT("\t}") LINE_TERMINATOR;
 	}
 
 	// static class data
@@ -5095,6 +5097,13 @@ void FNativeClassHeaderGenerator::ExportNativeFunctions(FOutputDevice& OutGenera
 	{
 		FString MacroName = SourceFile.GetGeneratedMacroName(ClassData, TEXT("_RPC_WRAPPERS"));
 
+		// WriteMacro has an assumption about what will be at the end of this block that is no longer true due to splitting the
+		// definition and implementation, so add on a line terminator to satisfy it
+		if (RuntimeStringBuilders.RPCWrappers.Len() > 0)
+		{
+			RuntimeStringBuilders.RPCWrappers += LINE_TERMINATOR;
+		}
+
 		WriteMacro(OutGeneratedHeaderText, MacroName, RuntimeStringBuilders.AutogeneratedBlueprintFunctionDeclarations + RuntimeStringBuilders.RPCWrappers);
 		OutMacroCalls.Logf(TEXT("\t%s\r\n"), *MacroName);
 
@@ -5110,6 +5119,8 @@ void FNativeClassHeaderGenerator::ExportNativeFunctions(FOutputDevice& OutGenera
 		}
 
 		OutNoPureDeclsMacroCalls.Logf(TEXT("\t%s\r\n"), *NoPureDeclsMacroName);
+
+		OutGeneratedCPPText.Log(RuntimeStringBuilders.RPCImplementations);
 	}
 
 	// Write editor only RPC wrappers if they exist
@@ -5119,6 +5130,12 @@ void FNativeClassHeaderGenerator::ExportNativeFunctions(FOutputDevice& OutGenera
 
 		FString MacroName = SourceFile.GetGeneratedMacroName(ClassData, TEXT("_EDITOR_ONLY_RPC_WRAPPERS"));
 
+		// WriteMacro has an assumption about what will be at the end of this block that is no longer true due to splitting the
+		// definition and implementation, so add on a line terminator to satisfy it
+		if (EditorStringBuilders.RPCWrappers.Len() > 0)
+		{
+			EditorStringBuilders.RPCWrappers += LINE_TERMINATOR;
+		}
 
 		WriteMacro(OutGeneratedHeaderText, MacroName, EditorStringBuilders.AutogeneratedBlueprintFunctionDeclarations + EditorStringBuilders.RPCWrappers);
 		OutMacroCalls.Logf(TEXT("\t%s\r\n"), *MacroName);
@@ -5145,6 +5162,10 @@ void FNativeClassHeaderGenerator::ExportNativeFunctions(FOutputDevice& OutGenera
 		}
 
 		OutNoPureDeclsMacroCalls.Logf(TEXT("\t%s\r\n"), *NoPureDeclsMacroName);
+
+		OutGeneratedCPPText.Log(BeginEditorOnlyGuard);
+		OutGeneratedCPPText.Log(EditorStringBuilders.RPCImplementations);
+		OutGeneratedCPPText.Log(EndEditorOnlyGuard);
 	}
 }
 
diff --git a/Engine/Source/Programs/UnrealHeaderTool/Private/NativeClassExporter.h b/Engine/Source/Programs/UnrealHeaderTool/Private/NativeClassExporter.h
--- a/Engine/Source/Programs/UnrealHeaderTool/Private/NativeClassExporter.h
+++ b/Engine/Source/Programs/UnrealHeaderTool/Private/NativeClassExporter.h
@@ -311,7 +311,7 @@ private:
 	 * @param Class			class
 	 * @param ClassData		class data
 	 */
-	void ExportNativeFunctions(FOutputDevice& OutGeneratedHeaderText, FOutputDevice& OutMacroCalls, FOutputDevice& OutNoPureDeclsMacroCalls, const FUnrealSourceFile& SourceFile, UClass* Class, FClassMetaData* ClassData);
+	void ExportNativeFunctions(FOutputDevice& OutGeneratedHeaderText, FOutputDevice& OutGeneratedCPPText, FOutputDevice& OutMacroCalls, FOutputDevice& OutNoPureDeclsMacroCalls, const FUnrealSourceFile& SourceFile, UClass* Class, FClassMetaData* ClassData);
 
 	/**
 	 * Export the actual internals to a standard thunk function
