Config properties need a PPO stub for the time being
#rnx

diff --git a/Engine/Source/Programs/UnrealHeaderTool/Private/CodeGenerator.cpp b/Engine/Source/Programs/UnrealHeaderTool/Private/CodeGenerator.cpp
--- a/Engine/Source/Programs/UnrealHeaderTool/Private/CodeGenerator.cpp
+++ b/Engine/Source/Programs/UnrealHeaderTool/Private/CodeGenerator.cpp
@@ -2753,7 +2753,8 @@ static FString PrivatePropertiesOffsetGetters(const UStruct* Struct, const FStri
 	FUHTStringBuilder Result;
 	for (const UProperty* Property : TFieldRange<UProperty>(Struct, EFieldIteratorFlags::ExcludeSuper))
 	{
-		if (Property && Property->HasAnyPropertyFlags(CPF_BlueprintVisible) && Property->HasAnyPropertyFlags(CPF_NativeAccessSpecifierPrivate | CPF_NativeAccessSpecifierProtected) && !Property->HasAnyPropertyFlags(CPF_EditorOnly))
+		// TODO: Config values should not be getting PPO accessors, however, due to CL# 3023067 they currently need to be, so until that bug can be revisited
+		if (Property && Property->HasAnyPropertyFlags(CPF_BlueprintVisible|CPF_Config) && Property->HasAnyPropertyFlags(CPF_NativeAccessSpecifierPrivate | CPF_NativeAccessSpecifierProtected) && !Property->HasAnyPropertyFlags(CPF_EditorOnly))
 		{
 			const UBoolProperty* BoolProperty = Cast<const UBoolProperty>(Property);
 			if (BoolProperty && !BoolProperty->IsNativeBool()) // if it's a bitfield
