Remove requirement for "= true" on AllowPrivateAccess meta tag

diff --git a/Engine/Source/Programs/UnrealHeaderTool/Private/HeaderParser.cpp b/Engine/Source/Programs/UnrealHeaderTool/Private/HeaderParser.cpp
--- a/Engine/Source/Programs/UnrealHeaderTool/Private/HeaderParser.cpp
+++ b/Engine/Source/Programs/UnrealHeaderTool/Private/HeaderParser.cpp
@@ -3591,7 +3591,7 @@ void FHeaderParser::GetVarType(
 					}
 
 					const FString* PrivateAccessMD = MetaDataFromNewStyle.Find(TEXT("AllowPrivateAccess"));  // FBlueprintMetadata::MD_AllowPrivateAccess
-					const bool bAllowPrivateAccess = PrivateAccessMD ? (*PrivateAccessMD == TEXT("true")) : false;
+					const bool bAllowPrivateAccess = PrivateAccessMD ? (*PrivateAccessMD != TEXT("false")) : false;
 					if (CurrentAccessSpecifier == ACCESS_Private && !bAllowPrivateAccess)
 					{
 						UE_LOG_ERROR_UHT(TEXT("BlueprintReadWrite should not be used on private members"));
@@ -3635,7 +3635,7 @@ void FHeaderParser::GetVarType(
 					}
 
 					const FString* PrivateAccessMD = MetaDataFromNewStyle.Find(TEXT("AllowPrivateAccess"));  // FBlueprintMetadata::MD_AllowPrivateAccess
-					const bool bAllowPrivateAccess = PrivateAccessMD ? (*PrivateAccessMD == TEXT("true")) : false;
+					const bool bAllowPrivateAccess = PrivateAccessMD ? (*PrivateAccessMD != TEXT("false")) : false;
 					if (CurrentAccessSpecifier == ACCESS_Private && !bAllowPrivateAccess)
 					{
 						UE_LOG_ERROR_UHT(TEXT("BlueprintReadOnly should not be used on private members"));
