Fix UHT Generation of RepIndices and don't sort Native Rep Properties.

Fixes an issue where the generated RepIndices could be invalid due to UHT not being able to know the size of static arrays or final offset of bitfield properties.

diff --git a/Engine/Source/Programs/UnrealHeaderTool/Private/CodeGenerator.cpp b/Engine/Source/Programs/UnrealHeaderTool/Private/CodeGenerator.cpp
--- a/Engine/Source/Programs/UnrealHeaderTool/Private/CodeGenerator.cpp
+++ b/Engine/Source/Programs/UnrealHeaderTool/Private/CodeGenerator.cpp
@@ -155,110 +155,111 @@ namespace
 		return false;
 	}
 
-	static void WriteReplicatedMacroData(
-		const ClassDefinitionRange&	ClassRange,
-		const TCHAR*				ClassCPPName,
-		FClass*						Class,
-		FClass*						SuperClass,
-		FOutputDevice&				Writer,
-		const FUnrealSourceFile&	SourceFile)
+	static void ExportNetData(FOutputDevice& Out, UClass* Class, UClass* ParentRepClass)
 	{
-		const bool bHasGetLifetimeReplicatedProps = HasIdentifierExactMatch(ClassRange.Start, ClassRange.End, TEXT("GetLifetimeReplicatedProps"));
-	
-		if (!bHasGetLifetimeReplicatedProps)
-		{
-			// Default version autogenerates declarations.
-			if (SourceFile.GetGeneratedCodeVersionForStruct(Class) == EGeneratedCodeVersion::V1)
-			{
-				Writer.Logf(TEXT("\tvoid GetLifetimeReplicatedProps(TArray<FLifetimeProperty>& OutLifetimeProps) const override;\r\n"));
-			}
-			else
-			{
-				FError::Throwf(TEXT("Class %s has Net flagged properties and should declare member function: void GetLifetimeReplicatedProps(TArray<FLifetimeProperty>& OutLifetimeProps) const override"), ClassCPPName);
-			}
-		}	
-	}
+		const TArray<FRepRecord>& ClassReps = Class->ClassReps;
+		
+		const FString ClassName = FString::Printf(TEXT("%s%s"), Class->GetPrefixCPP(), *Class->GetName());
+		const FString ParentClassName = ParentRepClass ? FString::Printf(TEXT("%s%s"), ParentRepClass->GetPrefixCPP(), *ParentRepClass->GetName()) : FString();
+		const FString NetfieldStart = ParentRepClass ? FString::Printf(TEXT("(uint16)%s::ENetFields_Private::NETFIELD_REP_END + (uint16)1"), *ParentClassName, *ParentClassName) : FString(TEXT("0"));
 
-	static void ExportNetData(FOutputDevice& Out, UClass* Class)
-	{
-		if (ClassHasReplicatedProperties(Class))
-		{
-			Class->SetUpRuntimeReplicationData();
+		FUHTStringBuilder NetFieldBuilder;
+		NetFieldBuilder.Logf(TEXT(""
+		"\tenum class ENetFields_Private : uint16\r\n"
+		"\t{\r\n"
+		"\t\tNETFIELD_REP_START=%s,\r\n"), *NetfieldStart);
 
-			int32 StartingClassRep = 0;
-			if (UClass * SuperClass = Class->GetSuperClass())
-			{
-				StartingClassRep = SuperClass->ClassReps.Num();
-			}
+		FUHTStringBuilder ArrayDimBuilder;
+
+		bool bAnyStaticArrays = false;
+		bool bIsFirst = true;
+		for (int32 ClassRepIndex = ClassReps.Num() - Class->NetFields.Num(); ClassRepIndex < ClassReps.Num(); ++ClassRepIndex)
+		{
+			const FRepRecord& ClassRep = ClassReps[ClassRepIndex];
+			const FString& PropertyName = ClassRep.Property->GetName();
 
-			// If we don't have any Replicated Properties, then don't do anything.
-			if (Class->ClassReps.Num() == StartingClassRep)
+			if (ClassRep.Property->ArrayDim == 1)
 			{
-				return;
+				if (UNLIKELY(bIsFirst))
+				{
+					NetFieldBuilder.Logf(TEXT("\t\t%s=NETFIELD_REP_START,\r\n"), *PropertyName);
+					bIsFirst = false;
+				}
+				else
+				{
+					NetFieldBuilder.Logf(TEXT("\t\t%s,\r\n"), *PropertyName);
+				}
 			}
-
-			FUHTStringBuilder NetFieldBuilder;
-			NetFieldBuilder.Logf(TEXT(""
-			"\t\tenum NetFields\r\n"
-			"\t\t{\r\n"
-			"\t\t\tNETFIELD_REP_START=%d,\r\n"), StartingClassRep);
-
-			FUHTStringBuilder ArrayDimBuilder;
-			ArrayDimBuilder.Log(TEXT(""
-			"\t\tenum ArrayDims\r\n"
-			"\t\t{\r\n"));
-
-			Out.Logf(TEXT("namespace NetworkingPrivate\r\n"
-				"{\r\n"
-				"\tnamespace Net_%s%s\r\n"
-				"\t{\r\n"), Class->GetPrefixCPP(), *Class->GetName());
-
-			TArray<const FRepRecord*> StaticArrayRecords;
-			const TArray<FRepRecord>& ClassReps = Class->ClassReps;
-			const int32 NumClassReps = ClassReps.Num();
-			for (int32 ClassRepIndex = StartingClassRep; ClassRepIndex < ClassReps.Num(); ++ClassRepIndex)
+			else
 			{
-				const FRepRecord& ClassRep = ClassReps[ClassRepIndex];
-				const FString& PropertyName = ClassRep.Property->GetName();
+				bAnyStaticArrays = true;
+				ArrayDimBuilder.Logf(TEXT("\t\t%s=%s,\r\n"), *PropertyName, *GArrayDimensions.FindChecked(ClassReps[ClassRepIndex].Property));
 
-				if (ClassRep.Property->ArrayDim == 1)
+				if (UNLIKELY(bIsFirst))
 				{
-					
-					NetFieldBuilder.Logf(TEXT("\t\t\tNETFIELD_%s=%d,\r\n"), *PropertyName, ClassRepIndex);
-					ArrayDimBuilder.Logf(TEXT("\t\t\tARRAYDIM_%s=1,\r\n"), *PropertyName);
+					NetFieldBuilder.Logf(TEXT("\t\t%s_STATIC_ARRAY=NETFIELD_REP_START,\r\n"), *PropertyName);
+					bIsFirst = false;
 				}
 				else
 				{
-					const int32 NumClassRepArrayElements = ClassRep.Property->ArrayDim;
-					const int32 ClassRepArrayStart = ClassRepIndex;
-				
-					// -1 because we want this to be the last valid RepIndex.
-					const int32 ClassRepArrayEnd = ClassRepArrayStart + NumClassRepArrayElements - 1;
-
-					for (int32 i = 0; i < NumClassRepArrayElements; ++i)
-					{
-						NetFieldBuilder.Logf(TEXT("\t\t\tNETFIELD_%s_%d=%d,\r\n"), *PropertyName, i, ClassRepIndex + i);
-					}
-				
-					ClassRepIndex += NumClassRepArrayElements - 1;
-					
-					ArrayDimBuilder.Logf(TEXT("\t\t\tARRAYDIM_%s=%d,\r\n"), *PropertyName, NumClassRepArrayElements);
+					NetFieldBuilder.Logf(TEXT("\t\t%s_STATIC_ARRAY,\r\n"), *PropertyName);
 				}
+
+				NetFieldBuilder.Logf(TEXT("\t\t%s_STATIC_ARRAY_END=((uint16)%s_STATIC_ARRAY + (uint16)EArrayDims_Private::%s - (uint16)1),\r\n"), *PropertyName, *PropertyName, *PropertyName);
 			}
+		}
 
-			NetFieldBuilder.Logf(TEXT("\t\t\tNETFIELD_REP_END=%d,\r\n"), ClassReps.Num() - 1);
+		const UProperty* LastProperty = ClassReps.Last().Property;
+		const FString NetfieldEnd = LastProperty->ArrayDim > 1 ?
+			FString::Printf(TEXT("\t\tNETFIELD_REP_END=%s_STATIC_ARRAY_END,\r\n"), *LastProperty->GetName()) :
+			FString::Printf(TEXT("\t\tNETFIELD_REP_END=%s,\r\n"), *LastProperty->GetName());
 
-			NetFieldBuilder.Log(TEXT("\t\t};"));
-			ArrayDimBuilder.Log(TEXT("\t\t};"));
+		NetFieldBuilder.Log(*NetfieldEnd);
 
+		NetFieldBuilder.Log(TEXT("\t};"));
+
+		if (bAnyStaticArrays)
+		{
 			Out.Logf(TEXT(""
-				"%s\r\n" // NetFields
-				"%s\r\n" // ArrayDims
-				"\t}\r\n"
-				"}\r\n"),
-				*NetFieldBuilder,
-				*ArrayDimBuilder);
+				"\tenum class EArrayDims_Private : uint16\r\n"
+				"\t{\r\n"
+				"%s"
+				"\t};\r\n"), *ArrayDimBuilder);
 		}
+
+		Out.Logf(TEXT(""
+			"%s\r\n" // NetFields
+			"\tvirtual void ValidateGeneratedRepEnums(const TArray<struct FRepRecord>& ClassReps) const override;\r\n"),
+			*NetFieldBuilder);
+	}
+
+	static void WriteReplicatedMacroData(
+		const ClassDefinitionRange& ClassRange,
+		const TCHAR* ClassCPPName,
+		FClass* Class,
+		FClass* SuperClass,
+		FOutputDevice& Writer,
+		const FUnrealSourceFile& SourceFile)
+	{
+		const bool bHasGetLifetimeReplicatedProps = HasIdentifierExactMatch(ClassRange.Start, ClassRange.End, TEXT("GetLifetimeReplicatedProps"));
+
+		if (!bHasGetLifetimeReplicatedProps)
+		{
+			// Default version autogenerates declarations.
+			if (SourceFile.GetGeneratedCodeVersionForStruct(Class) == EGeneratedCodeVersion::V1)
+			{
+				Writer.Logf(TEXT("\tvoid GetLifetimeReplicatedProps(TArray<FLifetimeProperty>& OutLifetimeProps) const override;\r\n"));
+			}
+			else
+			{
+				FError::Throwf(TEXT("Class %s has Net flagged properties and should declare member function: void GetLifetimeReplicatedProps(TArray<FLifetimeProperty>& OutLifetimeProps) const override"), ClassCPPName);
+			}
+		}
+
+		UClass* ParentRepClass = nullptr;
+		Class->SetUpUhtReplicationData(&ParentRepClass);
+
+		ExportNetData(Writer, Class, ParentRepClass);
 	}
 }
 
@@ -2382,6 +2383,49 @@ void FNativeClassHeaderGenerator::ExportNativeGeneratedInitCode(FOutputDevice& O
 		bIsDynamic ? *AsTEXT(FClass::GetTypePackageName(Class)) : TEXT("nullptr"),
 		bIsDynamic ? *AsTEXT(FNativeClassHeaderGenerator::GetOverriddenPathName(Class)) : TEXT("nullptr"),
 		*InitSearchableValuesFunctionParam);
+
+
+	if (ClassHasReplicatedProperties(Class))
+	{
+		Class->SetUpUhtReplicationData(nullptr);
+
+		Out.Logf(TEXT(
+			"\r\n"
+			"\tvoid %s::ValidateGeneratedRepEnums(const TArray<struct FRepRecord>& ClassReps) const\r\n"
+			"\t{\r\n"
+		), ClassNameCPP, ClassNameCPP);
+
+		FUHTStringBuilder NameBuilder;
+
+		FUHTStringBuilder ValidationBuilder;
+		ValidationBuilder.Log(TEXT("\t\tconst bool bIsValid = true"));
+
+		for (int32 i = Class->ClassReps.Num() - Class->NetFields.Num(); i < Class->ClassReps.Num(); ++i)
+		{
+			const UProperty* const Property = Class->ClassReps[i].Property;
+			const FString& PropertyName = Property->GetName();
+
+			NameBuilder.Logf(TEXT("\t\tstatic const FName Name_%s(TEXT(\"%s\"));\r\n"), *PropertyName, *PropertyName);
+
+			if (Property->ArrayDim == 1)
+			{
+				ValidationBuilder.Logf(TEXT("\r\n\t\t\t&& Name_%s == ClassReps[(int32)ENetFields_Private::%s].Property->GetFName()"), *PropertyName, *PropertyName);
+			}
+			else
+			{
+				ValidationBuilder.Logf(TEXT("\r\n\t\t\t&& Name_%s == ClassReps[(int32)ENetFields_Private::%s_STATIC_ARRAY].Property->GetFName()"), *PropertyName, *PropertyName);
+			}
+		}
+
+		ValidationBuilder.Log(TEXT(";\r\n"));
+
+		Out.Logf(TEXT(
+			"%s\r\n" // NameBuilder
+			"%s\r\n" // ValidationBuilder
+			"\t\tcheckf(bIsValid, TEXT(\"UHT Generated Rep Indices do not match runtime populated Rep Indices for properties in %s\"));\r\n"
+			"\t}\r\n"
+		), *NameBuilder, *ValidationBuilder, ClassNameCPP);
+	}
 }
 
 void FNativeClassHeaderGenerator::ExportFunction(FOutputDevice& Out, const FUnrealSourceFile& SourceFile, UFunction* Function, bool bIsNoExport)
@@ -5570,11 +5614,6 @@ FNativeClassHeaderGenerator::FNativeClassHeaderGenerator(
 			ExportEnum(GeneratedHeaderText, Enum);
 		}
 
-		for (UClass* Class : DefinedClasses)
-		{
-			ExportNetData(GeneratedHeaderText, Class);
-		}
-
 		FString HeaderPath = BaseSourceFilename + TEXT(".generated.h");
 		bool bHasChanged = WriteHeader(*HeaderPath, GeneratedHeaderText, ForwardDeclarations);
 
diff --git a/Engine/Source/Runtime/CoreUObject/Private/UObject/Class.cpp b/Engine/Source/Runtime/CoreUObject/Private/UObject/Class.cpp
--- a/Engine/Source/Runtime/CoreUObject/Private/UObject/Class.cpp
+++ b/Engine/Source/Runtime/CoreUObject/Private/UObject/Class.cpp
@@ -3452,7 +3452,7 @@ void UClass::Link(FArchive& Ar, bool bRelinkExistingProperties)
 	Super::Link(Ar, bRelinkExistingProperties);
 }
 
-#if (UE_BUILD_SHIPPING || HACK_HEADER_GENERATOR)
+#if (UE_BUILD_SHIPPING)
 static int32 GValidateReplicatedProperties = 0;
 #else 
 static int32 GValidateReplicatedProperties = 1;
@@ -3460,6 +3460,43 @@ static int32 GValidateReplicatedProperties = 1;
 
 static FAutoConsoleVariable CVarValidateReplicatedPropertyRegistration(TEXT("net.ValidateReplicatedPropertyRegistration"), GValidateReplicatedProperties, TEXT("Warns if replicated properties were not registered in GetLifetimeReplicatedProps."));
 
+void UClass::SetUpUhtReplicationData(UClass** OutSuperClassWithReplicatedData)
+{
+#if HACK_HEADER_GENERATOR
+	if (!HasAnyClassFlags(CLASS_ReplicationDataIsSetUp) && PropertyLink != NULL)
+	{
+		NetFields.Empty();
+        ClassReps.Empty();
+		if (UClass* SuperClass = GetSuperClass())
+		{
+			SuperClass->SetUpUhtReplicationData(nullptr);
+			ClassReps = SuperClass->ClassReps;
+		}
+
+		for (TFieldIterator<UProperty> It(this, EFieldIteratorFlags::ExcludeSuper); It; ++It)
+		{
+			if (It->PropertyFlags & CPF_Net)
+			{
+				It->RepIndex = ClassReps.Num();
+				NetFields.Add(*It);
+				new (ClassReps) FRepRecord(*It, 0);
+			}
+		}
+
+		ClassFlags |= CLASS_ReplicationDataIsSetUp;
+		ClassReps.Shrink();
+		NetFields.Shrink();
+	}
+
+	if (OutSuperClassWithReplicatedData)
+	{
+		UClass* SuperClass = GetSuperClass();
+		for (; SuperClass != nullptr && SuperClass->NetFields.Num() == 0; SuperClass = SuperClass->GetSuperClass()) {}
+		*OutSuperClassWithReplicatedData = SuperClass;
+	}
+#endif
+}
+
 void UClass::SetUpRuntimeReplicationData()
 {
 	if (!HasAnyClassFlags(CLASS_ReplicationDataIsSetUp) && PropertyLink != NULL)
@@ -3505,22 +3542,26 @@ void UClass::SetUpRuntimeReplicationData()
 			}
 		}
 
-		// Sort NetProperties so that their ClassReps are sorted by memory offset
-		struct FCompareUFieldOffsets
+		const bool bIsNativeClass = HasAnyClassFlags(CLASS_Native);
+		if (!bIsNativeClass)
 		{
-			FORCEINLINE bool operator()( UProperty & A, UProperty & B ) const
+			// Sort NetProperties so that their ClassReps are sorted by memory offset
+			struct FCompareUFieldOffsets
 			{
-				// Ensure stable sort
-				if ( A.GetOffset_ForGC() == B.GetOffset_ForGC() )
+				FORCEINLINE bool operator()(UProperty& A, UProperty& B) const
 				{
-					return A.GetName() < B.GetName();
-				}
+					// Ensure stable sort
+					if (A.GetOffset_ForGC() == B.GetOffset_ForGC())
+					{
+						return A.GetName() < B.GetName();
+					}
 
-				return A.GetOffset_ForGC() < B.GetOffset_ForGC();
-			}
-		};
+					return A.GetOffset_ForGC() < B.GetOffset_ForGC();
+				}
+			};
 
-		Sort(NetProperties.GetData(), NetProperties.Num(), FCompareUFieldOffsets());
+			Sort(NetProperties.GetData(), NetProperties.Num(), FCompareUFieldOffsets());
+		}
 
 		for ( int32 i = 0; i < NetProperties.Num(); i++ )
 		{
@@ -3531,6 +3572,11 @@ void UClass::SetUpRuntimeReplicationData()
 			}
 		}
 
+		if (bIsNativeClass && GValidateReplicatedProperties)
+		{
+			GetDefaultObject()->ValidateGeneratedRepEnums(ClassReps);
+		}
+
 		NetFields.Shrink();
 
 		struct FCompareUFieldNames
diff --git a/Engine/Source/Runtime/CoreUObject/Public/UObject/Class.h b/Engine/Source/Runtime/CoreUObject/Public/UObject/Class.h
--- a/Engine/Source/Runtime/CoreUObject/Public/UObject/Class.h
+++ b/Engine/Source/Runtime/CoreUObject/Public/UObject/Class.h
@@ -3047,6 +3047,7 @@ public:
 	 * Also happens after blueprint compiliation.
 	 */
 	void SetUpRuntimeReplicationData();
+	void SetUpUhtReplicationData(UClass** OutSuperClassWithReplicatedData);
 
 	/**
 	 * Helper function for determining if the given class is compatible with structured archive serialization
diff --git a/Engine/Source/Runtime/CoreUObject/Public/UObject/Object.h b/Engine/Source/Runtime/CoreUObject/Public/UObject/Object.h
--- a/Engine/Source/Runtime/CoreUObject/Public/UObject/Object.h
+++ b/Engine/Source/Runtime/CoreUObject/Public/UObject/Object.h
@@ -1455,6 +1455,10 @@ private:
 	* @param	bTransient					true if the component is being assigned to a transient property
 	*/
 	UObject* CreateEditorOnlyDefaultSubobjectImpl(FName SubobjectName, UClass* ReturnType, bool bTransient = false);
+
+public:
+
+	virtual void ValidateGeneratedRepEnums(const TArray<struct FRepRecord>& ClassReps) const {}
 };
 
 /**
diff --git a/Engine/Source/Runtime/Engine/Private/RepLayout.cpp b/Engine/Source/Runtime/Engine/Private/RepLayout.cpp
--- a/Engine/Source/Runtime/Engine/Private/RepLayout.cpp
+++ b/Engine/Source/Runtime/Engine/Private/RepLayout.cpp
@@ -1058,8 +1058,8 @@ static void CompareRoleProperties(
 {
 	if (RepState && EnumHasAnyFlags(SharedParams.Flags, ERepLayoutFlags::IsActor))
 	{
-		CompareRoleProperty(SharedParams, Data, NetworkingPrivate::Net_AActor::NETFIELD_RemoteRole, RepState->SavedRemoteRole, Changed);
-		CompareRoleProperty(SharedParams, Data, NetworkingPrivate::Net_AActor::NETFIELD_Role, RepState->SavedRole, Changed);
+		CompareRoleProperty(SharedParams, Data, (int32)AActor::ENetFields_Private::RemoteRole, RepState->SavedRemoteRole, Changed);
+		CompareRoleProperty(SharedParams, Data, (int32)AActor::ENetFields_Private::Role, RepState->SavedRole, Changed);
 	}
 }
 
@@ -1098,14 +1098,14 @@ static void CompareParentProperties(
 
 		if (bCheckForRole)
 		{
-			if (UNLIKELY(ParentIndex == NetworkingPrivate::Net_AActor::NETFIELD_Role))
+			if (UNLIKELY(ParentIndex == (int32)AActor::ENetFields_Private::Role))
 			{
-				CompareRoleProperty(SharedParams, Data, NetworkingPrivate::Net_AActor::NETFIELD_Role, RepState->SavedRole, Changed);
+				CompareRoleProperty(SharedParams, Data, (int32)AActor::ENetFields_Private::Role, RepState->SavedRole, Changed);
 				continue;
 			}
-			else if (UNLIKELY(ParentIndex == NetworkingPrivate::Net_AActor::NETFIELD_RemoteRole))
+			else if (UNLIKELY(ParentIndex == (int32)AActor::ENetFields_Private::RemoteRole))
 			{
-				CompareRoleProperty(SharedParams, Data, NetworkingPrivate::Net_AActor::NETFIELD_RemoteRole, RepState->SavedRemoteRole, Changed);
+				CompareRoleProperty(SharedParams, Data, (int32)AActor::ENetFields_Private::RemoteRole, RepState->SavedRemoteRole, Changed);
 				continue;
 			}
 		}
@@ -2576,13 +2576,13 @@ static bool ReceivePropertyHelper(
 		if (!bSkipSwapRoles)
 		{
 			// Swap Role to RemoteRole, and vice-versa. Leave everything else the same.
-			if (UNLIKELY(NetworkingPrivate::Net_AActor::NETFIELD_Role == Cmd.ParentIndex))
+			if (UNLIKELY((int32)AActor::ENetFields_Private::RemoteRole == Cmd.ParentIndex))
 			{
-				return Cmds[Parents[NetworkingPrivate::Net_AActor::NETFIELD_RemoteRole].CmdStart];
+				return Cmds[Parents[(int32)AActor::ENetFields_Private::Role].CmdStart];
 			}
-			else if (UNLIKELY(NetworkingPrivate::Net_AActor::NETFIELD_RemoteRole == Cmd.ParentIndex))
+			else if (UNLIKELY((int32)AActor::ENetFields_Private::Role == Cmd.ParentIndex))
 			{
-				return Cmds[Parents[NetworkingPrivate::Net_AActor::NETFIELD_Role].CmdStart];
+				return Cmds[Parents[(int32)AActor::ENetFields_Private::RemoteRole].CmdStart];
 			}
 		}
 
@@ -5101,13 +5101,13 @@ void FRepLayout::InitFromClass(
 	}
 
 	// Make sure RemoteRole has a lower RepIndex than Role, otherwise assumptions RepLayout may break.
-	static_assert(NetworkingPrivate::Net_AActor::NETFIELD_RemoteRole < NetworkingPrivate::Net_AActor::NETFIELD_Role, "Role and RemoteRole have been rearranged in AActor. This will break assumptions in RepLayout.");
+	static_assert((int32)AActor::ENetFields_Private::RemoteRole < (int32)AActor::ENetFields_Private::Role, "Role and RemoteRole have been rearranged in AActor. This will break assumptions in RepLayout.");
 
 	// Make sure that our RemoteRole property actually points to RemoteRole.
-	check(!bIsObjectActor || Parents[NetworkingPrivate::Net_AActor::NETFIELD_RemoteRole].Property->GetFName() == NAME_RemoteRole);
+	check(!bIsObjectActor || Parents[(int32)AActor::ENetFields_Private::RemoteRole].Property->GetFName() == NAME_RemoteRole);
 
 	// Make sure that our Role property actually points to Role.
-	check(!bIsObjectActor || Parents[NetworkingPrivate::Net_AActor::NETFIELD_Role].Property->GetFName() == NAME_Role);
+	check(!bIsObjectActor || Parents[(int32)AActor::ENetFields_Private::Role].Property->GetFName() == NAME_Role);
 
 	AddReturnCmd(Cmds);
 
@@ -5245,7 +5245,7 @@ void FRepLayout::InitFromClass(
 	{
 		// We handle remote role specially, since it can change between connections when downgraded
 		// So we force it on the conditional list
-		FRepParentCmd& RemoteRoleParent = Parents[NetworkingPrivate::Net_AActor::NETFIELD_RemoteRole];
+		FRepParentCmd& RemoteRoleParent = Parents[(int32)AActor::ENetFields_Private::RemoteRole];
 		if (RemoteRoleParent.Condition != COND_Never)
 		{
 			if (COND_None != RemoteRoleParent.Condition)
@@ -5253,8 +5253,8 @@ void FRepLayout::InitFromClass(
 				UE_LOG(LogRep, Warning, TEXT("FRepLayout::InitFromClass: Forcing replication of RemoteRole. Owner=%s"), *InObjectClass->GetPathName());
 			}
 
-			Parents[NetworkingPrivate::Net_AActor::NETFIELD_RemoteRole].Flags |= ERepParentFlags::IsConditional;
-			Parents[NetworkingPrivate::Net_AActor::NETFIELD_RemoteRole].Condition = COND_None;
+			Parents[(int32)AActor::ENetFields_Private::RemoteRole].Flags |= ERepParentFlags::IsConditional;
+			Parents[(int32)AActor::ENetFields_Private::RemoteRole].Condition = COND_None;
 		}
 	}
 
