There is no need to export the Z_Construct_UFunction declarations.

diff --git a/Engine/Source/Programs/UnrealHeaderTool/Private/CodeGenerator.cpp b/Engine/Source/Programs/UnrealHeaderTool/Private/CodeGenerator.cpp
--- a/Engine/Source/Programs/UnrealHeaderTool/Private/CodeGenerator.cpp
+++ b/Engine/Source/Programs/UnrealHeaderTool/Private/CodeGenerator.cpp
@@ -942,7 +942,8 @@ FString FNativeClassHeaderGenerator::GetSingletonName(UField* Item, bool bRequir
 
 	FString Result = Cache.GetName();
 
-	if (UniqueCrossModuleReferences)
+	// We don't need to export UFunction externs, though we may need the externs for UDelegateFunctions
+	if (UniqueCrossModuleReferences && (!Item->IsA<UFunction>() || Item->IsA<UDelegateFunction>()))
 	{
 		const FString& Extern = Cache.GetExternDecl();
 		UniqueCrossModuleReferences->Add(*Extern);
@@ -1974,8 +1975,8 @@ void FNativeClassHeaderGenerator::ExportGeneratedPackageInitCode(FOutputDevice&
 		for (UField* ScriptType : *SingletonsToOutput)
 		{
 			Out.Log(FTypeSingletonCache::Get(ScriptType, true).GetExternDecl());
-			}
 		}
+	}
 
 	FOutputDeviceNull OutputDeviceNull;
 	FString MetaDataParams = OutputMetaDataCodeForObject(OutputDeviceNull, Out, InPackage, TEXT("Package_MetaDataParams"), TEXT(""), TEXT("\t\t\t"));
@@ -2141,12 +2142,6 @@ void FNativeClassHeaderGenerator::ExportNativeGeneratedInitCode(FOutputDevice& O
 
 				if (!Function->IsA<UDelegateFunction>())
 				{
-					OutDeclarations.Logf(
-						TEXT("%s%s%s"),
-						BEGIN_WRAP_EDITOR_ONLY(bIsEditorOnlyFunction),
-						*FTypeSingletonCache::Get(Function).GetExternDecl(),
-						END_WRAP_EDITOR_ONLY(bIsEditorOnlyFunction)
-					);
 					ExportFunction(Out, SourceFile, Function, bIsNoExport);
 				}
 
