Export symbols for ValidateGeneratedRepEnums header declared functions.
This should prevent link errors for classes defined in plugins.

diff --git a/Engine/Source/Programs/UnrealHeaderTool/Private/CodeGenerator.cpp b/Engine/Source/Programs/UnrealHeaderTool/Private/CodeGenerator.cpp
--- a/Engine/Source/Programs/UnrealHeaderTool/Private/CodeGenerator.cpp
+++ b/Engine/Source/Programs/UnrealHeaderTool/Private/CodeGenerator.cpp
@@ -161,7 +161,7 @@ namespace
 		return false;
 	}
 
-	static void ExportNetData(FOutputDevice& Out, UClass* Class, UClass* ParentRepClass)
+	static void ExportNetData(FOutputDevice& Out, UClass* Class, UClass* ParentRepClass, const TCHAR* API)
 	{
 		const TArray<FRepRecord>& ClassReps = Class->ClassReps;
 		
@@ -235,13 +235,15 @@ namespace
 
 		Out.Logf(TEXT(""
 			"%s\r\n" // NetFields
-			"\tvirtual void ValidateGeneratedRepEnums(const TArray<struct FRepRecord>& ClassReps) const override;\r\n"),
-			*NetFieldBuilder);
+			"\t%s_API virtual void ValidateGeneratedRepEnums(const TArray<struct FRepRecord>& ClassReps) const override;\r\n"),
+			*NetFieldBuilder,
+			API);
 	}
 
 	static void WriteReplicatedMacroData(
 		const ClassDefinitionRange& ClassRange,
 		const TCHAR* ClassCPPName,
+		const TCHAR* API,
 		FClass* Class,
 		FClass* SuperClass,
 		FOutputDevice& Writer,
@@ -265,7 +267,7 @@ namespace
 		UClass* ParentRepClass = nullptr;
 		Class->SetUpUhtReplicationData(&ParentRepClass);
 
-		ExportNetData(Writer, Class, ParentRepClass);
+		ExportNetData(Writer, Class, ParentRepClass, API);
 	}
 }
 
@@ -3091,7 +3093,7 @@ void FNativeClassHeaderGenerator::ExportClassFromSourceFileInner(
 
 			if (ClassHasReplicatedProperties(Class))
 				{
-				WriteReplicatedMacroData(ClassRange, ClassCPPName, Class, SuperClass, InterfaceBoilerplate, SourceFile);
+				WriteReplicatedMacroData(ClassRange, ClassCPPName, *APIArg, Class, SuperClass, InterfaceBoilerplate, SourceFile);
 			}
 
 			FString NoPureDeclsMacroName = SourceFile.GetGeneratedMacroName(ClassData, TEXT("_INCLASS_IINTERFACE_NO_PURE_DECLS"));
@@ -3118,7 +3120,7 @@ void FNativeClassHeaderGenerator::ExportClassFromSourceFileInner(
 
 			if (ClassHasReplicatedProperties(Class))
 				{
-				WriteReplicatedMacroData(ClassRange, ClassCPPName, Class, SuperClass, Boilerplate, SourceFile);
+				WriteReplicatedMacroData(ClassRange, ClassCPPName, *APIArg, Class, SuperClass, Boilerplate, SourceFile);
 			}
 
 			{
