UHT optimizations part 1 - basic cleanup, minimal measurable gains
#rnx

diff --git a/Engine/Source/Programs/UnrealHeaderTool/Private/Classes.cpp b/Engine/Source/Programs/UnrealHeaderTool/Private/Classes.cpp
--- a/Engine/Source/Programs/UnrealHeaderTool/Private/Classes.cpp
+++ b/Engine/Source/Programs/UnrealHeaderTool/Private/Classes.cpp
@@ -16,7 +16,7 @@ namespace
 	 * @param InNameToCheck - Name w/ potential prefix to check
 	 * @param OriginalClass - Class to check against
 	 */
-	bool ClassNameHasValidPrefix(const FString InNameToCheck, const FClass* OriginalClass)
+	bool ClassNameHasValidPrefix(const FString& InNameToCheck, const FClass* OriginalClass)
 	{
 		bool bIsLabledDeprecated;
 		GetClassPrefix( InNameToCheck, bIsLabledDeprecated );
@@ -79,7 +79,7 @@ bool FClasses::IsDependentOn(const FClass* Suspect, const FClass* Source, TSet<c
 		return true;
 	}
 
-	// Prevent circular #includes from causing inifinite recursion
+	// Prevent circular #includes from causing infinite recursion
 	// Note that although it may mean there's a circular dependency somewhere, it does not
 	// necessarily mean it's the one we're looking for
 	if (VisitedDpendencies.Contains(Suspect))
@@ -101,10 +101,14 @@ FClass* FClasses::FindClass(const TCHAR* ClassName) const
 	UObject* ClassPackage = ANY_PACKAGE;
 
 	if (UClass* Result = FindObject<UClass>(ClassPackage, ClassName))
+	{
 		return (FClass*)Result;
+	}
 
 	if (UObjectRedirector* RenamedClassRedirector = FindObject<UObjectRedirector>(ClassPackage, ClassName))
+	{
 		return (FClass*)CastChecked<UClass>(RenamedClassRedirector->DestinationObject);
+	}
 
 	return nullptr;
 }
@@ -116,7 +120,8 @@ TArray<FClass*> FClasses::GetDerivedClasses(FClass* Parent) const
 	ClassLeaf->GetChildClasses(ChildLeaves);
 
 	TArray<FClass*> Result;
-	for (auto Node : ChildLeaves)
+	Result.Reserve(ChildLeaves.Num());
+	for (const FClassTree* Node : ChildLeaves)
 	{
 		Result.Add((FClass*)Node->GetClass());
 	}
@@ -127,20 +132,13 @@ TArray<FClass*> FClasses::GetDerivedClasses(FClass* Parent) const
 FClass* FClasses::FindAnyClass(const TCHAR* ClassName) const
 {
 	check(ClassName);
-
 	return (FClass*)FindObject<UClass>(ANY_PACKAGE, ClassName);
 }
 
-FClass* FClasses::FindScriptClass(const FString& InClassName) const
-{
-	FString ErrorMsg;
-	return FindScriptClass(InClassName, ErrorMsg);
-}
-
 FClass* FClasses::FindScriptClassOrThrow(const FString& InClassName) const
 {
 	FString ErrorMsg;
-	if (FClass* Result = FindScriptClass(InClassName, ErrorMsg))
+	if (FClass* Result = FindScriptClass(InClassName, &ErrorMsg))
 	{
 		return Result;
 	}
@@ -151,7 +149,7 @@ FClass* FClasses::FindScriptClassOrThrow(const FString& InClassName) const
 	return 0;
 }
 
-FClass* FClasses::FindScriptClass(const FString& InClassName, FString& OutErrorMsg) const
+FClass* FClasses::FindScriptClass(const FString& InClassName, FString* OutErrorMsg) const
 {
 	// Strip the class name of its prefix and then do a search for the class
 	FString ClassNameStripped = GetClassNameWithPrefixRemoved(InClassName);
@@ -160,8 +158,11 @@ FClass* FClasses::FindScriptClass(const FString& InClassName, FString& OutErrorM
 		// If the class was found with the stripped class name, verify that the correct prefix was used and throw an error otherwise
 		if (!ClassNameHasValidPrefix(InClassName, FoundClass))
 		{
-			OutErrorMsg = FString::Printf(TEXT("Class '%s' has an incorrect prefix, expecting '%s'"), *InClassName, *FoundClass->GetNameWithPrefix());
-			return NULL;
+			if (OutErrorMsg)
+			{
+				*OutErrorMsg = FString::Printf(TEXT("Class '%s' has an incorrect prefix, expecting '%s'"), *InClassName, *FoundClass->GetNameWithPrefix());
+			}
+			return nullptr;
 		}
 
 		return (FClass*)FoundClass;
@@ -172,15 +173,21 @@ FClass* FClasses::FindScriptClass(const FString& InClassName, FString& OutErrorM
 	if (FClass* FoundClass = FindClass(*InClassName))
 	{
 		// If the class was found with the given identifier, the user forgot to use the correct Unreal prefix	
-		OutErrorMsg = FString::Printf(TEXT("Class '%s' is missing a prefix, expecting '%s'"), *InClassName, *FoundClass->GetNameWithPrefix());
+		if (OutErrorMsg)
+		{
+			*OutErrorMsg = FString::Printf(TEXT("Class '%s' is missing a prefix, expecting '%s'"), *InClassName, *FoundClass->GetNameWithPrefix());
+		}
 	}
 	else
 	{
 		// If the class was still not found, it wasn't a valid identifier
-		OutErrorMsg = FString::Printf(TEXT("Class '%s' not found."), *InClassName );
+		if (OutErrorMsg)
+		{
+			*OutErrorMsg = FString::Printf(TEXT("Class '%s' not found."), *InClassName);
+		}
 	}
 
-	return NULL;
+	return nullptr;
 }
 
 TArray<FClass*> FClasses::GetClassesInPackage(const UPackage* InPackage) const
diff --git a/Engine/Source/Programs/UnrealHeaderTool/Private/Classes.h b/Engine/Source/Programs/UnrealHeaderTool/Private/Classes.h
--- a/Engine/Source/Programs/UnrealHeaderTool/Private/Classes.h
+++ b/Engine/Source/Programs/UnrealHeaderTool/Private/Classes.h
@@ -40,15 +40,6 @@ public:
 
 	FClass* FindAnyClass(const TCHAR* ClassName) const;
 
-	/** 
-	 * Attempts to find a script class based on the given name. Will attempt to strip
-	 * the prefix of the given name while searching. Throws script errors when appropriate.
-	 *
-	 * @param   InClassName  Name w/ Unreal prefix to use when searching for a class
-	 * @return               The found class, or NULL if the class was not found.
-	 */
-	FClass* FindScriptClass(const FString& InClassName) const;
-
 	/** 
 	 * Attempts to find a script class based on the given name. Will attempt to strip
 	 * the prefix of the given name while searching. Throws an exception with the script error
@@ -61,13 +52,13 @@ public:
 
 	/** 
 	 * Attempts to find a script class based on the given name. Will attempt to strip
-	 * the prefix of the given name while searching. Throws script errors when appropriate.
+	 * the prefix of the given name while searching. Optionally returns script errors when appropriate.
 	 *
 	 * @param   InClassName  Name w/ Unreal prefix to use when searching for a class
 	 * @param   OutErrorMsg  Error message (if any) giving the caller flexibility in how they present an error
 	 * @return               The found class, or NULL if the class was not found.
 	 */
-	FClass* FindScriptClass(const FString& InClassName, FString& OutErrorMsg) const;
+	FClass* FindScriptClass(const FString& InClassName, FString* OutErrorMsg = nullptr) const;
 
 	/**
 	 * Returns an array of classes for the given package.
diff --git a/Engine/Source/Programs/UnrealHeaderTool/Private/HeaderProvider.cpp b/Engine/Source/Programs/UnrealHeaderTool/Private/HeaderProvider.cpp
--- a/Engine/Source/Programs/UnrealHeaderTool/Private/HeaderProvider.cpp
+++ b/Engine/Source/Programs/UnrealHeaderTool/Private/HeaderProvider.cpp
@@ -52,5 +52,5 @@ const FString& FHeaderProvider::GetId() const
 
 bool operator==(const FHeaderProvider& A, const FHeaderProvider& B)
 {
-	return A.Id == B.Id && A.Type == B.Type;
+	return A.Type == B.Type && A.Id == B.Id;
 }
diff --git a/Engine/Source/Programs/UnrealHeaderTool/Private/ParserHelper.cpp b/Engine/Source/Programs/UnrealHeaderTool/Private/ParserHelper.cpp
--- a/Engine/Source/Programs/UnrealHeaderTool/Private/ParserHelper.cpp
+++ b/Engine/Source/Programs/UnrealHeaderTool/Private/ParserHelper.cpp
@@ -160,12 +160,14 @@ void FToken::Clone( const FToken& Other )
 
 /////////////////////////////////////////////////////
 // FAdvancedDisplayParameterHandler
+static const FName NAME_AdvancedDisplay(TEXT("AdvancedDisplay"));
+
 FAdvancedDisplayParameterHandler::FAdvancedDisplayParameterHandler(const TMap<FName, FString>* MetaData)
 	: NumberLeaveUnmarked(-1), AlreadyLeft(0), bUseNumber(false)
 {
 	if(MetaData)
 	{
-		const FString* FoundString = MetaData->Find(FName(TEXT("AdvancedDisplay")));
+		const FString* FoundString = MetaData->Find(NAME_AdvancedDisplay);
 		if(FoundString)
 		{
 			FoundString->ParseIntoArray(ParametersNames, TEXT(","), true);
diff --git a/Engine/Source/Programs/UnrealHeaderTool/Private/StringUtils.cpp b/Engine/Source/Programs/UnrealHeaderTool/Private/StringUtils.cpp
--- a/Engine/Source/Programs/UnrealHeaderTool/Private/StringUtils.cpp
+++ b/Engine/Source/Programs/UnrealHeaderTool/Private/StringUtils.cpp
@@ -5,38 +5,42 @@
 #include "Containers/UnrealString.h"
 #include "Hash/CityHash.h"
 
-FString GetClassNameWithPrefixRemoved(const FString InClassName)
+FString GetClassNameWithPrefixRemoved(const FString& InClassName)
 {
+	FString Result;
 	const FString ClassPrefix = GetClassPrefix( InClassName );
 	if( !ClassPrefix.IsEmpty() )
 	{	
-		return InClassName.Right(InClassName.Len() - ClassPrefix.Len());
+		Result = InClassName.Right(InClassName.Len() - ClassPrefix.Len());
 	}
-	return FString();
+	return Result;
 }
 
-FString GetClassNameWithoutPrefix( const FString& InClassNameOrFilename )
+FString GetClassNameWithoutPrefix(FString InClassNameOrFilename)
 {
+	FString ClassNameWithoutPrefix = MoveTemp(InClassNameOrFilename);
+
 	// Check for header names (they don't come with a full path so we only search for the first dot)
-	const int32 DotIndex = InClassNameOrFilename.Find(TEXT("."), ESearchCase::CaseSensitive);
-	if (DotIndex == INDEX_NONE)
+	int32 DotIndex;
+	if (ClassNameWithoutPrefix.FindChar(TEXT('.'), DotIndex))
 	{
-		const FString ClassPrefix = GetClassPrefix( InClassNameOrFilename );
-		return InClassNameOrFilename.Right(InClassNameOrFilename.Len() - ClassPrefix.Len());
+		ClassNameWithoutPrefix.MidInline(0, DotIndex, false);
 	}
 	else
 	{
-		return InClassNameOrFilename.Mid(0, DotIndex);
+		const FString ClassPrefix = GetClassPrefix(ClassNameWithoutPrefix);
+		ClassNameWithoutPrefix.RightInline(ClassNameWithoutPrefix.Len() - ClassPrefix.Len(), false);
 	}
+	return ClassNameWithoutPrefix;
 }
 
-FString GetClassPrefix( const FString InClassName )
+FString GetClassPrefix( const FString& InClassName )
 {
 	bool bIsLabledDeprecated;
 	return GetClassPrefix(InClassName, /*out*/ bIsLabledDeprecated);
 }
 
-FString GetClassPrefix( const FString InClassName, bool& bIsLabeledDeprecated )
+FString GetClassPrefix(const FString& InClassName, bool& bIsLabeledDeprecated )
 {
 	FString ClassPrefix = InClassName.Left(1);
 
@@ -51,7 +55,7 @@ FString GetClassPrefix( const FString InClassName, bool& bIsLabeledDeprecated )
 		case TEXT('A'):
 		case TEXT('U'):
 			// If it is a class prefix, check for deprecated class prefix also
-			if (InClassName.Mid(1, 11).Compare(TEXT("DEPRECATED_"), ESearchCase::CaseSensitive) == 0)
+			if (InClassName.Len() > 12 && FCString::Strncmp(&(InClassName[1]), TEXT("DEPRECATED_"), 11) == 0)
 			{
 				bIsLabeledDeprecated = true;
 				ClassPrefix = InClassName.Left(12);
@@ -65,7 +69,7 @@ FString GetClassPrefix( const FString InClassName, bool& bIsLabeledDeprecated )
 
 		default:
 			// If it's not a class or struct prefix, it's invalid
-			ClassPrefix.Empty();
+			ClassPrefix.Reset();
 			break;
 		}
 	}
diff --git a/Engine/Source/Programs/UnrealHeaderTool/Private/StringUtils.h b/Engine/Source/Programs/UnrealHeaderTool/Private/StringUtils.h
--- a/Engine/Source/Programs/UnrealHeaderTool/Private/StringUtils.h
+++ b/Engine/Source/Programs/UnrealHeaderTool/Private/StringUtils.h
@@ -10,21 +10,21 @@
  *
  * @param InClassName Class Name with a prefix to be stripped
  */
-FString GetClassNameWithPrefixRemoved(const FString InClassName);
+FString GetClassNameWithPrefixRemoved(const FString& InClassName);
 
 /** 
  * Attempts to strip the given class name of its affixed prefix. If no prefix exists, will return unchanged string.
  *
  * @param InClassName Class Name with a prefix to be stripped
  */
-FString GetClassNameWithoutPrefix(const FString& InClassNameOrFilename);
+FString GetClassNameWithoutPrefix(FString InClassNameOrFilename);
 
 /** 
  * Attempts to get class prefix. If the given class name does not start with a valid Unreal Prefix, it will return an empty string.
  *
  * @param InClassName Name w/ potential prefix to check
  */
-FString GetClassPrefix(const FString InClassName);
+FString GetClassPrefix(const FString& InClassName);
 
 /** 
  * Attempts to get class prefix. If the given class name does not start with a valid Unreal Prefix, it will return an empty string.
@@ -32,7 +32,7 @@ FString GetClassPrefix(const FString InClassName);
  * @param InClassName Name w/ potential prefix to check
  * @param out bIsLabeledDeprecated Reference param set to True if the class name is marked as deprecated
  */
-FString GetClassPrefix(const FString InClassName, bool& bIsLabeledDeprecated);
+FString GetClassPrefix(const FString& InClassName, bool& bIsLabeledDeprecated);
 
 /**
 * Generates hash for text, ignoring characters such as \r and the terminating \0
diff --git a/Engine/Source/Programs/UnrealHeaderTool/Private/UnrealHeaderToolMain.cpp b/Engine/Source/Programs/UnrealHeaderTool/Private/UnrealHeaderToolMain.cpp
--- a/Engine/Source/Programs/UnrealHeaderTool/Private/UnrealHeaderToolMain.cpp
+++ b/Engine/Source/Programs/UnrealHeaderTool/Private/UnrealHeaderToolMain.cpp
@@ -35,7 +35,7 @@ INT32_MAIN_INT32_ARGC_TCHAR_ARGV()
 	for (int32 Arg = 0; Arg < ArgC; Arg++)
 	{
 		FString LocalArg = ArgV[Arg];
-		if (LocalArg.Contains(TEXT(" ")))
+		if (LocalArg.Contains(TEXT(" "), ESearchCase::CaseSensitive))
 		{
 			CmdLine += TEXT("\"");
 			CmdLine += LocalArg;
