UHT optimizations part 2 - Class Metadata improvements
* Added a sorted specifier list to reduce string compares when evaluating metadata tags
* Use move semantices where possible
* Use RemoveSwap instead of Remove when processing show/hide metadata tags. Generated code can change as a result of hide/show tag order changing.
#rnx

diff --git a/Engine/Source/Programs/UnrealHeaderTool/Private/ClassDeclarationMetaData.cpp b/Engine/Source/Programs/UnrealHeaderTool/Private/ClassDeclarationMetaData.cpp
--- a/Engine/Source/Programs/UnrealHeaderTool/Private/ClassDeclarationMetaData.cpp
+++ b/Engine/Source/Programs/UnrealHeaderTool/Private/ClassDeclarationMetaData.cpp
@@ -8,16 +8,19 @@
 #include "BaseParser.h"
 #include "HeaderParser.h"
 #include "Classes.h"
+#include "Specifiers/ClassMetadataSpecifiers.h"
+#include "Algo/FindSortedStringCaseInsensitive.h"
 
 // Utility functions
 namespace
 {
+	static const FName NAME_IgnoreCategoryKeywordsInSubclasses(TEXT("IgnoreCategoryKeywordsInSubclasses"));
+	static const FName NAME_IsConversionRoot(TEXT("IsConversionRoot"));
 	bool IsActorClass(UClass *Class)
 	{
-		static FName ActorName = FName(TEXT("Actor"));
 		while (Class)
 		{
-			if (Class->GetFName() == ActorName)
+			if (Class->GetFName() == NAME_Actor)
 			{
 				return true;
 			}
@@ -33,7 +36,7 @@ FClassDeclarationMetaData::FClassDeclarationMetaData()
 {
 }
 
-void FClassDeclarationMetaData::ParseClassProperties(const TArray<FPropertySpecifier>& InClassSpecifiers, const FString& InRequiredAPIMacroIfPresent)
+void FClassDeclarationMetaData::ParseClassProperties(TArray<FPropertySpecifier>&& InClassSpecifiers, const FString& InRequiredAPIMacroIfPresent)
 {
 	ClassFlags = CLASS_None;
 	// Record that this class is RequiredAPI if the CORE_API style macro was present
@@ -45,222 +48,256 @@ void FClassDeclarationMetaData::ParseClassProperties(const TArray<FPropertySpeci
 
 	// Process all of the class specifiers
 
-	bool bWithinSpecified = false;
-	bool bDeclaresConfigFile = false;
-	for (const FPropertySpecifier& PropSpecifier : InClassSpecifiers)
+	for (FPropertySpecifier& PropSpecifier : InClassSpecifiers)
 	{
-		const FString& Specifier = PropSpecifier.Key;
-
-		if (Specifier == TEXT("noexport"))
-		{
-			// Don't export to C++ header.
-			ClassFlags |= CLASS_NoExport;
-		}
-		else if (Specifier == TEXT("intrinsic"))
-		{
-			ClassFlags |= CLASS_Intrinsic;
-		}
-		else if (Specifier == TEXT("ComponentWrapperClass"))
-		{
-			MetaData.Add(TEXT("IgnoreCategoryKeywordsInSubclasses"), TEXT("true"));
-		}
-		else if (Specifier == TEXT("within"))
-		{
-			ClassWithin = FHeaderParser::RequireExactlyOneSpecifierValue(PropSpecifier);
-		}
-		else if (Specifier == TEXT("editinlinenew"))
-		{
-			// Class can be constructed from the New button in editinline
-			ClassFlags |= CLASS_EditInlineNew;
-		}
-		else if (Specifier == TEXT("noteditinlinenew"))
-		{
-			// Class cannot be constructed from the New button in editinline
-			ClassFlags &= ~CLASS_EditInlineNew;
-		}
-		else if (Specifier == TEXT("placeable"))
-		{
-			WantsToBePlaceable = true;
-			ClassFlags &= ~CLASS_NotPlaceable;
-		}
-		else if (Specifier == TEXT("defaulttoinstanced"))
-		{
-			// these classed default to instanced.
-			ClassFlags |= CLASS_DefaultToInstanced;
-		}
-		else if (Specifier == TEXT("notplaceable"))
-		{
-			// Don't allow the class to be placed in the editor.
-			ClassFlags |= CLASS_NotPlaceable;
-		}
-		else if (Specifier == TEXT("hidedropdown"))
-		{
-			// Prevents class from appearing in class comboboxes in the property window
-			ClassFlags |= CLASS_HideDropDown;
-		}
-		else if (Specifier == TEXT("dependsOn"))
-		{
-			FError::Throwf(TEXT("The dependsOn specifier is deprecated. Please use #include \"ClassHeaderFilename.h\" instead."));
-		}
-		else if (Specifier == TEXT("MinimalAPI"))
+		switch ((EClassMetadataSpecifier)Algo::FindSortedStringCaseInsensitive(*PropSpecifier.Key, GClassMetadataSpecifierStrings))
 		{
-			ClassFlags |= CLASS_MinimalAPI;
-		}
-		else if (Specifier == TEXT("const"))
-		{
-			ClassFlags |= CLASS_Const;
-		}
-		else if (Specifier == TEXT("perObjectConfig"))
-		{
-			ClassFlags |= CLASS_PerObjectConfig;
-		}
-		else if (Specifier == TEXT("configdonotcheckdefaults"))
-		{
-			ClassFlags |= CLASS_ConfigDoNotCheckDefaults;
-		}
-		else if (Specifier == TEXT("abstract"))
-		{
-			// Hide all editable properties.
-			ClassFlags |= CLASS_Abstract;
-		}
-		else if (Specifier == TEXT("deprecated"))
-		{
-			ClassFlags |= CLASS_Deprecated;
+			case EClassMetadataSpecifier::NoExport:
 
-			// Don't allow the class to be placed in the editor.
-			ClassFlags |= CLASS_NotPlaceable;
-		}
-		else if (Specifier == TEXT("transient"))
-		{
-			// Transient class.
-			ClassFlags |= CLASS_Transient;
-		}
-		else if (Specifier == TEXT("nonTransient"))
-		{
-			// this child of a transient class is not transient - remove the transient flag
-			ClassFlags &= ~CLASS_Transient;
-		}
-		else if (Specifier == TEXT("customConstructor"))
-		{
-			// we will not export a constructor for this class, assuming it is in the CPP block
-			ClassFlags |= CLASS_CustomConstructor;
-		}
-		else if (Specifier == TEXT("config"))
-		{
-			// Class containing config properties - parse the name of the config file to use
-			ConfigName = FHeaderParser::RequireExactlyOneSpecifierValue(PropSpecifier);
-		}
-		else if (Specifier == TEXT("defaultconfig"))
-		{
-			// Save object config only to Default INIs, never to local INIs.
-			ClassFlags |= CLASS_DefaultConfig;
-		}
-		else if (Specifier == TEXT("globaluserconfig"))
-		{
-			// Save object config only to global user overrides, never to local INIs
-			ClassFlags |= CLASS_GlobalUserConfig;
-		}
-		else if (Specifier == TEXT("showCategories"))
-		{
-			FHeaderParser::RequireSpecifierValue(PropSpecifier);
+				// Don't export to C++ header.
+				ClassFlags |= CLASS_NoExport;
+				break;
 
-			for (const FString& Value : PropSpecifier.Values)
-			{
-				ShowCategories.AddUnique(Value);
-			}
-		}
-		else if (Specifier == TEXT("hideCategories"))
-		{
-			FHeaderParser::RequireSpecifierValue(PropSpecifier);
+			case EClassMetadataSpecifier::Intrinsic:
 
-			for (const FString& Value : PropSpecifier.Values)
-			{
-				HideCategories.AddUnique(Value);
-			}
-		}
-		else if (Specifier == TEXT("showFunctions"))
-		{
-			FHeaderParser::RequireSpecifierValue(PropSpecifier);
+				ClassFlags |= CLASS_Intrinsic;
+				break;
 
-			for (const FString& Value : PropSpecifier.Values)
-			{
-				HideFunctions.Remove(Value);
-			}
-		}
-		else if (Specifier == TEXT("hideFunctions"))
-		{
-			FHeaderParser::RequireSpecifierValue(PropSpecifier);
+			case EClassMetadataSpecifier::ComponentWrapperClass:
 
-			for (const FString& Value : PropSpecifier.Values)
-			{
-				HideFunctions.AddUnique(Value);
-			}
-		}
-		// Currently some code only handles a single sidecar data structure so we enforce that here
-		else if (Specifier == TEXT("SparseClassDataTypes"))
-		{
-			SparseClassDataTypes.AddUnique(FHeaderParser::RequireExactlyOneSpecifierValue(PropSpecifier));
-		}
-		else if (Specifier == TEXT("classGroup"))
-		{
-			FHeaderParser::RequireSpecifierValue(PropSpecifier);
+				MetaData.Add(NAME_IgnoreCategoryKeywordsInSubclasses, TEXT("true"));
+				break;
 
-			for (const FString& Value : PropSpecifier.Values)
-			{
-				ClassGroupNames.Add(Value);
-			}
-		}
-		else if (Specifier == TEXT("autoExpandCategories"))
-		{
-			FHeaderParser::RequireSpecifierValue(PropSpecifier);
+			case EClassMetadataSpecifier::Within:
 
-			for (const FString& Value : PropSpecifier.Values)
-			{
-				AutoCollapseCategories.Remove(Value);
-				AutoExpandCategories.AddUnique(Value);
-			}
-		}
-		else if (Specifier == TEXT("autoCollapseCategories"))
-		{
-			FHeaderParser::RequireSpecifierValue(PropSpecifier);
+				ClassWithin = FHeaderParser::RequireExactlyOneSpecifierValue(PropSpecifier);
+				break;
 
-			for (const FString& Value : PropSpecifier.Values)
-			{
-				AutoExpandCategories.Remove(Value);
-				AutoCollapseCategories.AddUnique(Value);
-			}
-		}
-		else if (Specifier == TEXT("dontAutoCollapseCategories"))
-		{
-			FHeaderParser::RequireSpecifierValue(PropSpecifier);
+			case EClassMetadataSpecifier::EditInlineNew:
 
-			for (const FString& Value : PropSpecifier.Values)
-			{
-				AutoCollapseCategories.Remove(Value);
-			}
-		}
-		else if (Specifier == TEXT("collapseCategories"))
-		{
-			// Class' properties should not be shown categorized in the editor.
-			ClassFlags |= CLASS_CollapseCategories;
-		}
-		else if (Specifier == TEXT("dontCollapseCategories"))
-		{
-			// Class' properties should be shown categorized in the editor.
-			ClassFlags &= ~CLASS_CollapseCategories;
-		}
-		else if (Specifier == TEXT("AdvancedClassDisplay"))
-		{
-			// By default the class properties are shown in advanced sections in UI
-			ClassFlags |= CLASS_AdvancedDisplay;
-		}
-		else if (Specifier == TEXT("ConversionRoot"))
-		{
-			MetaData.Add(FName(TEXT("IsConversionRoot")), "true");
-		}
-		else
-		{
-			FError::Throwf(TEXT("Unknown class specifier '%s'"), *Specifier);
+				// Class can be constructed from the New button in editinline
+				ClassFlags |= CLASS_EditInlineNew;
+				break;
+
+			case EClassMetadataSpecifier::NotEditInlineNew:
+
+				// Class cannot be constructed from the New button in editinline
+				ClassFlags &= ~CLASS_EditInlineNew;
+				break;
+
+			case EClassMetadataSpecifier::Placeable:
+
+				WantsToBePlaceable = true;
+				ClassFlags &= ~CLASS_NotPlaceable;
+				break;
+
+			case EClassMetadataSpecifier::DefaultToInstanced:
+
+				// these classed default to instanced.
+				ClassFlags |= CLASS_DefaultToInstanced;
+				break;
+
+			case EClassMetadataSpecifier::NotPlaceable:
+
+				// Don't allow the class to be placed in the editor.
+				ClassFlags |= CLASS_NotPlaceable;
+				break;
+
+			case EClassMetadataSpecifier::HideDropdown:
+
+				// Prevents class from appearing in class comboboxes in the property window
+				ClassFlags |= CLASS_HideDropDown;
+				break;
+
+			case EClassMetadataSpecifier::DependsOn:
+
+				FError::Throwf(TEXT("The dependsOn specifier is deprecated. Please use #include \"ClassHeaderFilename.h\" instead."));
+				break;
+
+			case EClassMetadataSpecifier::MinimalAPI:
+
+				ClassFlags |= CLASS_MinimalAPI;
+				break;
+
+			case EClassMetadataSpecifier::Const:
+
+				ClassFlags |= CLASS_Const;
+				break;
+
+			case EClassMetadataSpecifier::PerObjectConfig:
+
+				ClassFlags |= CLASS_PerObjectConfig;
+				break;
+
+			case EClassMetadataSpecifier::ConfigDoNotCheckDefaults:
+
+				ClassFlags |= CLASS_ConfigDoNotCheckDefaults;
+				break;
+
+			case EClassMetadataSpecifier::Abstract:
+
+				// Hide all editable properties.
+				ClassFlags |= CLASS_Abstract;
+				break;
+
+			case EClassMetadataSpecifier::Deprecated:
+
+				ClassFlags |= CLASS_Deprecated;
+
+				// Don't allow the class to be placed in the editor.
+				ClassFlags |= CLASS_NotPlaceable;
+
+				break;
+
+			case EClassMetadataSpecifier::Transient:
+
+				// Transient class.
+				ClassFlags |= CLASS_Transient;
+				break;
+
+			case EClassMetadataSpecifier::NonTransient:
+
+				// this child of a transient class is not transient - remove the transient flag
+				ClassFlags &= ~CLASS_Transient;
+				break;
+
+			case EClassMetadataSpecifier::CustomConstructor:
+
+				// we will not export a constructor for this class, assuming it is in the CPP block
+				ClassFlags |= CLASS_CustomConstructor;
+				break;
+
+			case EClassMetadataSpecifier::Config:
+
+				// Class containing config properties - parse the name of the config file to use
+				ConfigName = FHeaderParser::RequireExactlyOneSpecifierValue(PropSpecifier);
+				break;
+
+			case EClassMetadataSpecifier::DefaultConfig:
+
+				// Save object config only to Default INIs, never to local INIs.
+				ClassFlags |= CLASS_DefaultConfig;
+				break;
+
+			case EClassMetadataSpecifier::GlobalUserConfig:
+
+				// Save object config only to global user overrides, never to local INIs
+				ClassFlags |= CLASS_GlobalUserConfig;
+				break;
+
+			case EClassMetadataSpecifier::ShowCategories:
+
+				FHeaderParser::RequireSpecifierValue(PropSpecifier);
+
+				for (FString& Value : PropSpecifier.Values)
+				{
+					ShowCategories.AddUnique(MoveTemp(Value));
+				}
+				break;
+
+			case EClassMetadataSpecifier::HideCategories:
+
+				FHeaderParser::RequireSpecifierValue(PropSpecifier);
+
+				for (FString& Value : PropSpecifier.Values)
+				{
+					HideCategories.AddUnique(MoveTemp(Value));
+				}
+				break;
+
+			case EClassMetadataSpecifier::ShowFunctions:
+
+				FHeaderParser::RequireSpecifierValue(PropSpecifier);
+
+				for (const FString& Value : PropSpecifier.Values)
+				{
+					HideFunctions.RemoveSwap(Value);
+				}
+				break;
+
+			case EClassMetadataSpecifier::HideFunctions:
+
+				FHeaderParser::RequireSpecifierValue(PropSpecifier);
+
+				for (FString& Value : PropSpecifier.Values)
+				{
+					HideFunctions.AddUnique(MoveTemp(Value));
+				}
+				break;
+
+			// Currently some code only handles a single sidecar data structure so we enforce that here
+			case EClassMetadataSpecifier::SparseClassDataTypes:
+
+				SparseClassDataTypes.AddUnique(FHeaderParser::RequireExactlyOneSpecifierValue(PropSpecifier));
+				break;
+
+			case EClassMetadataSpecifier::ClassGroup:
+
+				FHeaderParser::RequireSpecifierValue(PropSpecifier);
+
+				for (FString& Value : PropSpecifier.Values)
+				{
+					ClassGroupNames.Add(MoveTemp(Value));
+				}
+				break;
+
+			case EClassMetadataSpecifier::AutoExpandCategories:
+
+				FHeaderParser::RequireSpecifierValue(PropSpecifier);
+
+				for (FString& Value : PropSpecifier.Values)
+				{
+					AutoCollapseCategories.RemoveSwap(Value);
+					AutoExpandCategories.AddUnique(MoveTemp(Value));
+				}
+				break;
+
+			case EClassMetadataSpecifier::AutoCollapseCategories:
+
+				FHeaderParser::RequireSpecifierValue(PropSpecifier);
+
+				for (FString& Value : PropSpecifier.Values)
+				{
+					AutoExpandCategories.RemoveSwap(Value);
+					AutoCollapseCategories.AddUnique(MoveTemp(Value));
+				}
+				break;
+
+			case EClassMetadataSpecifier::DontAutoCollapseCategories:
+
+				FHeaderParser::RequireSpecifierValue(PropSpecifier);
+
+				for (const FString& Value : PropSpecifier.Values)
+				{
+					AutoCollapseCategories.RemoveSwap(Value);
+				}
+				break;
+
+			case EClassMetadataSpecifier::CollapseCategories:
+
+				// Class' properties should not be shown categorized in the editor.
+				ClassFlags |= CLASS_CollapseCategories;
+				break;
+
+			case EClassMetadataSpecifier::DontCollapseCategories:
+
+				// Class' properties should be shown categorized in the editor.
+				ClassFlags &= ~CLASS_CollapseCategories;
+				break;
+
+			case EClassMetadataSpecifier::AdvancedClassDisplay:
+
+				// By default the class properties are shown in advanced sections in UI
+				ClassFlags |= CLASS_AdvancedDisplay;
+				break;
+
+			case EClassMetadataSpecifier::ConversionRoot:
+
+				MetaData.Add(NAME_IsConversionRoot, TEXT("true"));
+				break;
+
+			default:
+				FError::Throwf(TEXT("Unknown class specifier '%s'"), *PropSpecifier.Key);
 		}
 	}
 }
@@ -270,7 +307,7 @@ void FClassDeclarationMetaData::MergeShowCategories()
 	for (const FString& Value : ShowCategories)
 	{
 		// if we didn't find this specific category path in the HideCategories metadata
-		if (HideCategories.Remove(Value) == 0)
+		if (HideCategories.RemoveSwap(Value) == 0)
 		{
 			TArray<FString> SubCategoryList;
 			Value.ParseIntoArray(SubCategoryList, TEXT("|"), true);
@@ -299,61 +336,59 @@ void FClassDeclarationMetaData::MergeClassCategories(FClass* Class)
 	TArray<FString> ParentHideCategories;
 	TArray<FString> ParentShowSubCatgories;
 	TArray<FString> ParentHideFunctions;
-	TArray<FString> ParentSparseClassDataTypes;
 	TArray<FString> ParentAutoExpandCategories;
 	TArray<FString> ParentAutoCollapseCategories;
 	Class->GetHideCategories(ParentHideCategories);
 	Class->GetShowCategories(ParentShowSubCatgories);
 	Class->GetHideFunctions(ParentHideFunctions);
-	Class->GetSparseClassDataTypes(ParentSparseClassDataTypes);
 	Class->GetAutoExpandCategories(ParentAutoExpandCategories);
 	Class->GetAutoCollapseCategories(ParentAutoCollapseCategories);
 
 	// Add parent categories. We store the opposite of HideCategories and HideFunctions in a separate array anyway.
-	HideCategories.Append(ParentHideCategories);
-	ShowSubCatgories.Append(ParentShowSubCatgories);
-	HideFunctions.Append(ParentHideFunctions);
+	HideCategories.Append(MoveTemp(ParentHideCategories));
+	ShowSubCatgories.Append(MoveTemp(ParentShowSubCatgories));
+	HideFunctions.Append(MoveTemp(ParentHideFunctions));
 
 	MergeShowCategories();
 
 	// Merge ShowFunctions and HideFunctions
 	for (const FString& Value : ShowFunctions)
 	{
-		HideFunctions.Remove(Value);
+		HideFunctions.RemoveSwap(Value);
 	}
 	ShowFunctions.Empty();
 
 	// Merge DontAutoCollapseCategories and AutoCollapseCategories
 	for (const FString& Value : DontAutoCollapseCategories)
 	{
-		AutoCollapseCategories.Remove(Value);
+		AutoCollapseCategories.RemoveSwap(Value);
 	}
 	DontAutoCollapseCategories.Empty();
 
 	// Merge ShowFunctions and HideFunctions
 	for (const FString& Value : ShowFunctions)
 	{
-		HideFunctions.Remove(Value);
+		HideFunctions.RemoveSwap(Value);
 	}
 	ShowFunctions.Empty();
 
 	// Merge AutoExpandCategories and AutoCollapseCategories (we still want to keep AutoExpandCategories though!)
 	for (const FString& Value : AutoExpandCategories)
 	{
-		AutoCollapseCategories.Remove(Value);
-		ParentAutoCollapseCategories.Remove(Value);
+		AutoCollapseCategories.RemoveSwap(Value);
+		ParentAutoCollapseCategories.RemoveSwap(Value);
 	}
 	
 	// Do the same as above but the other way around
 	for (const FString& Value : AutoCollapseCategories)
 	{
-		AutoExpandCategories.Remove(Value);
-		ParentAutoExpandCategories.Remove(Value);
+		AutoExpandCategories.RemoveSwap(Value);
+		ParentAutoExpandCategories.RemoveSwap(Value);
 	}	
 
 	// Once AutoExpandCategories and AutoCollapseCategories for THIS class have been parsed, add the parent inherited categories
-	AutoCollapseCategories.Append(ParentAutoCollapseCategories);
-	AutoExpandCategories.Append(ParentAutoExpandCategories);
+	AutoCollapseCategories.Append(MoveTemp(ParentAutoCollapseCategories));
+	AutoExpandCategories.Append(MoveTemp(ParentAutoExpandCategories));
 }
 
 void FClassDeclarationMetaData::MergeAndValidateClassFlags(const FString& DeclaredClassName, uint32 PreviousClassFlags, FClass* Class, const FClasses& AllClasses)
diff --git a/Engine/Source/Programs/UnrealHeaderTool/Private/ClassDeclarationMetaData.h b/Engine/Source/Programs/UnrealHeaderTool/Private/ClassDeclarationMetaData.h
--- a/Engine/Source/Programs/UnrealHeaderTool/Private/ClassDeclarationMetaData.h
+++ b/Engine/Source/Programs/UnrealHeaderTool/Private/ClassDeclarationMetaData.h
@@ -38,7 +38,7 @@ public:
 	* @param	InRequiredAPIMacroIfPresent *_API macro if present (empty otherwise)
 	* @param	OutClassData Parsed class meta data
 	*/
-	void ParseClassProperties(const TArray<FPropertySpecifier>& InClassSpecifiers, const FString& InRequiredAPIMacroIfPresent);
+	void ParseClassProperties(TArray<FPropertySpecifier>&& InClassSpecifiers, const FString& InRequiredAPIMacroIfPresent);
 
 	/**
 	* Merges all category properties with the class which at this point only has its parent propagated categories
diff --git a/Engine/Source/Programs/UnrealHeaderTool/Private/HeaderParser.cpp b/Engine/Source/Programs/UnrealHeaderTool/Private/HeaderParser.cpp
--- a/Engine/Source/Programs/UnrealHeaderTool/Private/HeaderParser.cpp
+++ b/Engine/Source/Programs/UnrealHeaderTool/Private/HeaderParser.cpp
@@ -9317,7 +9317,7 @@ void FHeaderPreParser::ParseClassDeclaration(const TCHAR* Filename, const TCHAR*
 		// Add class declaration meta data so that we can access class flags before the class is fully parsed
 		TSharedRef<FClassDeclarationMetaData> DeclarationData = MakeShareable(new FClassDeclarationMetaData());
 		DeclarationData->MetaData = MetaData;
-		DeclarationData->ParseClassProperties(SpecifiersFound, RequiredAPIMacroIfPresent);
+		DeclarationData->ParseClassProperties(MoveTemp(SpecifiersFound), RequiredAPIMacroIfPresent);
 		GClassDeclarations.Add(out_StrippedClassName, DeclarationData);
 	}
 
diff --git a/Engine/Source/Programs/UnrealHeaderTool/Private/Specifiers/ClassMetadataSpecifiers.def b/Engine/Source/Programs/UnrealHeaderTool/Private/Specifiers/ClassMetadataSpecifiers.def
new file mode 100644
--- /dev/null
+++ b/Engine/Source/Programs/UnrealHeaderTool/Private/Specifiers/ClassMetadataSpecifiers.def
@@ -0,0 +1,42 @@
+// CLASS_METADATA_SPECIFIER(<specifier name>)
+//
+// Must be listed such that the strings are in lexicographical order!
+
+// [[ IncludeTool: Inline ]] // Markup to tell IncludeTool that this file is state changing and cannot be optimized out.
+
+CLASS_METADATA_SPECIFIER(Abstract)
+CLASS_METADATA_SPECIFIER(AdvancedClassDisplay)
+CLASS_METADATA_SPECIFIER(AutoCollapseCategories)
+CLASS_METADATA_SPECIFIER(AutoExpandCategories)
+CLASS_METADATA_SPECIFIER(ClassGroup)
+CLASS_METADATA_SPECIFIER(CollapseCategories)
+CLASS_METADATA_SPECIFIER(ComponentWrapperClass)
+CLASS_METADATA_SPECIFIER(Config)
+CLASS_METADATA_SPECIFIER(ConfigDoNotCheckDefaults)
+CLASS_METADATA_SPECIFIER(Const)
+CLASS_METADATA_SPECIFIER(ConversionRoot)
+CLASS_METADATA_SPECIFIER(CustomConstructor)
+CLASS_METADATA_SPECIFIER(DefaultConfig)
+CLASS_METADATA_SPECIFIER(DefaultToInstanced)
+CLASS_METADATA_SPECIFIER(DependsOn)
+CLASS_METADATA_SPECIFIER(Deprecated)
+CLASS_METADATA_SPECIFIER(DontAutoCollapseCategories)
+CLASS_METADATA_SPECIFIER(DontCollapseCategories)
+CLASS_METADATA_SPECIFIER(EditInlineNew)
+CLASS_METADATA_SPECIFIER(GlobalUserConfig)
+CLASS_METADATA_SPECIFIER(HideCategories)
+CLASS_METADATA_SPECIFIER(HideDropdown)
+CLASS_METADATA_SPECIFIER(HideFunctions)
+CLASS_METADATA_SPECIFIER(Intrinsic)
+CLASS_METADATA_SPECIFIER(MinimalAPI)
+CLASS_METADATA_SPECIFIER(NoExport)
+CLASS_METADATA_SPECIFIER(NonTransient)
+CLASS_METADATA_SPECIFIER(NotEditInlineNew)
+CLASS_METADATA_SPECIFIER(NotPlaceable)
+CLASS_METADATA_SPECIFIER(PerObjectConfig)
+CLASS_METADATA_SPECIFIER(Placeable)
+CLASS_METADATA_SPECIFIER(ShowCategories)
+CLASS_METADATA_SPECIFIER(ShowFunctions)
+CLASS_METADATA_SPECIFIER(SparseClassDataTypes)
+CLASS_METADATA_SPECIFIER(Transient)
+CLASS_METADATA_SPECIFIER(Within)
diff --git a/Engine/Source/Programs/UnrealHeaderTool/Private/Specifiers/ClassMetadataSpecifiers.h b/Engine/Source/Programs/UnrealHeaderTool/Private/Specifiers/ClassMetadataSpecifiers.h
new file mode 100644
--- /dev/null
+++ b/Engine/Source/Programs/UnrealHeaderTool/Private/Specifiers/ClassMetadataSpecifiers.h
@@ -0,0 +1,18 @@
+// Copyright Epic Games, Inc. All Rights Reserved.
+
+#pragma once
+
+#include "CoreMinimal.h"
+
+enum class EClassMetadataSpecifier
+{
+	None = -1
+
+	#define CLASS_METADATA_SPECIFIER(SpecifierName) , SpecifierName
+		#include "ClassMetadataSpecifiers.def"
+	#undef CLASS_METADATA_SPECIFIER
+
+	, Max
+};
+
+extern const TCHAR* GClassMetadataSpecifierStrings[(int32)EClassMetadataSpecifier::Max];
diff --git a/Engine/Source/Programs/UnrealHeaderTool/Private/Specifiers/SpecifierArrays.cpp b/Engine/Source/Programs/UnrealHeaderTool/Private/Specifiers/SpecifierArrays.cpp
--- a/Engine/Source/Programs/UnrealHeaderTool/Private/Specifiers/SpecifierArrays.cpp
+++ b/Engine/Source/Programs/UnrealHeaderTool/Private/Specifiers/SpecifierArrays.cpp
@@ -6,6 +6,7 @@
 #include "Misc/CString.h"
 
 #include "CheckedMetadataSpecifiers.h"
+#include "ClassMetadataSpecifiers.h"
 #include "FunctionSpecifiers.h"
 #include "InterfaceSpecifiers.h"
 #include "StructSpecifiers.h"
@@ -47,6 +48,13 @@ const TCHAR* GCheckedMetadataSpecifierStrings[(int32)ECheckedMetadataSpecifier::
 	#undef CHECKED_METADATA_SPECIFIER
 };
 
+const TCHAR* GClassMetadataSpecifierStrings[(int32)EClassMetadataSpecifier::Max] =
+{
+	#define CLASS_METADATA_SPECIFIER(SpecifierName) TEXT(#SpecifierName),
+		#include "ClassMetadataSpecifiers.def"
+	#undef CLASS_METADATA_SPECIFIER
+};
+
 struct FCStringsLessThanCaseInsensitive
 {
 	FORCEINLINE bool operator()(const TCHAR* Lhs, const TCHAR* Rhs)
@@ -60,3 +68,4 @@ const bool GIsGStructSpecifierStringsSorted          = ensure(Algo::IsSorted(GSt
 const bool GIsGInterfaceSpecifierStringsSorted       = ensure(Algo::IsSorted(GInterfaceSpecifierStrings,       FCStringsLessThanCaseInsensitive()));
 const bool GIsGVariableSpecifierStringsSorted        = ensure(Algo::IsSorted(GVariableSpecifierStrings,        FCStringsLessThanCaseInsensitive()));
 const bool GIsGCheckedMetadataSpecifierStringsSorted = ensure(Algo::IsSorted(GCheckedMetadataSpecifierStrings, FCStringsLessThanCaseInsensitive()));
+const bool GIsGClassMetadataSpecifierStringsSorted   = ensure(Algo::IsSorted(GClassMetadataSpecifierStrings,   FCStringsLessThanCaseInsensitive()));
