UHT Optimizations Part 3 - Avoid potentially expensive scope search for an FName when the supplied name does not exist as an FName
#rnx

diff --git a/Engine/Source/Programs/UnrealHeaderTool/Private/Scope.cpp b/Engine/Source/Programs/UnrealHeaderTool/Private/Scope.cpp
--- a/Engine/Source/Programs/UnrealHeaderTool/Private/Scope.cpp
+++ b/Engine/Source/Programs/UnrealHeaderTool/Private/Scope.cpp
@@ -66,9 +66,9 @@ void DispatchType(UField* Type, TArray<UEnum*> &Enums, TArray<UScriptStruct*> &S
 
 void FScope::SplitTypesIntoArrays(TArray<UEnum*>& Enums, TArray<UScriptStruct*>& Structs, TArray<UDelegateFunction*>& DelegateFunctions)
 {
-	for (auto& TypePair : TypeMap)
+	for (TPair<FName, UField*>& TypePair : TypeMap)
 	{
-		auto* Type = TypePair.Value;
+		UField* Type = TypePair.Value;
 		DispatchType(Type, Enums, Structs, DelegateFunctions);
 	}
 }
@@ -96,14 +96,17 @@ TSharedRef<FScope> FScope::AddTypeScope(UStruct* Type, FScope* ParentScope)
 
 UField* FScope::FindTypeByName(FName Name)
 {
-	auto TypeIterator = TDeepScopeTypeIterator<UField, false>(this);
-
-	while (TypeIterator.MoveNext())
+	if (!Name.IsNone())
 	{
-		auto* Type = *TypeIterator;
-		if (Type->GetFName() == Name)
+		TDeepScopeTypeIterator<UField, false> TypeIterator(this);
+
+		while (TypeIterator.MoveNext())
 		{
-			return Type;
+			UField* Type = *TypeIterator;
+			if (Type->GetFName() == Name)
+			{
+				return Type;
+			}
 		}
 	}
 
@@ -112,14 +115,17 @@ UField* FScope::FindTypeByName(FName Name)
 
 const UField* FScope::FindTypeByName(FName Name) const
 {
-	auto TypeIterator = GetTypeIterator();
-
-	while (TypeIterator.MoveNext())
+	if (!Name.IsNone())
 	{
-		auto* Type = *TypeIterator;
-		if (Type->GetFName() == Name)
+		TScopeTypeIterator<UField, true> TypeIterator = GetTypeIterator();
+
+		while (TypeIterator.MoveNext())
 		{
-			return Type;
+			UField* Type = *TypeIterator;
+			if (Type->GetFName() == Name)
+			{
+				return Type;
+			}
 		}
 	}
 
diff --git a/Engine/Source/Programs/UnrealHeaderTool/Private/Scope.h b/Engine/Source/Programs/UnrealHeaderTool/Private/Scope.h
--- a/Engine/Source/Programs/UnrealHeaderTool/Private/Scope.h
+++ b/Engine/Source/Programs/UnrealHeaderTool/Private/Scope.h
@@ -64,6 +64,19 @@ public:
 	 */
 	UField* FindTypeByName(FName Name);
 
+	/**
+	 * Finds type by its name.
+	 *
+	 * @param Name Name to look for.
+	 *
+	 * @returns Found type or nullptr on failure.
+	 */
+	template <typename CharType>
+	UField* FindTypeByName(const CharType* Name)
+	{
+		return FindTypeByName(FName(Name, FNAME_Find));
+	}
+
 	/**
 	 * Finds type by its name.
 	 *
@@ -75,6 +88,20 @@ public:
 	 */
 	const UField* FindTypeByName(FName Name) const;
 
+	/**
+	 * Finds type by its name.
+	 *
+	 * Const version.
+	 *
+	 * @param Name Name to look for.
+	 *
+	 * @returns Found type or nullptr on failure.
+	 */
+	template <typename CharType>
+	const UField* FindTypeByName(const CharType* Name) const
+	{
+		return FindTypeByName(FName(Name, FNAME_Find));
+	}
 	/**
 	 * Checks if scope contains this type.
 	 *
