UHT optimizations part 4 - Improve HeaderParser times by ~35%
* FToken now lazily initializes FName version when needed
* FToken now avoids string compares when matching single character symbols and requires that case sensitivity be specified for matching identifiers
* FClass now maintains a type package name map rather than reconstructing package name repeatedly
* HeaderParser now avoids name lookups, avoids case insensitive string compares, and uses move semantics whenever possible
#rnx

diff --git a/Engine/Source/Programs/UnrealHeaderTool/Private/BaseParser.cpp b/Engine/Source/Programs/UnrealHeaderTool/Private/BaseParser.cpp
--- a/Engine/Source/Programs/UnrealHeaderTool/Private/BaseParser.cpp
+++ b/Engine/Source/Programs/UnrealHeaderTool/Private/BaseParser.cpp
@@ -38,7 +38,7 @@ namespace
 		{
 		}
 
-		FString                    Mapping;
+		FName                      Mapping;
 		FString                    DefaultValue;
 		EMetadataValueAction::Type ValueAction;
 	};
@@ -60,17 +60,17 @@ namespace
 			ValueActions.Add(FMetadataValueAction(InMapping, TEXT(""), EMetadataValueAction::Remove));
 		}
 
-		void ApplyToMetadata(TMap<FName, FString>& Metadata, const FString* Value = NULL)
+		void ApplyToMetadata(TMap<FName, FString>& Metadata, const FString* Value = nullptr)
 		{
-			for (auto It = ValueActions.CreateConstIterator(); It; ++It)
+			for (const FMetadataValueAction& ValueAction : ValueActions)
 			{
-				if (It->ValueAction == EMetadataValueAction::Add)
+				if (ValueAction.ValueAction == EMetadataValueAction::Add)
 				{
-					FBaseParser::InsertMetaDataPair(Metadata, It->Mapping, Value ? *Value : It->DefaultValue);
+					FBaseParser::InsertMetaDataPair(Metadata, ValueAction.Mapping, Value ? *Value : ValueAction.DefaultValue);
 				}
 				else
 				{
-					Metadata.Remove(*It->Mapping);
+					Metadata.Remove(ValueAction.Mapping);
 				}
 			}
 		}
@@ -370,14 +370,14 @@ bool FBaseParser::IsWhitespace( TCHAR c )
 // Gets the next token from the input stream, advancing the variables which keep track of the current input position and line.
 bool FBaseParser::GetToken( FToken& Token, bool bNoConsts/*=false*/, ESymbolParseOption bParseTemplateClosingBracket/*=ESymbolParseOption::Normal*/ )
 {
-	Token.TokenName	= NAME_None;
+	Token.ClearTokenName();
 	TCHAR c = GetLeadingChar();
-	TCHAR p = PeekChar();
 	if( c == 0 )
 	{
 		UngetChar();
 		return 0;
 	}
+	TCHAR p = PeekChar();
 	Token.StartPos		= PrevPos;
 	Token.StartLine		= PrevLine;
 	if( (c>='A' && c<='Z') || (c>='a' && c<='z') || (c=='_') )
@@ -401,20 +401,17 @@ bool FBaseParser::GetToken( FToken& Token, bool bNoConsts/*=false*/, ESymbolPars
 		// Assume this is an identifier unless we find otherwise.
 		Token.TokenType = TOKEN_Identifier;
 
-		// Lookup the token's global name.
-		Token.TokenName = FName(Token.Identifier, FNAME_Find);
-
 		// If const values are allowed, determine whether the identifier represents a constant
 		if ( !bNoConsts )
 		{
 			// See if the identifier is part of a vector, rotation or other struct constant.
 			// boolean true/false
-			if( Token.Matches(TEXT("true")) )
+			if( Token.Matches(TEXT("true"), ESearchCase::IgnoreCase) )
 			{
 				Token.SetConstBool(true);
 				return true;
 			}
-			else if( Token.Matches(TEXT("false")) )
+			else if( Token.Matches(TEXT("false"), ESearchCase::IgnoreCase) )
 			{
 				Token.SetConstBool(false);
 				return true;
@@ -578,21 +575,25 @@ bool FBaseParser::GetToken( FToken& Token, bool bNoConsts/*=false*/, ESymbolPars
 			Token.Identifier[Length++] = d;
 			if( c=='>' && d=='>' )
 			{
-				if( GetChar()=='>' )
+				if (GetChar()=='>')
+				{
 					Token.Identifier[Length++] = '>';
+				}
 				else
+				{
 					UngetChar();
+				}
 			}
 		}
-		else UngetChar();
+		else
+		{
+			UngetChar();
+		}
 		#undef PAIR
 
 		Token.Identifier[Length] = 0;
 		Token.TokenType = TOKEN_Symbol;
 
-		// Lookup the token's global name.
-		Token.TokenName = FName(Token.Identifier, FNAME_Find);
-
 		return true;
 	}
 }
@@ -773,13 +774,32 @@ bool FBaseParser::GetConstInt64(int64& Result, const TCHAR* Tag)
 	return false;
 }
 
-bool FBaseParser::MatchSymbol( const TCHAR* Match, ESymbolParseOption bParseTemplateClosingBracket/*=ESymbolParseOption::Normal*/ )
+bool FBaseParser::MatchSymbol( const TCHAR Match, ESymbolParseOption bParseTemplateClosingBracket/*=ESymbolParseOption::Normal*/ )
 {
 	FToken Token;
 
 	if (GetToken(Token, /*bNoConsts=*/ true, bParseTemplateClosingBracket))
 	{
-		if (Token.TokenType==TOKEN_Symbol && !FCString::Stricmp(Token.Identifier, Match))
+		if (Token.TokenType==TOKEN_Symbol && Token.Identifier[0] == Match && Token.Identifier[1] == 0)
+		{
+			return true;
+		}
+		else
+		{
+			UngetToken(Token);
+		}
+	}
+
+	return false;
+}
+
+bool FBaseParser::MatchSymbol(const TCHAR* Match, ESymbolParseOption bParseTemplateClosingBracket/*=ESymbolParseOption::Normal*/)
+{
+	FToken Token;
+
+	if (GetToken(Token, /*bNoConsts=*/ true, bParseTemplateClosingBracket))
+	{
+		if (Token.TokenType==TOKEN_Symbol && !FCString::Strcmp(Token.Identifier, Match))
 		{
 			return true;
 		}
@@ -796,7 +816,7 @@ bool FBaseParser::MatchSymbol( const TCHAR* Match, ESymbolParseOption bParseTemp
 // Get a specific identifier and return 1 if gotten, 0 if not.
 // This is used primarily for checking for required symbols during compilation.
 //
-bool FBaseParser::MatchIdentifier( FName Match )
+bool FBaseParser::MatchIdentifierByName( FName Match )
 {
 	FToken Token;
 	if (!GetToken(Token))
@@ -804,7 +824,7 @@ bool FBaseParser::MatchIdentifier( FName Match )
 		return false;
 	}
 
-	if ((Token.TokenType == TOKEN_Identifier) && (Token.TokenName == Match))
+	if ((Token.TokenType == TOKEN_Identifier) && (Token.GetTokenName() == Match))
 	{
 		return true;
 	}
@@ -813,12 +833,12 @@ bool FBaseParser::MatchIdentifier( FName Match )
 	return false;
 }
 
-bool FBaseParser::MatchIdentifier( const TCHAR* Match )
+bool FBaseParser::MatchIdentifier( const TCHAR* Match, ESearchCase::Type SearchCase)
 {
 	FToken Token;
 	if (GetToken(Token))
 	{
-		if( Token.TokenType==TOKEN_Identifier && FCString::Stricmp(Token.Identifier,Match)==0 )
+		if (Token.TokenType==TOKEN_Identifier && ((SearchCase == ESearchCase::CaseSensitive) ? !FCString::Strcmp(Token.Identifier, Match) : !FCString::Stricmp(Token.Identifier, Match)))
 		{
 			return true;
 		}
@@ -836,7 +856,7 @@ bool FBaseParser::MatchConstInt( const TCHAR* Match )
 	FToken Token;
 	if (GetToken(Token))
 	{
-		if( Token.TokenType==TOKEN_Const && (Token.Type == CPT_Int || Token.Type == CPT_Int64) && FCString::Stricmp(Token.Identifier,Match)==0 )
+		if( Token.TokenType==TOKEN_Const && (Token.Type == CPT_Int || Token.Type == CPT_Int64) && FCString::Strcmp(Token.Identifier,Match)==0 )
 		{
 			return true;
 		}
@@ -869,7 +889,7 @@ bool FBaseParser::MatchAnyConstInt()
 
 void FBaseParser::MatchSemi()
 {
-	if( !MatchSymbol(TEXT(";")) )
+	if( !MatchSymbol(TEXT(';')) )
 	{
 		FToken Token;
 		if( GetToken(Token) )
@@ -887,7 +907,7 @@ void FBaseParser::MatchSemi()
 //
 // Peek ahead and see if a symbol follows in the stream.
 //
-bool FBaseParser::PeekSymbol( const TCHAR* Match )
+bool FBaseParser::PeekSymbol( const TCHAR Match )
 {
 	FToken Token;
 	if (!GetToken(Token, true))
@@ -896,13 +916,13 @@ bool FBaseParser::PeekSymbol( const TCHAR* Match )
 	}
 	UngetToken(Token);
 
-	return Token.TokenType==TOKEN_Symbol && FCString::Stricmp(Token.Identifier,Match)==0;
+	return Token.TokenType==TOKEN_Symbol && Token.Identifier[0] == Match && Token.Identifier[1] == 0;
 }
 
 //
 // Peek ahead and see if an identifier follows in the stream.
 //
-bool FBaseParser::PeekIdentifier( FName Match )
+bool FBaseParser::PeekIdentifierByName( FName Match )
 {
 	FToken Token;
 	if (!GetToken(Token, true))
@@ -910,10 +930,10 @@ bool FBaseParser::PeekIdentifier( FName Match )
 		return false;
 	}
 	UngetToken(Token);
-	return Token.TokenType==TOKEN_Identifier && Token.TokenName==Match;
+	return Token.TokenType==TOKEN_Identifier && Token.GetTokenName()==Match;
 }
 
-bool FBaseParser::PeekIdentifier( const TCHAR* Match )
+bool FBaseParser::PeekIdentifier( const TCHAR* Match, ESearchCase::Type SearchCase)
 {
 	FToken Token;
 	if (!GetToken(Token, true))
@@ -921,7 +941,7 @@ bool FBaseParser::PeekIdentifier( const TCHAR* Match )
 		return false;
 	}
 	UngetToken(Token);
-	return Token.TokenType==TOKEN_Identifier && FCString::Stricmp(Token.Identifier,Match)==0;
+	return Token.TokenType==TOKEN_Identifier && ((SearchCase == ESearchCase::CaseSensitive) ? !FCString::Strcmp(Token.Identifier, Match) : !FCString::Stricmp(Token.Identifier, Match));
 }
 
 //
@@ -936,7 +956,7 @@ void FBaseParser::UngetToken( const FToken& Token )
 //
 // Require a symbol.
 //
-void FBaseParser::RequireSymbol( const TCHAR* Match, const TCHAR* Tag, ESymbolParseOption bParseTemplateClosingBracket/*=ESymbolParseOption::Normal*/ )
+void FBaseParser::RequireSymbol( const TCHAR Match, const TCHAR* Tag, ESymbolParseOption bParseTemplateClosingBracket/*=ESymbolParseOption::Normal*/ )
 {
 	if (!MatchSymbol(Match, bParseTemplateClosingBracket))
 	{
@@ -944,6 +964,14 @@ void FBaseParser::RequireSymbol( const TCHAR* Match, const TCHAR* Tag, ESymbolPa
 	}
 }
 
+void FBaseParser::RequireSymbol(const TCHAR Match, TFunctionRef<FString()> ErrorGetter, ESymbolParseOption bParseTemplateClosingBracket/*=ESymbolParseOption::Normal*/)
+{
+	if (!MatchSymbol(Match, bParseTemplateClosingBracket))
+	{
+		FError::Throwf(TEXT("Missing '%s' in %s"), Match, *ErrorGetter());
+	}
+}
+
 //
 // Require an integer.
 //
@@ -966,17 +994,9 @@ void FBaseParser::RequireAnyConstInt( const TCHAR* Tag )
 //
 // Require an identifier.
 //
-void FBaseParser::RequireIdentifier( FName Match, const TCHAR* Tag )
-{
-	if (!MatchIdentifier(Match))
-	{
-		FError::Throwf(TEXT("Missing '%s' in %s"), *Match.ToString(), Tag );
-	}
-}
-
-void FBaseParser::RequireIdentifier( const TCHAR* Match, const TCHAR* Tag )
+void FBaseParser::RequireIdentifier( const TCHAR* Match, const ESearchCase::Type SearchCase, const TCHAR* Tag )
 {
-	if (!MatchIdentifier(Match))
+	if (!MatchIdentifier(Match, SearchCase))
 	{
 		FError::Throwf(TEXT("Missing '%s' in %s"), Match, Tag );
 	}
@@ -985,54 +1005,56 @@ void FBaseParser::RequireIdentifier( const TCHAR* Match, const TCHAR* Tag )
 // Clears out the stored comment.
 void FBaseParser::ClearComment()
 {
-	// Can't call Reset as FString uses protected inheritance
-	PrevComment.Empty( PrevComment.Len() );
+	PrevComment.Reset();
 }
 
 // Reads a new-style value
 //@TODO: UCREMOVAL: Needs a better name
-FString FBaseParser::ReadNewStyleValue(const FString& TypeOfSpecifier)
+FString FBaseParser::ReadNewStyleValue(const TCHAR* TypeOfSpecifier)
 {
 	FToken ValueToken;
-	if (!GetToken( ValueToken, false ))
-		FError::Throwf(TEXT("Expected a value when handling a "), *TypeOfSpecifier);
+	if (!GetToken(ValueToken, false))
+	{
+		FError::Throwf(TEXT("Expected a value when handling a %s"), TypeOfSpecifier);
+	}
+
+	FString Result;
 
 	switch (ValueToken.TokenType)
 	{
 		case TOKEN_Identifier:
 		case TOKEN_Symbol:
-		{
-			FString Result = ValueToken.Identifier;
 
-			if (MatchSymbol(TEXT("=")))
+			Result = ValueToken.Identifier;
+
+			if (MatchSymbol(TEXT('=')))
 			{
 				Result += TEXT("=");
 				Result += ReadNewStyleValue(TypeOfSpecifier);
 			}
-
-			return Result;
-		}
+			break;
 
 		case TOKEN_Const:
-			return ValueToken.GetConstantValue();
 
-		default:
-			return TEXT("");
+			Result = ValueToken.GetConstantValue();
+			break;
 	}
+
+	return Result;
 }
 
 // Reads [ Value | ['(' Value [',' Value]* ')'] ] and places each value into the Items array
-bool FBaseParser::ReadOptionalCommaSeparatedListInParens(TArray<FString>& Items, const FString& TypeOfSpecifier)
+bool FBaseParser::ReadOptionalCommaSeparatedListInParens(TArray<FString>& Items, const TCHAR* TypeOfSpecifier)
 {
-	if (MatchSymbol(TEXT("(")))
+	if (MatchSymbol(TEXT('(')))
 	{
 		do 
 		{
 			FString Value = ReadNewStyleValue(TypeOfSpecifier);
 			Items.Add(Value);
-		} while ( MatchSymbol(TEXT(",")) );
+		} while ( MatchSymbol(TEXT(',')) );
 
-		RequireSymbol(TEXT(")"), *TypeOfSpecifier);
+		RequireSymbol(TEXT(')'), TypeOfSpecifier);
 
 		return true;
 	}
@@ -1055,7 +1077,7 @@ void FBaseParser::ParseNameWithPotentialAPIMacroPrefix(FString& DeclaredName, FS
 	FString NameTokenStr = NameToken.Identifier;
 	if (NameTokenStr.EndsWith(TEXT("_API"), ESearchCase::CaseSensitive))
 	{
-		RequiredAPIMacroIfPresent = NameTokenStr;
+		RequiredAPIMacroIfPresent = MoveTemp(NameTokenStr);
 
 		// Read the real name
 		if (!GetIdentifier(NameToken))
@@ -1075,28 +1097,29 @@ void FBaseParser::ParseNameWithPotentialAPIMacroPrefix(FString& DeclaredName, FS
 		DeclaredName = NameToken.Identifier;
 
 		// Not required for things such a Minimal. Only used for Linux/Mac which just exposes the vtable for link-age
-		RequiredAPIMacroIfPresent.Empty();
+		RequiredAPIMacroIfPresent.Reset();
 	}
 	else
 	{
-		DeclaredName = NameTokenStr;
-		RequiredAPIMacroIfPresent.Empty();
+		DeclaredName = MoveTemp(NameTokenStr);
+		RequiredAPIMacroIfPresent.Reset();
 	}
 }
 
 // Reads a set of specifiers (with optional values) inside the () of a new-style metadata macro like UPROPERTY or UFUNCTION
-void FBaseParser::ReadSpecifierSetInsideMacro(TArray<FPropertySpecifier>& SpecifiersFound, const FString& TypeOfSpecifier, TMap<FName, FString>& MetaData)
+void FBaseParser::ReadSpecifierSetInsideMacro(TArray<FPropertySpecifier>& SpecifiersFound, const TCHAR* TypeOfSpecifier, TMap<FName, FString>& MetaData)
 {
 	int32 FoundSpecifierCount = 0;
-	FString ErrorMessage = FString::Printf(TEXT("%s declaration specifier"), *TypeOfSpecifier);
 
-	RequireSymbol(TEXT("("), *ErrorMessage);
+	auto ErrorMessageGetter = [TypeOfSpecifier]() { return FString::Printf(TEXT("%s declaration specifier"), TypeOfSpecifier); };
+
+	RequireSymbol(TEXT('('), ErrorMessageGetter);
 
-	while (!MatchSymbol(TEXT(")")))
+	while (!MatchSymbol(TEXT(')')))
 	{
 		if (FoundSpecifierCount > 0)
 		{
-			RequireSymbol(TEXT(","), *ErrorMessage);
+			RequireSymbol(TEXT(','), ErrorMessageGetter);
 		}
 		++FoundSpecifierCount;
 
@@ -1104,13 +1127,13 @@ void FBaseParser::ReadSpecifierSetInsideMacro(TArray<FPropertySpecifier>& Specif
 		FToken Specifier;
 		if (!GetToken(Specifier))
 		{
-			FError::Throwf(TEXT("Expected %s"), *ErrorMessage);
+			FError::Throwf(TEXT("Expected %s"), *ErrorMessageGetter());
 		}
 
-		if (Specifier.Matches(TEXT("meta")))
+		if (Specifier.Matches(TEXT("meta"), ESearchCase::IgnoreCase))
 		{
-			RequireSymbol(TEXT("="), *ErrorMessage);
-			RequireSymbol(TEXT("("), *ErrorMessage);
+			RequireSymbol(TEXT('='), ErrorMessageGetter);
+			RequireSymbol(TEXT('('), ErrorMessageGetter);
 
 			// Keep reading comma-separated metadata pairs
 			do 
@@ -1126,21 +1149,21 @@ void FBaseParser::ReadSpecifierSetInsideMacro(TArray<FPropertySpecifier>& Specif
 
 				// Potentially read a value
 				FString Value;
-				if (MatchSymbol(TEXT("=")))
+				if (MatchSymbol(TEXT('=')))
 				{
 					Value = ReadNewStyleValue(TypeOfSpecifier);
 				}
 
 				// Validate the value is a valid type for the key and insert it into the map
-				InsertMetaDataPair(MetaData, Key, Value);
-			} while ( MatchSymbol(TEXT(",")) );
+				InsertMetaDataPair(MetaData, MoveTemp(Key), MoveTemp(Value));
+			} while ( MatchSymbol(TEXT(',')) );
 
-			RequireSymbol(TEXT(")"), *ErrorMessage);
+			RequireSymbol(TEXT(')'), ErrorMessageGetter);
 		}
 		// Look up specifier in metadata dictionary
 		else if (FMetadataKeyword* MetadataKeyword = GetMetadataKeyword(Specifier.Identifier))
 		{
-			if (MatchSymbol(TEXT("=")))
+			if (MatchSymbol(TEXT('=')))
 			{
 				if (MetadataKeyword->ValueArgument == EMetadataValueArgument::None)
 				{
@@ -1166,45 +1189,45 @@ void FBaseParser::ReadSpecifierSetInsideMacro(TArray<FPropertySpecifier>& Specif
 			SpecifiersFound.Emplace(Specifier.Identifier);
 
 			// Look for a value for this specifier
-			if (MatchSymbol(TEXT("=")) || PeekSymbol(TEXT("(")))
+			if (MatchSymbol(TEXT('=')) || PeekSymbol(TEXT('(')))
 			{
 				TArray<FString>& NewPairValues = SpecifiersFound.Last().Values;
 				if (!ReadOptionalCommaSeparatedListInParens(NewPairValues, TypeOfSpecifier))
 				{
-					FString Value = ReadNewStyleValue(TypeOfSpecifier);
-					NewPairValues.Add(Value);
+					NewPairValues.Add(ReadNewStyleValue(TypeOfSpecifier));
 				}
 			}
 		}
 	}
 }
 
-void FBaseParser::InsertMetaDataPair(TMap<FName, FString>& MetaData, const FString& InKey, const FString& InValue)
+void FBaseParser::InsertMetaDataPair(TMap<FName, FString>& MetaData, FString Key, FString Value)
 {
-	FString Key = InKey;
-	FString Value = InValue;
-
-	// trim extra white space and quotes
-	Key.TrimStartAndEndInline();
-	Value.TrimStartAndEndInline();
-	Value = Value.TrimQuotes();
-
 	// make sure the key is valid
-	if (InKey.Len() == 0)
+	if (Key.Len() == 0)
 	{
 		FError::Throwf(TEXT("Invalid metadata"));
 	}
 
-	FName KeyName(*Key);
+	// trim extra white space and quotes
+	Key.TrimStartAndEndInline();
+
+	InsertMetaDataPair(MetaData, FName(*Key), MoveTemp(Value));
+}
+
+void FBaseParser::InsertMetaDataPair(TMap<FName, FString>& MetaData, FName KeyName, FString Value)
+{
+	Value.TrimStartAndEndInline();
+	Value.TrimQuotesInline();
 
 	FString* ExistingValue = MetaData.Find(KeyName);
 	if (ExistingValue && Value != *ExistingValue)
 	{
-		FError::Throwf(TEXT("Metadata key '%s' first seen with value '%s' then '%s'"), *Key, **ExistingValue, *Value);
+		FError::Throwf(TEXT("Metadata key '%s' first seen with value '%s' then '%s'"), *KeyName.ToString(), **ExistingValue, *Value);
 	}
 
 	// finally we have enough to put it into our metadata
-	MetaData.Add(FName(*Key), *Value);
+	MetaData.Add(KeyName, MoveTemp(Value));
 }
 
 
diff --git a/Engine/Source/Programs/UnrealHeaderTool/Private/BaseParser.h b/Engine/Source/Programs/UnrealHeaderTool/Private/BaseParser.h
--- a/Engine/Source/Programs/UnrealHeaderTool/Private/BaseParser.h
+++ b/Engine/Source/Programs/UnrealHeaderTool/Private/BaseParser.h
@@ -139,20 +139,21 @@ public:
 	bool GetConstInt64(int64& Result, const TCHAR* Tag = NULL);
 
 	// Matching predefined text.
-	bool MatchIdentifier( FName Match );
-	bool MatchIdentifier( const TCHAR* Match );
+	bool MatchIdentifierByName( FName Match );
+	bool MatchIdentifier( const TCHAR* Match, ESearchCase::Type SearchCase);
 	bool MatchConstInt( const TCHAR* Match );
 	bool MatchAnyConstInt();
-	bool PeekIdentifier( FName Match );
-	bool PeekIdentifier( const TCHAR* Match );
-	bool MatchSymbol( const TCHAR* Match, ESymbolParseOption bParseTemplateClosingBracket = ESymbolParseOption::Normal );
+	bool PeekIdentifierByName( FName Match );
+	bool PeekIdentifier( const TCHAR* Match, ESearchCase::Type SearchCase);
+	bool MatchSymbol( const TCHAR Match, ESymbolParseOption bParseTemplateClosingBracket = ESymbolParseOption::Normal );
+	bool MatchSymbol(const TCHAR* Match, ESymbolParseOption bParseTemplateClosingBracket = ESymbolParseOption::Normal);
 	void MatchSemi();
-	bool PeekSymbol( const TCHAR* Match );
+	bool PeekSymbol( const TCHAR Match );
 
 	// Requiring predefined text.
-	void RequireIdentifier( FName Match, const TCHAR* Tag );
-	void RequireIdentifier( const TCHAR* Match, const TCHAR* Tag );
-	void RequireSymbol( const TCHAR* Match, const TCHAR* Tag, ESymbolParseOption bParseTemplateClosingBracket = ESymbolParseOption::Normal );
+	void RequireIdentifier( const TCHAR* Match, ESearchCase::Type SearchCase, const TCHAR* Tag );
+	void RequireSymbol( const TCHAR Match, const TCHAR* Tag, ESymbolParseOption bParseTemplateClosingBracket = ESymbolParseOption::Normal );
+	void RequireSymbol(const TCHAR Match, TFunctionRef<FString()> TagGetter, ESymbolParseOption bParseTemplateClosingBracket = ESymbolParseOption::Normal);
 	void RequireConstInt( const TCHAR* Match, const TCHAR* Tag );
 	void RequireAnyConstInt( const TCHAR* Tag );
 
@@ -161,10 +162,10 @@ public:
 
 	// Reads a new-style value
 	//@TODO: UCREMOVAL: Needs a better name
-	FString ReadNewStyleValue(const FString& TypeOfSpecifier);
+	FString ReadNewStyleValue(const TCHAR* TypeOfSpecifier);
 
 	// Reads ['(' Value [',' Value]* ')'] and places each value into the Items array
-	bool ReadOptionalCommaSeparatedListInParens(TArray<FString>& Items, const FString& TypeOfSpecifier);
+	bool ReadOptionalCommaSeparatedListInParens(TArray<FString>& Items, const TCHAR* TypeOfSpecifier);
 
 	//////////////
 	// Complicated* parsing code that needs to be shared between the preparser and the parser
@@ -177,10 +178,13 @@ public:
 	void ParseNameWithPotentialAPIMacroPrefix(FString& DeclaredName, FString& RequiredAPIMacroIfPresent, const TCHAR* FailureMessage);
 
 	// Reads a set of specifiers (with optional values) inside the () of a new-style metadata macro like UPROPERTY or UFUNCTION
-	void ReadSpecifierSetInsideMacro(TArray<FPropertySpecifier>& SpecifiersFound, const FString& TypeOfSpecifier, TMap<FName, FString>& MetaData);
+	void ReadSpecifierSetInsideMacro(TArray<FPropertySpecifier>& SpecifiersFound, const TCHAR* TypeOfSpecifier, TMap<FName, FString>& MetaData);
 
 	// Validates and inserts one key-value pair into the meta data map
-	static void InsertMetaDataPair(TMap<FName, FString>& MetaData, const FString& InKey, const FString& InValue);
+	static void InsertMetaDataPair(TMap<FName, FString>& MetaData, FString InKey, FString InValue);
+
+	// Validates and inserts one key-value pair into the meta data map
+	static void InsertMetaDataPair(TMap<FName, FString>& MetaData, FName InKey, FString InValue);
 
 	//////////////
 };
diff --git a/Engine/Source/Programs/UnrealHeaderTool/Private/HeaderParser.cpp b/Engine/Source/Programs/UnrealHeaderTool/Private/HeaderParser.cpp
--- a/Engine/Source/Programs/UnrealHeaderTool/Private/HeaderParser.cpp
+++ b/Engine/Source/Programs/UnrealHeaderTool/Private/HeaderParser.cpp
@@ -69,6 +69,34 @@ enum {MAX_ARRAY_SIZE=2048};
 static const FName NAME_Comment(TEXT("Comment"));
 static const FName NAME_ToolTip(TEXT("ToolTip"));
 static const FName NAME_DocumentationPolicy(TEXT("DocumentationPolicy"));
+static const FName NAME_AllowPrivateAccess(TEXT("AllowPrivateAccess"));
+static const FName NAME_ExposeOnSpawn(TEXT("ExposeOnSpawn"));
+static const FName NAME_NativeConst(TEXT("NativeConst"));
+static const FName NAME_NativeConstTemplateArg(TEXT("NativeConstTemplateArg"));
+static const FName NAME_BlueprintInternalUseOnly(TEXT("BlueprintInternalUseOnly"));
+static const FName NAME_DeprecatedFunction(TEXT("DeprecatedFunction"));
+static const FName NAME_BlueprintSetter(TEXT("BlueprintSetter"));
+static const FName NAME_BlueprintGetter(TEXT("BlueprintGetter"));
+static const FName NAME_Category(TEXT("Category"));
+static const FName NAME_ReturnValue(TEXT("ReturnValue"));
+static const FName NAME_CppFromBpEvent(TEXT("CppFromBpEvent"));
+static const FName NAME_CustomThunk(TEXT("CustomThunk"));
+static const FName NAME_ArraySizeEnum(TEXT("ArraySizeEnum"));
+static const FName NAME_ClassGroupNames(TEXT("ClassGroupNames"));
+static const FName NAME_AutoCollapseCategories(TEXT("AutoCollapseCategories"));
+static const FName NAME_HideCategories(TEXT("HideCategories"));
+static const FName NAME_ShowCategories(TEXT("ShowCategories"));
+static const FName NAME_SparseClassDataTypes(TEXT("SparseClassDataTypes"));
+static const FName NAME_HideFunctions(TEXT("HideFunctions"));
+static const FName NAME_AutoExpandCategories(TEXT("AutoExpandCategories"));
+static const FName NAME_IsConversionRoot(TEXT("IsConversionRoot"));
+static const FName NAME_EditInline(TEXT("EditInline"));
+static const FName NAME_IncludePath(TEXT("IncludePath"));
+static const FName NAME_ModuleRelativePath(TEXT("ModuleRelativePath"));
+static const FName NAME_CannotImplementInterfaceInBlueprint(TEXT("CannotImplementInterfaceInBlueprint"));
+static const FName NAME_UIMin(TEXT("UIMin"));
+static const FName NAME_UIMax(TEXT("UIMax"));
+
 EGeneratedCodeVersion FHeaderParser::DefaultGeneratedCodeVersion = EGeneratedCodeVersion::V1;
 TArray<FString> FHeaderParser::StructsWithNoPrefix;
 TArray<FString> FHeaderParser::StructsWithTPrefix;
@@ -204,7 +232,7 @@ namespace
 		HeaderParser.UngetToken(Token);
 		HeaderParser.GetToken(Token);
 
-		bool bResult = PossibleBracketToken.TokenType != TOKEN_Symbol || FCString::Strcmp(PossibleBracketToken.Identifier, TEXT("("));
+		bool bResult = PossibleBracketToken.TokenType != TOKEN_Symbol || !PossibleBracketToken.Matches(TEXT('('));
 		return bResult;
 	}
 
@@ -475,7 +503,7 @@ namespace
 					bSawPropertyAccessor = true;
 					FuncInfo.FunctionFlags |= FUNC_BlueprintCallable;
 					FuncInfo.FunctionFlags |= FUNC_BlueprintPure;
-					MetaData.Add(TEXT("BlueprintGetter"));
+					MetaData.Add(NAME_BlueprintGetter);
 				}
 				break;
 
@@ -488,7 +516,7 @@ namespace
 
 					bSawPropertyAccessor = true;
 					FuncInfo.FunctionFlags |= FUNC_BlueprintCallable;
-					MetaData.Add(TEXT("BlueprintSetter"));
+					MetaData.Add(NAME_BlueprintSetter);
 				}
 				break;
 
@@ -592,7 +620,7 @@ namespace
 
 	void AddEditInlineMetaData(TMap<FName, FString>& MetaData)
 	{
-		MetaData.Add(TEXT("EditInline"), TEXT("true"));
+		MetaData.Add(NAME_EditInline, TEXT("true"));
 	}
 
 	const TCHAR* GetHintText(EVariableCategory::Type VariableCategory)
@@ -1208,11 +1236,11 @@ namespace
 
 	void SkipAlignasIfNecessary(FBaseParser& Parser)
 	{
-		if (Parser.MatchIdentifier(TEXT("alignas")))
+		if (Parser.MatchIdentifier(TEXT("alignas"), ESearchCase::CaseSensitive))
 		{
-			Parser.RequireSymbol(TEXT("("), TEXT("'alignas'"));
+			Parser.RequireSymbol(TEXT('('), TEXT("'alignas'"));
 			Parser.RequireAnyConstInt(TEXT("'alignas'"));
-			Parser.RequireSymbol(TEXT(")"), TEXT("'alignas'"));
+			Parser.RequireSymbol(TEXT(')'), TEXT("'alignas'"));
 		}
 	}
 
@@ -1224,15 +1252,15 @@ namespace
 			return;
 		}
 
-		if (MacroToken.TokenType != TOKEN_Identifier || (FCString::Stricmp(MacroToken.Identifier, TEXT("DEPRECATED")) != 0 && FCString::Stricmp(MacroToken.Identifier, TEXT("UE_DEPRECATED")) != 0))
+		if (MacroToken.TokenType != TOKEN_Identifier || (FCString::Strcmp(MacroToken.Identifier, TEXT("DEPRECATED")) != 0 && FCString::Strcmp(MacroToken.Identifier, TEXT("UE_DEPRECATED")) != 0))
 		{
 			Parser.UngetToken(MacroToken);
 			return;
 		}
 
-		FString ErrorScope = FString::Printf(TEXT("%s macro"), MacroToken.Identifier);
+		auto ErrorMessageGetter = [&MacroToken]() { return FString::Printf(TEXT("%s macro"), MacroToken.Identifier); };
 
-		Parser.RequireSymbol(TEXT("("), *ErrorScope);
+		Parser.RequireSymbol(TEXT('('), ErrorMessageGetter);
 
 		FToken Token;
 		if (Parser.GetToken(Token) && (Token.Type != CPT_Float || Token.TokenType != TOKEN_Const))
@@ -1240,13 +1268,13 @@ namespace
 			FError::Throwf(TEXT("Expected engine version in %s macro"), MacroToken.Identifier);
 		}
 
-		Parser.RequireSymbol(TEXT(","), *ErrorScope);
+		Parser.RequireSymbol(TEXT(','), ErrorMessageGetter);
 		if (Parser.GetToken(Token) && (Token.Type != CPT_String || Token.TokenType != TOKEN_Const))
 		{
 			FError::Throwf(TEXT("Expected deprecation message in %s macro"), MacroToken.Identifier);
 		}
 
-		Parser.RequireSymbol(TEXT(")"), *ErrorScope);
+		Parser.RequireSymbol(TEXT(')'), ErrorMessageGetter);
 	}
 }
 
@@ -1476,7 +1504,7 @@ void AddIncludePathToMetadata(UField* Type, TMap<FName, FString> &MetaData)
 	TSharedRef<FUnrealTypeDefinitionInfo>* TypeDefinitionPtr = GTypeDefinitionInfoMap.Find(Type);
 	if (TypeDefinitionPtr != nullptr)
 	{
-		MetaData.Add(TEXT("IncludePath"), *(*TypeDefinitionPtr)->GetUnrealSourceFile().GetIncludePath());
+		MetaData.Add(NAME_IncludePath, (*TypeDefinitionPtr)->GetUnrealSourceFile().GetIncludePath());
 	}
 }
 
@@ -1488,7 +1516,7 @@ void AddIncludePathToMetadata(UField* Type, TMap<FName, FString> &MetaData)
  */
 void AddModuleRelativePathToMetadata(FUnrealSourceFile& SourceFile, TMap<FName, FString> &MetaData)
 {
-	MetaData.Add(TEXT("ModuleRelativePath"), *SourceFile.GetModuleRelativePath());
+	MetaData.Add(NAME_ModuleRelativePath, SourceFile.GetModuleRelativePath());
 }
 
 /**
@@ -1503,7 +1531,7 @@ void AddModuleRelativePathToMetadata(UField* Type, TMap<FName, FString> &MetaDat
 	TSharedRef<FUnrealTypeDefinitionInfo>* TypeDefinitionPtr = GTypeDefinitionInfoMap.Find(Type);
 	if (TypeDefinitionPtr != nullptr)
 	{
-		MetaData.Add(TEXT("ModuleRelativePath"), *(*TypeDefinitionPtr)->GetUnrealSourceFile().GetModuleRelativePath());
+		MetaData.Add(NAME_ModuleRelativePath, (*TypeDefinitionPtr)->GetUnrealSourceFile().GetModuleRelativePath());
 	}
 }
 
@@ -1586,7 +1614,7 @@ UEnum* FHeaderParser::CompileEnum()
 	UField* Existing = Scope->FindTypeByName(EnumToken.Identifier);
 	if (Existing)
 	{
-		FError::Throwf(TEXT("enum: '%s' already defined here"), *EnumToken.TokenName.ToString());
+		FError::Throwf(TEXT("enum: '%s' already defined here"), *EnumToken.GetTokenName().ToString());
 	}
 
 	ParseFieldMetaData(EnumToken.MetaData, EnumToken.Identifier);
@@ -1608,7 +1636,7 @@ UEnum* FHeaderParser::CompileEnum()
 	EUnderlyingEnumType UnderlyingType = EUnderlyingEnumType::uint8;
 	if (CppForm == UEnum::ECppForm::EnumClass)
 	{
-		if (MatchSymbol(TEXT(":")))
+		if (MatchSymbol(TEXT(':')))
 		{
 			FToken BaseToken;
 			if (!GetIdentifier(BaseToken))
@@ -1668,14 +1696,14 @@ UEnum* FHeaderParser::CompileEnum()
 	}
 
 	// Get opening brace.
-	RequireSymbol( TEXT("{"), TEXT("'Enum'") );
+	RequireSymbol( TEXT('{'), TEXT("'Enum'") );
 
 	switch (CppForm)
 	{
 		case UEnum::ECppForm::Namespaced:
 		{
 			// Now handle the inner true enum portion
-			RequireIdentifier(TEXT("enum"), TEXT("'Enum'"));
+			RequireIdentifier(TEXT("enum"), ESearchCase::CaseSensitive, TEXT("'Enum'"));
 
 			SkipAlignasIfNecessary(*this);
 
@@ -1687,7 +1715,7 @@ UEnum* FHeaderParser::CompileEnum()
 
 			Enum->CppType = FString::Printf(TEXT("%s::%s"), EnumToken.Identifier, InnerEnumToken.Identifier);
 
-			RequireSymbol( TEXT("{"), TEXT("'Enum'") );
+			RequireSymbol( TEXT('{'), TEXT("'Enum'") );
 		}
 		break;
 
@@ -1716,7 +1744,7 @@ UEnum* FHeaderParser::CompileEnum()
 		AddFormattedPrevCommentAsTooltipMetaData(TagToken.MetaData);
 
 		// Try to read an optional explicit enum value specification
-		if (MatchSymbol(TEXT("=")))
+		if (MatchSymbol(TEXT('=')))
 		{
 			FToken InitToken;
 			if (!GetToken(InitToken))
@@ -1741,7 +1769,7 @@ UEnum* FHeaderParser::CompileEnum()
 
 				if (InitToken.TokenType == TOKEN_Symbol)
 				{
-					if (FCString::Stricmp(InitToken.Identifier, TEXT(",")) == 0 || FCString::Stricmp(InitToken.Identifier, TEXT("}")) == 0)
+					if (InitToken.Matches(TEXT(',')) || InitToken.Matches(TEXT('}')))
 					{
 						UngetToken(InitToken);
 						break;
@@ -1790,7 +1818,7 @@ UEnum* FHeaderParser::CompileEnum()
 			++CurrentEnumValue;
 		}
 
-		TagToken.MetaData.Add(TEXT("Name"), NewTag.ToString());
+		TagToken.MetaData.Add(NAME_Name, NewTag.ToString());
 		EntryMetaData.Add(TagToken.MetaData);
 
 		// check for metadata on this enum value
@@ -1809,7 +1837,7 @@ UEnum* FHeaderParser::CompileEnum()
 			TagToken.MetaData.Empty();
 		}
 
-		if (!MatchSymbol(TEXT(",")))
+		if (!MatchSymbol(TEXT(',')))
 		{
 			FToken ClosingBrace;
 			if (!GetToken(ClosingBrace))
@@ -1817,7 +1845,7 @@ UEnum* FHeaderParser::CompileEnum()
 				FError::Throwf(TEXT("UENUM: end of file encountered"));
 			}
 
-			if (ClosingBrace.TokenType == TOKEN_Symbol && !FCString::Stricmp(ClosingBrace.Identifier, TEXT("}")))
+			if (ClosingBrace.TokenType == TOKEN_Symbol && ClosingBrace.Matches(TEXT('}')))
 			{
 				UngetToken(ClosingBrace);
 				break;
@@ -1835,13 +1863,13 @@ UEnum* FHeaderParser::CompileEnum()
 	}
 
 	// Trailing brace and semicolon for the enum
-	RequireSymbol( TEXT("}"), TEXT("'Enum'") );
+	RequireSymbol( TEXT('}'), TEXT("'Enum'") );
 	MatchSemi();
 
 	if (CppForm == UEnum::ECppForm::Namespaced)
 	{
 		// Trailing brace for the namespace.
-		RequireSymbol( TEXT("}"), TEXT("'Enum'") );
+		RequireSymbol( TEXT('}'), TEXT("'Enum'") );
 	}
 
 	// Register the list of enum names.
@@ -1882,27 +1910,14 @@ bool IsAllSameChar(const TCHAR* Str, TCHAR Ch)
 	while (TCHAR StrCh = *Str++)
 	{
 		if (StrCh != Ch)
+		{
 			return false;
+		}
 	}
 
 	return true;
 }
 
-/**
- * Checks if a string is made up of all the same character.
- *
- * @param  Str The string to check for all
- * @param  Ch  The character to check for
- *
- * @return True if the string is made up only of Ch characters.
- */
-bool IsLineSeparator(const TCHAR* Str)
-{
-	check(Str);
-
-	return IsAllSameChar(Str, TEXT('-')) || IsAllSameChar(Str, TEXT('=')) || IsAllSameChar(Str, TEXT('*'));
-}
-
 /**
  * @param		Input		An input string, expected to be a script comment.
  * @return					The input string, reformatted in such a way as to be appropriate for use as a tooltip.
@@ -1923,7 +1938,7 @@ FString FHeaderParser::FormatCommentForToolTip(const FString& Input)
 
 	if ( !bFoundAlphaNumericChar )
 	{
-		return FString( TEXT("") );
+		return FString();
 	}
 
 	FString Result(Input);
@@ -1932,7 +1947,7 @@ FString FHeaderParser::FormatCommentForToolTip(const FString& Input)
 	{
 		int32 CommentStart, CommentEnd;
 		// Block comments go first
-		for (CommentStart = Result.Find(TEXT("/*~")); CommentStart != INDEX_NONE; CommentStart = Result.Find(TEXT("/*~")))
+		for (CommentStart = Result.Find(TEXT("/*~"), ESearchCase::CaseSensitive); CommentStart != INDEX_NONE; CommentStart = Result.Find(TEXT("/*~"), ESearchCase::CaseSensitive))
 		{
 			CommentEnd = Result.Find(TEXT("*/"), ESearchCase::CaseSensitive, ESearchDir::FromStart, CommentStart);
 			if (CommentEnd != INDEX_NONE)
@@ -1946,7 +1961,7 @@ FString FHeaderParser::FormatCommentForToolTip(const FString& Input)
 			}
 		}
 		// Leftover line comments go next
-		for (CommentStart = Result.Find(TEXT("//~")); CommentStart != INDEX_NONE; CommentStart = Result.Find(TEXT("//~")))
+		for (CommentStart = Result.Find(TEXT("//~"), ESearchCase::CaseSensitive); CommentStart != INDEX_NONE; CommentStart = Result.Find(TEXT("//~"), ESearchCase::CaseSensitive))
 		{
 			CommentEnd = Result.Find(TEXT("\n"), ESearchCase::CaseSensitive, ESearchDir::FromStart, CommentStart);
 			if (CommentEnd != INDEX_NONE)
@@ -1964,37 +1979,43 @@ FString FHeaderParser::FormatCommentForToolTip(const FString& Input)
 	}
 
 	// Check for known commenting styles.
-	const bool bJavaDocStyle = Result.Contains(TEXT("/**"));
-	const bool bCStyle = Result.Contains(TEXT("/*"));
-	const bool bCPPStyle = Result.StartsWith(TEXT("//"));
+	const bool bJavaDocStyle = Result.Contains(TEXT("/**"), ESearchCase::CaseSensitive);
+	const bool bCStyle = Result.Contains(TEXT("/*"), ESearchCase::CaseSensitive);
+	const bool bCPPStyle = Result.StartsWith(TEXT("//"), ESearchCase::CaseSensitive);
 
 	if ( bJavaDocStyle || bCStyle)
 	{
 		// Remove beginning and end markers.
-		Result = Result.Replace( TEXT("/**"), TEXT("") );
-		Result = Result.Replace( TEXT("/*"), TEXT("") );
-		Result = Result.Replace( TEXT("*/"), TEXT("") );
+		if (bJavaDocStyle)
+		{
+			Result.ReplaceInline(TEXT("/**"), TEXT(""), ESearchCase::CaseSensitive);
+		}
+		if (bCStyle)
+		{
+			Result.ReplaceInline(TEXT("/*"), TEXT(""), ESearchCase::CaseSensitive);
+		}
+		Result.ReplaceInline(TEXT("*/"), TEXT(""), ESearchCase::CaseSensitive);
 	}
 
 	if ( bCPPStyle )
 	{
-		// Remove c++-style comment markers.  Also handle javadoc-style comments by replacing
-		// all triple slashes with double-slashes
-		Result = Result.Replace(TEXT("///"), TEXT("//")).Replace( TEXT("//"), TEXT("") );
+		// Remove c++-style comment markers.  Also handle javadoc-style comments 
+		Result.ReplaceInline(TEXT("///"), TEXT(""), ESearchCase::CaseSensitive);
+		Result.ReplaceInline(TEXT("//"), TEXT(""), ESearchCase::CaseSensitive);
 
 		// Parser strips cpptext and replaces it with "// (cpptext)" -- prevent
 		// this from being treated as a comment on variables declared below the
 		// cpptext section
-		Result = Result.Replace( TEXT("(cpptext)"), TEXT("") );
+		Result.ReplaceInline(TEXT("(cpptext)"), TEXT(""));
 	}
 
 	// Get rid of carriage return or tab characters, which mess up tooltips.
-	Result = Result.Replace( TEXT( "\r" ), TEXT( "" ) );
+	Result.ReplaceInline(TEXT( "\r" ), TEXT( "" ), ESearchCase::CaseSensitive);
 
 	//wx widgets has a hard coded tab size of 8
 	{
 		const int32 SpacesPerTab = 8;
-		Result = Result.ConvertTabsToSpaces (SpacesPerTab);
+		Result.ConvertTabsToSpacesInline(SpacesPerTab);
 	}
 
 	// get rid of uniform leading whitespace and all trailing whitespace, on each line
@@ -2025,19 +2046,37 @@ FString FHeaderParser::FormatCommentForToolTip(const FString& Input)
 					++Pos;
 				}
 
-				Line = Line.RightChop(Pos + 1);
+				Line.RightChopInline(Pos + 1, false);
+			}
+		}
 			}
+
+	auto IsWhitespaceOrLineSeparator = [](const FString& Line)
+	{
+		int32 LineLength = Line.Len();
+		int32 WhitespaceCount = 0;
+		while (WhitespaceCount < LineLength && FChar::IsWhitespace(Line[WhitespaceCount]))
+		{
+			++WhitespaceCount;
 		}
+
+		if (WhitespaceCount == LineLength)
+		{
+			return true;
 	}
 
+		const TCHAR* Str = (*Line) + WhitespaceCount;
+		return IsAllSameChar(Str, TEXT('-')) || IsAllSameChar(Str, TEXT('=')) || IsAllSameChar(Str, TEXT('*'));
+	};
+
 	// Find first meaningful line
 	int32 FirstIndex = 0;
-	for (FString Line : Lines)
+	for (const FString& Line : Lines)
+	{
+		if (!IsWhitespaceOrLineSeparator(Line))
 	{
-		Line.TrimStartInline();
-
-		if (Line.Len() && !IsLineSeparator(*Line))
 			break;
+		}
 
 		++FirstIndex;
 	}
@@ -2045,16 +2084,17 @@ FString FHeaderParser::FormatCommentForToolTip(const FString& Input)
 	int32 LastIndex = Lines.Num();
 	while (LastIndex != FirstIndex)
 	{
-		FString Line = Lines[LastIndex - 1];
-		Line.TrimStartInline();
+		const FString& Line = Lines[LastIndex - 1];
 
-		if (Line.Len() && !IsLineSeparator(*Line))
+		if (!IsWhitespaceOrLineSeparator(Line))
+		{
 			break;
+		}
 
 		--LastIndex;
 	}
 
-	Result.Empty();
+	Result.Reset();
 
 	if (FirstIndex != LastIndex)
 	{
@@ -2072,7 +2112,7 @@ FString FHeaderParser::FormatCommentForToolTip(const FString& Input)
 
 		for (int32 Index = FirstIndex; Index != LastIndex; ++Index)
 		{
-			FString Line = Lines[Index];
+			FString& Line = Lines[Index];
 
 			int32 TemporaryMaxWhitespace = MaxNumWhitespaceToRemove;
 
@@ -2091,7 +2131,7 @@ FString FHeaderParser::FormatCommentForToolTip(const FString& Input)
 
 			if (Pos > 0)
 			{
-				Line = Line.Mid(Pos);
+				Line.RightChopInline(Pos, false);
 			}
 
 			if (Index > 0)
@@ -2110,13 +2150,13 @@ FString FHeaderParser::FormatCommentForToolTip(const FString& Input)
 	// Make sure it doesn't start with a newline
 	if (!Result.IsEmpty() && FChar::IsLinebreak(Result[0]))
 	{
-		Result = Result.Mid(1);
+		Result.RightChopInline(1, false);
 	}
 
 	// Make sure it doesn't end with a dead newline
 	if (!Result.IsEmpty() && FChar::IsLinebreak(Result[Result.Len() - 1]))
 	{
-		Result = Result.Left(Result.Len() - 1);
+		Result.LeftInline(Result.Len() - 1, false);
 	}
 
 	// Done.
@@ -2180,9 +2220,10 @@ TMap<FName, FString> FHeaderParser::GetParameterToolTipsFromFunctionComment(cons
 		Params.Add(ParamPrefix + Input.Mid(ParamStart, ParamEnd - ParamStart - 1));
 	}
 
-	for (FString Param : Params)
+	for (FString& Param : Params)
 	{
-		Param = Param.ConvertTabsToSpaces(4).TrimStartAndEnd();
+		Param.ConvertTabsToSpacesInline(4);
+		Param.TrimStartAndEndInline();
 
 		int32 FirstSpaceIndex = -1;
 		if (!Param.FindChar(' ', FirstSpaceIndex))
@@ -2190,11 +2231,12 @@ TMap<FName, FString> FHeaderParser::GetParameterToolTipsFromFunctionComment(cons
 			continue;
 		}
 
-		FString ParamName = Param.Left(FirstSpaceIndex);
 		FString ParamToolTip = Param.Mid(FirstSpaceIndex + 1);
 		ParamToolTip.TrimStartInline();
 
-		Map.Add(*ParamName, ParamToolTip);
+		Param.LeftInline(FirstSpaceIndex);
+
+		Map.Add(*Param, MoveTemp(ParamToolTip));
 	}
 
 	return Map;
@@ -2251,10 +2293,12 @@ EAccessSpecifier FHeaderParser::ParseAccessProtectionSpecifier(FToken& Token)
 
 	for (EAccessSpecifier Test = EAccessSpecifier(ACCESS_NotAnAccessSpecifier + 1); Test != ACCESS_Num; Test = EAccessSpecifier(Test + 1))
 	{
-		if (Token.Matches(GetAccessSpecifierName(Test)) || (Token.Matches(TEXT("private_subobject")) && Test == ACCESS_Public))
+		if (Token.Matches(GetAccessSpecifierName(Test), ESearchCase::CaseSensitive) || (Test == ACCESS_Public && Token.Matches(TEXT("private_subobject"), ESearchCase::CaseSensitive)))
 		{
+			auto ErrorMessageGetter = [&Token]() { return FString::Printf(TEXT("after %s"), Token.Identifier);  };
+
 			// Consume the colon after the specifier
-			RequireSymbol(TEXT(":"), *FString::Printf(TEXT("after %s"), Token.Identifier));
+			RequireSymbol(TEXT(':'), ErrorMessageGetter);
 			return Test;
 		}
 	}
@@ -2286,7 +2330,7 @@ UScriptStruct* FHeaderParser::CompileStructDeclaration(FClasses& AllClasses)
 	ReadSpecifierSetInsideMacro(SpecifiersFound, TEXT("Struct"), MetaData);
 
 	// Consume the struct keyword
-	RequireIdentifier(TEXT("struct"), TEXT("Struct declaration specifier"));
+	RequireIdentifier(TEXT("struct"), ESearchCase::CaseSensitive, TEXT("Struct declaration specifier"));
 
 	// The struct name as parsed in script and stripped of it's prefix
 	FString StructNameInScript;
@@ -2385,9 +2429,9 @@ UScriptStruct* FHeaderParser::CompileStructDeclaration(FClasses& AllClasses)
 	// Get optional superstruct.
 	bool bExtendsBaseStruct = false;
 	
-	if (MatchSymbol(TEXT(":")))
+	if (MatchSymbol(TEXT(':')))
 	{
-		RequireIdentifier(TEXT("public"), TEXT("struct inheritance"));
+		RequireIdentifier(TEXT("public"), ESearchCase::CaseSensitive, TEXT("struct inheritance"));
 		bExtendsBaseStruct = true;
 	}
 
@@ -2401,7 +2445,7 @@ UScriptStruct* FHeaderParser::CompileStructDeclaration(FClasses& AllClasses)
 
 			TSharedPtr<FScope> StructScope = Scope;
 			FString ParentStructNameInScript = FString(ParentScope.Identifier);
-			if (MatchSymbol(TEXT(".")))
+			if (MatchSymbol(TEXT('.')))
 			{
 				if (GetIdentifier(ParentName))
 				{
@@ -2536,7 +2580,7 @@ UScriptStruct* FHeaderParser::CompileStructDeclaration(FClasses& AllClasses)
 	AddMetaDataToClassData(Struct, MetaData);
 
 	// Get opening brace.
-	RequireSymbol( TEXT("{"), TEXT("'struct'") );
+	RequireSymbol( TEXT('{'), TEXT("'struct'") );
 
 	// Members of structs have a default public access level in c++
 	// Assume that, but restore the parser state once we finish parsing this struct
@@ -2578,7 +2622,7 @@ UScriptStruct* FHeaderParser::CompileStructDeclaration(FClasses& AllClasses)
 		{
 			CompileRigVMMethodDeclaration(AllClasses, Struct);
 		}
-		else if (Token.Matches(TEXT("GENERATED_USTRUCT_BODY")) || Token.Matches(TEXT("GENERATED_BODY")))
+		else if (Token.Matches(TEXT("GENERATED_USTRUCT_BODY"), ESearchCase::CaseSensitive) || Token.Matches(TEXT("GENERATED_BODY"), ESearchCase::CaseSensitive))
 		{
 			// Match 'GENERATED_USTRUCT_BODY' '(' [StructName] ')' or 'GENERATED_BODY' '(' [StructName] ')'
 			if (CurrentAccessSpecifier != ACCESS_Public)
@@ -2592,24 +2636,24 @@ UScriptStruct* FHeaderParser::CompileStructDeclaration(FClasses& AllClasses)
 			}
 
 			Struct->StructMacroDeclaredLineNumber = InputLine;
-			RequireSymbol(TEXT("("), TEXT("'struct'"));
+			RequireSymbol(TEXT('('), TEXT("'struct'"));
 
 			CompileVersionDeclaration(Struct);
 
-			RequireSymbol(TEXT(")"), TEXT("'struct'"));
+			RequireSymbol(TEXT(')'), TEXT("'struct'"));
 
 			// Eat a semicolon if present (not required)
-			SafeMatchSymbol(TEXT(";"));
+			SafeMatchSymbol(TEXT(';'));
 		}
-		else if ( Token.Matches(TEXT("#")) && MatchIdentifier(TEXT("ifdef")) )
+		else if ( Token.Matches(TEXT('#')) && MatchIdentifier(TEXT("ifdef"), ESearchCase::CaseSensitive) )
 		{
 			PushCompilerDirective(ECompilerDirective::Insignificant);
 		}
-		else if ( Token.Matches(TEXT("#")) && MatchIdentifier(TEXT("ifndef")) )
+		else if ( Token.Matches(TEXT('#')) && MatchIdentifier(TEXT("ifndef"), ESearchCase::CaseSensitive) )
 		{
 			PushCompilerDirective(ECompilerDirective::Insignificant);
 		}
-		else if (Token.Matches(TEXT("#")) && MatchIdentifier(TEXT("endif")))
+		else if (Token.Matches(TEXT('#')) && MatchIdentifier(TEXT("endif"), ESearchCase::CaseSensitive))
 		{
 			if (CompilerDirectiveStack.Num() < 1)
 			{
@@ -2618,13 +2662,13 @@ UScriptStruct* FHeaderParser::CompileStructDeclaration(FClasses& AllClasses)
 			CompilerDirectiveStack.Pop();
 			// Do nothing and hope that the if code below worked out OK earlier
 		}
-		else if ( Token.Matches(TEXT("#")) && MatchIdentifier(TEXT("if")) )
+		else if ( Token.Matches(TEXT('#')) && MatchIdentifier(TEXT("if"), ESearchCase::CaseSensitive) )
 		{
 			//@TODO: This parsing should be combined with CompileDirective and probably happen much much higher up!
-			bool bInvertConditional = MatchSymbol(TEXT("!"));
+			bool bInvertConditional = MatchSymbol(TEXT('!'));
 			bool bConsumeAsCppText = false;
 
-			if (MatchIdentifier(TEXT("WITH_EDITORONLY_DATA")) )
+			if (MatchIdentifier(TEXT("WITH_EDITORONLY_DATA"), ESearchCase::CaseSensitive) )
 			{
 				if (bInvertConditional)
 				{
@@ -2633,7 +2677,7 @@ UScriptStruct* FHeaderParser::CompileStructDeclaration(FClasses& AllClasses)
 
 				PushCompilerDirective(ECompilerDirective::WithEditorOnlyData);
 			}
-			else if (MatchIdentifier(TEXT("WITH_EDITOR")) )
+			else if (MatchIdentifier(TEXT("WITH_EDITOR"), ESearchCase::CaseSensitive) )
 			{
 				if (bInvertConditional)
 				{
@@ -2641,7 +2685,7 @@ UScriptStruct* FHeaderParser::CompileStructDeclaration(FClasses& AllClasses)
 				}
 				PushCompilerDirective(ECompilerDirective::WithEditor);
 			}
-			else if (MatchIdentifier(TEXT("CPP")) || MatchConstInt(TEXT("0")) || MatchConstInt(TEXT("1")) || MatchIdentifier(TEXT("WITH_HOT_RELOAD")) || MatchIdentifier(TEXT("WITH_HOT_RELOAD_CTORS")))
+			else if (MatchIdentifier(TEXT("CPP"), ESearchCase::CaseSensitive) || MatchConstInt(TEXT("0")) || MatchConstInt(TEXT("1")) || MatchIdentifier(TEXT("WITH_HOT_RELOAD"), ESearchCase::CaseSensitive) || MatchIdentifier(TEXT("WITH_HOT_RELOAD_CTORS"), ESearchCase::CaseSensitive))
 			{
 				bConsumeAsCppText = !bInvertConditional;
 				PushCompilerDirective(ECompilerDirective::Insignificant);
@@ -2663,23 +2707,23 @@ UScriptStruct* FHeaderParser::CompileStructDeclaration(FClasses& AllClasses)
 					{
 						FError::Throwf(TEXT("Unexpected end of struct definition %s"), *Struct->GetName());
 					}
-					else if ( ch=='{' || (ch=='#' && (PeekIdentifier(TEXT("if")) || PeekIdentifier(TEXT("ifdef")))) )
+					else if ( ch=='{' || (ch=='#' && (PeekIdentifier(TEXT("if"), ESearchCase::CaseSensitive) || PeekIdentifier(TEXT("ifdef"), ESearchCase::CaseSensitive))) )
 					{
 						nest++;
 					}
-					else if ( ch=='}' || (ch=='#' && PeekIdentifier(TEXT("endif"))) )
+					else if ( ch=='}' || (ch=='#' && PeekIdentifier(TEXT("endif"), ESearchCase::CaseSensitive)) )
 					{
 						nest--;
 					}
 
 					if (nest==0)
 					{
-						RequireIdentifier(TEXT("endif"),TEXT("'if'"));
+						RequireIdentifier(TEXT("endif"), ESearchCase::CaseSensitive, TEXT("'if'"));
 					}
 				}
 			}
 		}
-		else if (Token.Matches(TEXT("#")) && MatchIdentifier(TEXT("pragma")))
+		else if (Token.Matches(TEXT('#')) && MatchIdentifier(TEXT("pragma"), ESearchCase::CaseSensitive))
 		{
 			// skip it and skip over the text, it is not recorded or processed
 			TCHAR c;
@@ -2693,7 +2737,7 @@ UScriptStruct* FHeaderParser::CompileStructDeclaration(FClasses& AllClasses)
 		}
 		else
 		{
-			if ( !Token.Matches( TEXT("}") ) )
+			if (!Token.Matches( TEXT('}')))
 			{
 				FToken DeclarationFirstToken = Token;
 				if (!SkipDeclaration(Token))
@@ -3097,14 +3141,14 @@ void FHeaderParser::CheckSparseClassData(const UStruct* StructToCheck)
 	if (!ClassToCheck)
 	{
 		// make sure we don't try to have sparse class data inside of a struct instead of a class
-		if (StructToCheck->HasMetaData(TEXT("SparseClassDataTypes")))
+		if (StructToCheck->HasMetaData(NAME_SparseClassDataTypes))
 		{
 			UE_LOG_ERROR_UHT(TEXT("%s contains sparse class data but is not a class."), *StructToCheck->GetName());
 		}
 		return;
 	}
 
-	if (!ClassToCheck->HasMetaData(TEXT("SparseClassDataTypes")))
+	if (!ClassToCheck->HasMetaData(NAME_SparseClassDataTypes))
 	{
 		return;
 	}
@@ -3324,7 +3368,7 @@ void FHeaderParser::VerifyPropertyMarkups( UClass* TargetClass )
 
 			if (Prop->HasAnyPropertyFlags(CPF_BlueprintVisible))
 			{
-				const FString& GetterFuncName = Prop->GetMetaData(TEXT("BlueprintGetter"));
+				const FString& GetterFuncName = Prop->GetMetaData(NAME_BlueprintGetter);
 				if (!GetterFuncName.IsEmpty())
 				{
 					if (UFunction* TargetFunc = FindTargetFunction(*GetterFuncName))
@@ -3340,7 +3384,7 @@ void FHeaderParser::VerifyPropertyMarkups( UClass* TargetClass )
 
 				if (!Prop->HasAnyPropertyFlags(CPF_BlueprintReadOnly))
 				{
-					const FString& SetterFuncName = Prop->GetMetaData(TEXT("BlueprintSetter"));
+					const FString& SetterFuncName = Prop->GetMetaData(NAME_BlueprintSetter);
 					if (!SetterFuncName.IsEmpty())
 					{
 						if (UFunction* TargetFunc = FindTargetFunction(*SetterFuncName))
@@ -3381,15 +3425,15 @@ void FHeaderParser::CompileDirective(FClasses& AllClasses)
 	{
 		FError::Throwf(TEXT("Missing compiler directive after '#'") );
 	}
-	else if (Directive.Matches(TEXT("Error")))
+	else if (Directive.Matches(TEXT("error"), ESearchCase::CaseSensitive))
 	{
-		FError::Throwf(TEXT("#Error directive encountered") );
+		FError::Throwf(TEXT("#error directive encountered") );
 	}
-	else if (Directive.Matches(TEXT("pragma")))
+	else if (Directive.Matches(TEXT("pragma"), ESearchCase::CaseSensitive))
 	{
 		// Ignore all pragmas
 	}
-	else if (Directive.Matches(TEXT("linenumber")))
+	else if (Directive.Matches(TEXT("linenumber"), ESearchCase::CaseSensitive))
 	{
 		FToken Number;
 		if (!GetToken(Number) || (Number.TokenType != TOKEN_Const) || (Number.Type != CPT_Int && Number.Type != CPT_Int64))
@@ -3403,7 +3447,7 @@ void FHeaderParser::CompileDirective(FClasses& AllClasses)
 			InputLine = newInputLine;
 		}
 	}
-	else if (Directive.Matches(TEXT("include")))
+	else if (Directive.Matches(TEXT("include"), ESearchCase::CaseSensitive))
 	{
 		FString ExpectedHeaderName = CurrentSrcFile->GetGeneratedHeaderFilename();
 		FToken IncludeName;
@@ -3415,10 +3459,10 @@ void FHeaderParser::CompileDirective(FClasses& AllClasses)
 			}
 		}
 	}
-	else if (Directive.Matches(TEXT("if")))
+	else if (Directive.Matches(TEXT("if"), ESearchCase::CaseSensitive))
 	{
 		// Eat the ! if present
-		bool bNotDefined = MatchSymbol(TEXT("!"));
+		bool bNotDefined = MatchSymbol(TEXT('!'));
 
 		int32 TempInt;
 		const bool bParsedInt = GetConstInt(TempInt);
@@ -3434,19 +3478,19 @@ void FHeaderParser::CompileDirective(FClasses& AllClasses)
 				FError::Throwf(TEXT("Missing define name '#if'") );
 			}
 
-			if ( Define.Matches(TEXT("WITH_EDITORONLY_DATA")) )
+			if ( Define.Matches(TEXT("WITH_EDITORONLY_DATA"), ESearchCase::CaseSensitive) )
 			{
 				PushCompilerDirective(ECompilerDirective::WithEditorOnlyData);
 			}
-			else if ( Define.Matches(TEXT("WITH_EDITOR")) )
+			else if ( Define.Matches(TEXT("WITH_EDITOR"), ESearchCase::CaseSensitive) )
 			{
 				PushCompilerDirective(ECompilerDirective::WithEditor);
 			}
-			else if (Define.Matches(TEXT("WITH_HOT_RELOAD")) || Define.Matches(TEXT("WITH_HOT_RELOAD_CTORS")) || Define.Matches(TEXT("1")))
+			else if (Define.Matches(TEXT("WITH_HOT_RELOAD"), ESearchCase::CaseSensitive) || Define.Matches(TEXT("WITH_HOT_RELOAD_CTORS"), ESearchCase::CaseSensitive) || Define.Matches(TEXT('1')))
 			{
 				PushCompilerDirective(ECompilerDirective::Insignificant);
 			}
-			else if ( Define.Matches(TEXT("CPP")) && bNotDefined)
+			else if ( Define.Matches(TEXT("CPP"), ESearchCase::CaseSensitive) && bNotDefined)
 			{
 				PushCompilerDirective(ECompilerDirective::Insignificant);
 			}
@@ -3456,7 +3500,7 @@ void FHeaderParser::CompileDirective(FClasses& AllClasses)
 			}
 		}
 	}
-	else if (Directive.Matches(TEXT("endif")))
+	else if (Directive.Matches(TEXT("endif"), ESearchCase::CaseSensitive))
 	{
 		if (CompilerDirectiveStack.Num() < 1)
 		{
@@ -3464,16 +3508,16 @@ void FHeaderParser::CompileDirective(FClasses& AllClasses)
 		}
 		CompilerDirectiveStack.Pop();
 	}
-	else if (Directive.Matches(TEXT("define")))
+	else if (Directive.Matches(TEXT("define"), ESearchCase::CaseSensitive))
 	{
 		// Ignore the define directive (can be multiline).
 		bDefineDirective = true;
 	}
-	else if (Directive.Matches(TEXT("ifdef")) || Directive.Matches(TEXT("ifndef")))
+	else if (Directive.Matches(TEXT("ifdef"), ESearchCase::CaseSensitive) || Directive.Matches(TEXT("ifndef"), ESearchCase::CaseSensitive))
 	{
 		PushCompilerDirective(ECompilerDirective::Insignificant);
 	}
-	else if (Directive.Matches(TEXT("undef")) || Directive.Matches(TEXT("else")))
+	else if (Directive.Matches(TEXT("undef"), ESearchCase::CaseSensitive) || Directive.Matches(TEXT("else"), ESearchCase::CaseSensitive))
 	{
 		// Ignore. UHT can only handle #if directive
 	}
@@ -3546,10 +3590,10 @@ void FHeaderParser::GetVarType(
 	bool bNativeConst = false;
 	bool bNativeConstTemplateArg = false;
 
-	const bool bIsParamList = (VariableCategory != EVariableCategory::Member) && MatchIdentifier(TEXT("UPARAM"));
+	const bool bIsParamList = (VariableCategory != EVariableCategory::Member) && MatchIdentifier(TEXT("UPARAM"), ESearchCase::CaseSensitive);
 
 	// No specifiers are allowed inside a TArray
-	if ((OuterPropertyType == NULL) || !OuterPropertyType->Matches(TEXT("TArray")))
+	if ((OuterPropertyType == NULL) || !OuterPropertyType->Matches(TEXT("TArray"), ESearchCase::CaseSensitive))
 	{
 		// New-style UPROPERTY() syntax 
 		if (PropertyDeclarationStyle == EPropertyDeclarationStyle::UPROPERTY || bIsParamList)
@@ -3561,7 +3605,7 @@ void FHeaderParser::GetVarType(
 	if (VariableCategory != EVariableCategory::Member)
 	{
 		// const before the variable type support (only for params)
-		if (MatchIdentifier(TEXT("const")))
+		if (MatchIdentifier(TEXT("const"), ESearchCase::CaseSensitive))
 		{
 			Flags |= CPF_ConstParm;
 			bNativeConst = true;
@@ -3664,7 +3708,7 @@ void FHeaderParser::GetVarType(
 						UE_LOG_ERROR_UHT(TEXT("Cannot specify a property as being both BlueprintReadOnly and BlueprintReadWrite."));
 					}
 
-					const FString* PrivateAccessMD = MetaDataFromNewStyle.Find(TEXT("AllowPrivateAccess"));  // FBlueprintMetadata::MD_AllowPrivateAccess
+					const FString* PrivateAccessMD = MetaDataFromNewStyle.Find(NAME_AllowPrivateAccess);  // FBlueprintMetadata::MD_AllowPrivateAccess
 					const bool bAllowPrivateAccess = PrivateAccessMD ? (*PrivateAccessMD != TEXT("false")) : false;
 					if (CurrentAccessSpecifier == ACCESS_Private && !bAllowPrivateAccess)
 					{
@@ -3694,7 +3738,7 @@ void FHeaderParser::GetVarType(
 					}
 
 					const FString BlueprintSetterFunction = RequireExactlyOneSpecifierValue(Specifier);
-					MetaDataFromNewStyle.Add(TEXT("BlueprintSetter"), BlueprintSetterFunction);
+					MetaDataFromNewStyle.Add(NAME_BlueprintSetter, BlueprintSetterFunction);
 
 					Flags |= CPF_BlueprintVisible;
 					bSeenBlueprintWriteSpecifier = true;
@@ -3708,7 +3752,7 @@ void FHeaderParser::GetVarType(
 						UE_LOG_ERROR_UHT(TEXT("Cannot specify both BlueprintReadOnly and BlueprintReadWrite or BlueprintSetter."), *Specifier.Key);
 					}
 
-					const FString* PrivateAccessMD = MetaDataFromNewStyle.Find(TEXT("AllowPrivateAccess"));  // FBlueprintMetadata::MD_AllowPrivateAccess
+					const FString* PrivateAccessMD = MetaDataFromNewStyle.Find(NAME_AllowPrivateAccess);  // FBlueprintMetadata::MD_AllowPrivateAccess
 					const bool bAllowPrivateAccess = PrivateAccessMD ? (*PrivateAccessMD != TEXT("false")) : false;
 					if (CurrentAccessSpecifier == ACCESS_Private && !bAllowPrivateAccess)
 					{
@@ -3734,7 +3778,7 @@ void FHeaderParser::GetVarType(
 					}
 
 					const FString BlueprintGetterFunction = RequireExactlyOneSpecifierValue(Specifier);
-					MetaDataFromNewStyle.Add(TEXT("BlueprintGetter"), BlueprintGetterFunction);
+					MetaDataFromNewStyle.Add(NAME_BlueprintGetter, BlueprintGetterFunction);
 
 					Flags        |= CPF_BlueprintVisible;
 					bSeenBlueprintGetterSpecifier = true;
@@ -3974,7 +4018,7 @@ void FHeaderParser::GetVarType(
 	}
 
 	{
-		const FString* ExposeOnSpawnStr = MetaDataFromNewStyle.Find(TEXT("ExposeOnSpawn"));
+		const FString* ExposeOnSpawnStr = MetaDataFromNewStyle.Find(NAME_ExposeOnSpawn);
 		const bool bExposeOnSpawn = (NULL != ExposeOnSpawnStr);
 		if (bExposeOnSpawn)
 		{
@@ -4045,27 +4089,27 @@ void FHeaderParser::GetVarType(
 	bool bUnconsumedEnumKeyword   = false;
 	bool bUnconsumedConstKeyword  = false;
 
-	if (MatchIdentifier(TEXT("const")))
+	if (MatchIdentifier(TEXT("const"), ESearchCase::CaseSensitive))
 	{
 		//@TODO: UCREMOVAL: Should use this to set the new (currently non-existent) CPF_Const flag appropriately!
 		bUnconsumedConstKeyword = true;
 		bNativeConst = true;
 	}
 
-	if (MatchIdentifier(TEXT("mutable")))
+	if (MatchIdentifier(TEXT("mutable"), ESearchCase::CaseSensitive))
 	{
 		//@TODO: Should flag as settable from a const context, but this is at least good enough to allow use for C++ land
 	}
 
-	if (MatchIdentifier(TEXT("struct")))
+	if (MatchIdentifier(TEXT("struct"), ESearchCase::CaseSensitive))
 	{
 		bUnconsumedStructKeyword = true;
 	}
-	else if (MatchIdentifier(TEXT("class")))
+	else if (MatchIdentifier(TEXT("class"), ESearchCase::CaseSensitive))
 	{
 		bUnconsumedClassKeyword = true;
 	}
-	else if (MatchIdentifier(TEXT("enum")))
+	else if (MatchIdentifier(TEXT("enum"), ESearchCase::CaseSensitive))
 	{
 		if (VariableCategory == EVariableCategory::Member)
 		{
@@ -4084,57 +4128,57 @@ void FHeaderParser::GetVarType(
 
 	RedirectTypeIdentifier(VarType);
 
-	if ( VarType.Matches(TEXT("int8")) )
+	if ( VarType.Matches(TEXT("int8"), ESearchCase::CaseSensitive) )
 	{
 		VarProperty = FPropertyBase(CPT_Int8);
 	}
-	else if ( VarType.Matches(TEXT("int16")) )
+	else if ( VarType.Matches(TEXT("int16"), ESearchCase::CaseSensitive) )
 	{
 		VarProperty = FPropertyBase(CPT_Int16);
 	}
-	else if ( VarType.Matches(TEXT("int32")) )
+	else if ( VarType.Matches(TEXT("int32"), ESearchCase::CaseSensitive) )
 	{
 		VarProperty = FPropertyBase(CPT_Int);
 	}
-	else if ( VarType.Matches(TEXT("int64")) )
+	else if ( VarType.Matches(TEXT("int64"), ESearchCase::CaseSensitive) )
 	{
 		VarProperty = FPropertyBase(CPT_Int64);
 	}
-	else if ( VarType.Matches(TEXT("uint64")) && IsBitfieldProperty() )
+	else if ( VarType.Matches(TEXT("uint64"), ESearchCase::CaseSensitive) && IsBitfieldProperty() )
 	{
 		// 64-bit bitfield (bool) type, treat it like 8 bit type
 		VarProperty = FPropertyBase(CPT_Bool8);
 	}
-	else if ( VarType.Matches(TEXT("uint32")) && IsBitfieldProperty() )
+	else if ( VarType.Matches(TEXT("uint32"), ESearchCase::CaseSensitive) && IsBitfieldProperty() )
 	{
 		// 32-bit bitfield (bool) type, treat it like 8 bit type
 		VarProperty = FPropertyBase(CPT_Bool8);
 	}
-	else if ( VarType.Matches(TEXT("uint16")) && IsBitfieldProperty() )
+	else if ( VarType.Matches(TEXT("uint16"), ESearchCase::CaseSensitive) && IsBitfieldProperty() )
 	{
 		// 16-bit bitfield (bool) type, treat it like 8 bit type.
 		VarProperty = FPropertyBase(CPT_Bool8);
 	}
-	else if ( VarType.Matches(TEXT("uint8")) && IsBitfieldProperty() )
+	else if ( VarType.Matches(TEXT("uint8"), ESearchCase::CaseSensitive) && IsBitfieldProperty() )
 	{
 		// 8-bit bitfield (bool) type
 		VarProperty = FPropertyBase(CPT_Bool8);
 	}
-	else if ( VarType.Matches(TEXT("int")) )
+	else if ( VarType.Matches(TEXT("int"), ESearchCase::CaseSensitive) )
 	{
 		VarProperty = FPropertyBase(CPT_Int, EIntType::Unsized);
 	}
-	else if ( VarType.Matches(TEXT("signed")) )
+	else if ( VarType.Matches(TEXT("signed"), ESearchCase::CaseSensitive) )
 	{
-		MatchIdentifier(TEXT("int"));
+		MatchIdentifier(TEXT("int"), ESearchCase::CaseSensitive);
 		VarProperty = FPropertyBase(CPT_Int, EIntType::Unsized);
 	}
-	else if (VarType.Matches(TEXT("unsigned")))
+	else if (VarType.Matches(TEXT("unsigned"), ESearchCase::CaseSensitive))
 	{
-		MatchIdentifier(TEXT("int"));
+		MatchIdentifier(TEXT("int"), ESearchCase::CaseSensitive);
 		VarProperty = FPropertyBase(CPT_UInt32, EIntType::Unsized);
 	}
-	else if ( VarType.Matches(TEXT("bool")) )
+	else if ( VarType.Matches(TEXT("bool"), ESearchCase::CaseSensitive) )
 	{
 		if (IsBitfieldProperty())
 		{
@@ -4143,41 +4187,41 @@ void FHeaderParser::GetVarType(
 		// C++ bool type
 		VarProperty = FPropertyBase(CPT_Bool);
 	}
-	else if ( VarType.Matches(TEXT("uint8")) )
+	else if ( VarType.Matches(TEXT("uint8"), ESearchCase::CaseSensitive) )
 	{
 		// Intrinsic Byte type.
 		VarProperty = FPropertyBase(CPT_Byte);
 	}
-	else if ( VarType.Matches(TEXT("uint16")) )
+	else if ( VarType.Matches(TEXT("uint16"), ESearchCase::CaseSensitive) )
 	{
 		VarProperty = FPropertyBase(CPT_UInt16);
 	}
-	else if ( VarType.Matches(TEXT("uint32")) )
+	else if ( VarType.Matches(TEXT("uint32"), ESearchCase::CaseSensitive) )
 	{
 		VarProperty = FPropertyBase(CPT_UInt32);
 	}
-	else if ( VarType.Matches(TEXT("uint64")) )
+	else if ( VarType.Matches(TEXT("uint64"), ESearchCase::CaseSensitive) )
 	{
 		VarProperty = FPropertyBase(CPT_UInt64);
 	}
-	else if ( VarType.Matches(TEXT("float")) )
+	else if ( VarType.Matches(TEXT("float"), ESearchCase::CaseSensitive) )
 	{
 		// Intrinsic single precision floating point type.
 		VarProperty = FPropertyBase(CPT_Float);
 	}
-	else if ( VarType.Matches(TEXT("double")) )
+	else if ( VarType.Matches(TEXT("double"), ESearchCase::CaseSensitive) )
 	{
 		// Intrinsic double precision floating point type type.
 		VarProperty = FPropertyBase(CPT_Double);
 	}
-	else if ( VarType.Matches(TEXT("FName")) )
+	else if ( VarType.Matches(TEXT("FName"), ESearchCase::CaseSensitive) )
 	{
 		// Intrinsic Name type.
 		VarProperty = FPropertyBase(CPT_Name);
 	}
-	else if ( VarType.Matches(TEXT("TArray")) )
+	else if ( VarType.Matches(TEXT("TArray"), ESearchCase::CaseSensitive) )
 	{
-		RequireSymbol( TEXT("<"), TEXT("'tarray'") );
+		RequireSymbol( TEXT('<'), TEXT("'tarray'") );
 
 		VarType.PropertyFlags = Flags;
 
@@ -4188,7 +4232,7 @@ void FHeaderParser::GetVarType(
 		}
 		// TODO: Prevent sparse delegate types from being used in a container
 
-		if (VarProperty.MetaData.Find(TEXT("NativeConst")))
+		if (VarProperty.MetaData.Find(NAME_NativeConst))
 		{
 			bNativeConstTemplateArg = true;
 		}
@@ -4202,10 +4246,10 @@ void FHeaderParser::GetVarType(
 			FError::Throwf(TEXT("Missing token while parsing TArray."));
 		}
 
-		if (CloseTemplateToken.TokenType != TOKEN_Symbol || FCString::Stricmp(CloseTemplateToken.Identifier, TEXT(">")))
+		if (CloseTemplateToken.TokenType != TOKEN_Symbol || !CloseTemplateToken.Matches(TEXT('>')))
 		{
 			// If we didn't find a comma, report it
-			if (FCString::Stricmp(CloseTemplateToken.Identifier, TEXT(",")))
+			if (!CloseTemplateToken.Matches(TEXT(',')))
 			{
 				FError::Throwf(TEXT("Expected '>' but found '%s'"), CloseTemplateToken.Identifier);
 			}
@@ -4220,9 +4264,9 @@ void FHeaderParser::GetVarType(
 			FError::Throwf(TEXT("Found '%s' - explicit allocators are not supported in TArray properties."), AllocatorToken.Identifier);
 		}
 	}
-	else if ( VarType.Matches(TEXT("TMap")) )
+	else if ( VarType.Matches(TEXT("TMap"), ESearchCase::CaseSensitive) )
 	{
-		RequireSymbol( TEXT("<"), TEXT("'tmap'") );
+		RequireSymbol( TEXT('<'), TEXT("'tmap'") );
 
 		VarType.PropertyFlags = Flags;
 
@@ -4245,7 +4289,7 @@ void FHeaderParser::GetVarType(
 		}
 
 		FToken CommaToken;
-		if (!GetToken(CommaToken, /*bNoConsts=*/ true) || CommaToken.TokenType != TOKEN_Symbol || FCString::Stricmp(CommaToken.Identifier, TEXT(",")))
+		if (!GetToken(CommaToken, /*bNoConsts=*/ true) || CommaToken.TokenType != TOKEN_Symbol || !CommaToken.Matches(TEXT(',')))
 		{
 			FError::Throwf(TEXT("Missing value type while parsing TMap."));
 		}
@@ -4268,10 +4312,10 @@ void FHeaderParser::GetVarType(
 			FError::Throwf(TEXT("Missing token while parsing TMap."));
 		}
 
-		if (CloseTemplateToken.TokenType != TOKEN_Symbol || FCString::Stricmp(CloseTemplateToken.Identifier, TEXT(">")))
+		if (CloseTemplateToken.TokenType != TOKEN_Symbol || !CloseTemplateToken.Matches(TEXT('>')))
 		{
 			// If we didn't find a comma, report it
-			if (FCString::Stricmp(CloseTemplateToken.Identifier, TEXT(",")))
+			if (!CloseTemplateToken.Matches(TEXT(',')))
 			{
 				FError::Throwf(TEXT("Expected '>' but found '%s'"), CloseTemplateToken.Identifier);
 			}
@@ -4286,9 +4330,9 @@ void FHeaderParser::GetVarType(
 			FError::Throwf(TEXT("Found '%s' - explicit allocators are not supported in TMap properties."), AllocatorToken.Identifier);
 		}
 	}
-	else if ( VarType.Matches(TEXT("TSet")) )
+	else if ( VarType.Matches(TEXT("TSet"), ESearchCase::CaseSensitive) )
 	{
-		RequireSymbol( TEXT("<"), TEXT("'tset'") );
+		RequireSymbol( TEXT('<'), TEXT("'tset'") );
 
 		VarType.PropertyFlags = Flags;
 
@@ -4318,10 +4362,10 @@ void FHeaderParser::GetVarType(
 			FError::Throwf(TEXT("Missing token while parsing TArray."));
 		}
 
-		if (CloseTemplateToken.TokenType != TOKEN_Symbol || FCString::Stricmp(CloseTemplateToken.Identifier, TEXT(">")))
+		if (CloseTemplateToken.TokenType != TOKEN_Symbol || !CloseTemplateToken.Matches(TEXT('>')))
 		{
 			// If we didn't find a comma, report it
-			if (FCString::Stricmp(CloseTemplateToken.Identifier, TEXT(",")))
+			if (!CloseTemplateToken.Matches(TEXT(',')))
 			{
 				FError::Throwf(TEXT("Expected '>' but found '%s'"), CloseTemplateToken.Identifier);
 			}
@@ -4336,13 +4380,13 @@ void FHeaderParser::GetVarType(
 			FError::Throwf(TEXT("Found '%s' - explicit KeyFuncs are not supported in TSet properties."), AllocatorToken.Identifier);
 		}
 	}
-	else if ( VarType.Matches(TEXT("FString")) )
+	else if ( VarType.Matches(TEXT("FString"), ESearchCase::CaseSensitive) )
 	{
 		VarProperty = FPropertyBase(CPT_String);
 
 		if (VariableCategory != EVariableCategory::Member)
 		{
-			if (MatchSymbol(TEXT("&")))
+			if (MatchSymbol(TEXT('&')))
 			{
 				if (Flags & CPF_ConstParm)
 				{
@@ -4363,20 +4407,20 @@ void FHeaderParser::GetVarType(
 			}
 		}
 	}
-	else if ( VarType.Matches(TEXT("Text") ) )
+	else if ( VarType.Matches(TEXT("Text"), ESearchCase::IgnoreCase) )
 	{
 		FError::Throwf(TEXT("%s' is missing a prefix, expecting 'FText'"), VarType.Identifier);
 	}
-	else if ( VarType.Matches(TEXT("FText") ) )
+	else if ( VarType.Matches(TEXT("FText"), ESearchCase::CaseSensitive) )
 	{
 		VarProperty = FPropertyBase(CPT_Text);
 	}
-	else if (VarType.Matches(TEXT("TEnumAsByte")))
+	else if (VarType.Matches(TEXT("TEnumAsByte"), ESearchCase::CaseSensitive))
 	{
-		RequireSymbol(TEXT("<"), VarType.Identifier);
+		RequireSymbol(TEXT('<'), VarType.Identifier);
 
 		// Eat the forward declaration enum text if present
-		MatchIdentifier(TEXT("enum"));
+		MatchIdentifier(TEXT("enum"), ESearchCase::CaseSensitive);
 
 		bool bFoundEnum = false;
 
@@ -4407,11 +4451,11 @@ void FHeaderParser::GetVarType(
 			FError::Throwf(TEXT("Expected the name of a previously defined enum"));
 		}
 
-		RequireSymbol(TEXT(">"), VarType.Identifier, ESymbolParseOption::CloseTemplateBracket);
+		RequireSymbol(TEXT('>'), VarType.Identifier, ESymbolParseOption::CloseTemplateBracket);
 	}
-	else if ( VarType.Matches(TEXT("TFieldPath")) )
+	else if (VarType.Matches(TEXT("TFieldPath"), ESearchCase::CaseSensitive ))
 	{
-		RequireSymbol( TEXT("<"), TEXT("'TFieldPath'") );
+		RequireSymbol( TEXT('<'), TEXT("'TFieldPath'") );
 
 		FFieldClass* PropertyClass = nullptr;
 		FToken PropertyTypeToken;
@@ -4432,7 +4476,7 @@ void FHeaderParser::GetVarType(
 			}
 		}
 
-		RequireSymbol(TEXT(">"), VarType.Identifier, ESymbolParseOption::CloseTemplateBracket);
+		RequireSymbol(TEXT('>'), VarType.Identifier, ESymbolParseOption::CloseTemplateBracket);
 
 		VarProperty = FPropertyBase(PropertyClass, CPT_FieldPath);
 	}
@@ -4560,24 +4604,24 @@ void FHeaderParser::GetVarType(
 			// An object reference of some type (maybe a restricted class?)
 			UClass* TempClass = NULL;
 
-			const bool bIsLazyPtrTemplate        = VarType.Matches(TEXT("TLazyObjectPtr"));
-			const bool bIsSoftObjectPtrTemplate  = VarType.Matches(TEXT("TSoftObjectPtr"));
-			const bool bIsSoftClassPtrTemplate   = VarType.Matches(TEXT("TSoftClassPtr"));
-			const bool bIsWeakPtrTemplate        = VarType.Matches(TEXT("TWeakObjectPtr"));
-			const bool bIsAutoweakPtrTemplate    = VarType.Matches(TEXT("TAutoWeakObjectPtr"));
-			const bool bIsScriptInterfaceWrapper = VarType.Matches(TEXT("TScriptInterface"));
-			const bool bIsSubobjectPtrTemplate   = VarType.Matches(TEXT("TSubobjectPtr"));
+			const bool bIsLazyPtrTemplate        = VarType.Matches(TEXT("TLazyObjectPtr"), ESearchCase::CaseSensitive);
+			const bool bIsSoftObjectPtrTemplate  = VarType.Matches(TEXT("TSoftObjectPtr"), ESearchCase::CaseSensitive);
+			const bool bIsSoftClassPtrTemplate   = VarType.Matches(TEXT("TSoftClassPtr"), ESearchCase::CaseSensitive);
+			const bool bIsWeakPtrTemplate        = VarType.Matches(TEXT("TWeakObjectPtr"), ESearchCase::CaseSensitive);
+			const bool bIsAutoweakPtrTemplate    = VarType.Matches(TEXT("TAutoWeakObjectPtr"), ESearchCase::CaseSensitive);
+			const bool bIsScriptInterfaceWrapper = VarType.Matches(TEXT("TScriptInterface"), ESearchCase::CaseSensitive);
+			const bool bIsSubobjectPtrTemplate   = VarType.Matches(TEXT("TSubobjectPtr"), ESearchCase::CaseSensitive);
 
 			bool bIsWeak     = false;
 			bool bIsLazy     = false;
 			bool bIsSoft     = false;
 			bool bWeakIsAuto = false;
 
-			if (VarType.Matches(TEXT("TSubclassOf")))
+			if (VarType.Matches(TEXT("TSubclassOf"), ESearchCase::CaseSensitive))
 			{
 				TempClass = UClass::StaticClass();
 			}
-			else if (VarType.Matches(TEXT("FScriptInterface")))
+			else if (VarType.Matches(TEXT("FScriptInterface"), ESearchCase::CaseSensitive))
 			{
 				TempClass = UInterface::StaticClass();
 				Flags |= CPF_UObjectWrapper;
@@ -4589,13 +4633,13 @@ void FHeaderParser::GetVarType(
 			}
 			else if (bIsLazyPtrTemplate || bIsWeakPtrTemplate || bIsAutoweakPtrTemplate || bIsScriptInterfaceWrapper || bIsSoftObjectPtrTemplate || bIsSubobjectPtrTemplate)
 			{
-				RequireSymbol(TEXT("<"), VarType.Identifier);
+				RequireSymbol(TEXT('<'), VarType.Identifier);
 
 				// Consume a forward class declaration 'class' if present
-				MatchIdentifier(TEXT("class"));
+				MatchIdentifier(TEXT("class"), ESearchCase::CaseSensitive);
 
 				// Also consume const
-				bNativeConstTemplateArg |= MatchIdentifier(TEXT("const"));
+				bNativeConstTemplateArg |= MatchIdentifier(TEXT("const"), ESearchCase::CaseSensitive);
 				
 				// Find the lazy/weak class
 				FToken InnerClass;
@@ -4638,7 +4682,7 @@ void FHeaderParser::GetVarType(
 					FError::Throwf(TEXT("%s: Missing template type"), VarType.Identifier);
 				}
 
-				RequireSymbol(TEXT(">"), VarType.Identifier, ESymbolParseOption::CloseTemplateBracket);
+				RequireSymbol(TEXT('>'), VarType.Identifier, ESymbolParseOption::CloseTemplateBracket);
 			}
 			else
 			{
@@ -4653,12 +4697,12 @@ void FHeaderParser::GetVarType(
 				VarProperty = FPropertyBase(TempClass, bAllowWeak && bIsWeak, bWeakIsAuto, bIsLazy, bIsSoft);
 				if (TempClass->IsChildOf(UClass::StaticClass()))
 				{
-					if ( MatchSymbol(TEXT("<")) )
+					if ( MatchSymbol(TEXT('<')) )
 					{
 						Flags |= CPF_UObjectWrapper;
 
 						// Consume a forward class declaration 'class' if present
-						MatchIdentifier(TEXT("class"));
+						MatchIdentifier(TEXT("class"), ESearchCase::CaseSensitive);
 
 						// Get the actual class type to restrict this to
 						FToken Limitor;
@@ -4671,7 +4715,7 @@ void FHeaderParser::GetVarType(
 
 						VarProperty.MetaClass = AllClasses.FindScriptClassOrThrow(Limitor.Identifier);
 
-						RequireSymbol( TEXT(">"), TEXT("'class limitor'"), ESymbolParseOption::CloseTemplateBracket );
+						RequireSymbol( TEXT('>'), TEXT("'class limitor'"), ESymbolParseOption::CloseTemplateBracket );
 					}
 					else
 					{
@@ -4704,14 +4748,14 @@ void FHeaderParser::GetVarType(
 				if (!(Flags & CPF_UObjectWrapper))
 				{
 					// Const after variable type but before pointer symbol
-					bNativeConst |= MatchIdentifier(TEXT("const"));
+					bNativeConst |= MatchIdentifier(TEXT("const"), ESearchCase::CaseSensitive);
 
-					RequireSymbol(TEXT("*"), TEXT("Expected a pointer type"));
+					RequireSymbol(TEXT('*'), TEXT("Expected a pointer type"));
 
 					// Swallow trailing 'const' after pointer properties
 					if (VariableCategory == EVariableCategory::Member)
 					{
-						MatchIdentifier(TEXT("const"));
+						MatchIdentifier(TEXT("const"), ESearchCase::CaseSensitive);
 					}
 
 					VarProperty.PointerType = EPointerType::Native;
@@ -4747,7 +4791,7 @@ void FHeaderParser::GetVarType(
 	if (VariableCategory != EVariableCategory::Member)
 	{
 		// const after the variable type support (only for params)
-		if (MatchIdentifier(TEXT("const")))
+		if (MatchIdentifier(TEXT("const"), ESearchCase::CaseSensitive))
 		{
 			Flags |= CPF_ConstParm;
 			bNativeConst = true;
@@ -4781,7 +4825,7 @@ void FHeaderParser::GetVarType(
 		FError::Throwf(TEXT("Inappropriate keyword 'enum' on variable of type '%s'"), VarType.Identifier );
 	}
 
-	if (MatchSymbol(TEXT("*")))
+	if (MatchSymbol(TEXT('*')))
 	{
 		FError::Throwf(TEXT("Inappropriate '*' on variable of type '%s', cannot have an exposed pointer to this type."), VarType.Identifier );
 	}
@@ -4790,11 +4834,11 @@ void FHeaderParser::GetVarType(
 	if (VariableCategory == EVariableCategory::Member && OwnerStruct->IsA<UClass>() && ((UClass*)OwnerStruct)->HasAnyClassFlags(CLASS_Const))
 	{
 		// Eat a 'not quite truthful' const after the type; autogenerated for member variables of const classes.
-		bNativeConst |= MatchIdentifier(TEXT("const"));
+		bNativeConst |= MatchIdentifier(TEXT("const"), ESearchCase::CaseSensitive);
 	}
 
 	// Arrays are passed by reference but are only implicitly so; setting it explicitly could cause a problem with replicated functions
-	if (MatchSymbol(TEXT("&")))
+	if (MatchSymbol(TEXT('&')))
 	{
 		switch (VariableCategory)
 		{
@@ -4923,11 +4967,11 @@ void FHeaderParser::GetVarType(
 	VarProperty.MetaData = MetaDataFromNewStyle;
 	if (bNativeConst)
 	{
-		VarProperty.MetaData.Add(TEXT("NativeConst"), TEXT(""));
+		VarProperty.MetaData.Add(NAME_NativeConst, FString());
 	}
 	if (bNativeConstTemplateArg)
 	{
-		VarProperty.MetaData.Add(TEXT("NativeConstTemplateArg"), TEXT(""));
+		VarProperty.MetaData.Add(NAME_NativeConstTemplateArg, FString());
 	}
 	
 	if (ParsedVarIndexRange)
@@ -4955,7 +4999,7 @@ EFindName FHeaderParser::GetFindFlagForPropertyName(const TCHAR* PropertyName)
 		return FNAME_Add;
 	}
 	// Add it to the list for future look ups
-	PreviousNames.Add(UpperPropertyStr,1);
+	PreviousNames.Add(MoveTemp(UpperPropertyStr),1);
 	FName CurrentText(PropertyName,FNAME_Find); // keep generating this FName in case it has been affecting the case of future FNames.
 	return FNAME_Replace_Not_Safe_For_Threading;
 }
@@ -5015,8 +5059,8 @@ FProperty* FHeaderParser::GetVarNameAndDim
 
 			// We allow deprecated properties in blueprints that have getters and setters assigned as they may be part of a backwards compatibility path
 			const bool bBlueprintVisible = (VarProperty.PropertyFlags & CPF_BlueprintVisible) > 0;
-			const bool bWarnOnGetter = bBlueprintVisible && !VarProperty.MetaData.Contains(TEXT("BlueprintGetter"));
-			const bool bWarnOnSetter = bBlueprintVisible && !(VarProperty.PropertyFlags & CPF_BlueprintReadOnly) && !VarProperty.MetaData.Contains(TEXT("BlueprintSetter"));
+			const bool bWarnOnGetter = bBlueprintVisible && !VarProperty.MetaData.Contains(NAME_BlueprintGetter);
+			const bool bWarnOnSetter = bBlueprintVisible && !(VarProperty.PropertyFlags & CPF_BlueprintReadOnly) && !VarProperty.MetaData.Contains(NAME_BlueprintSetter);
 
 			if (bWarnOnGetter)
 			{
@@ -5037,7 +5081,7 @@ FProperty* FHeaderParser::GetVarNameAndDim
 			}
 
 			VarProperty.PropertyFlags |= CPF_Deprecated;
-			VarName = VarName.Mid(0, DeprecatedIndex);
+			VarName.MidInline(0, DeprecatedIndex, false);
 
 			FCString::Strcpy(VarProperty.Identifier, *VarName);
 		}
@@ -5084,7 +5128,7 @@ FProperty* FHeaderParser::GetVarNameAndDim
 
 	// Get optional dimension immediately after name.
 	FToken Dimensions;
-	if (MatchSymbol(TEXT("[")))
+	if (MatchSymbol(TEXT('[')))
 	{
 		switch (VariableCategory)
 		{
@@ -5152,16 +5196,16 @@ FProperty* FHeaderParser::GetVarNameAndDim
 					int32 ClosingParen = FindMatchingClosingParenthesis(Temp);
 					if (ClosingParen == TempLen - 1)
 					{
-						Temp = Temp.Mid(1, TempLen - 2);
+						Temp.MidInline(1, TempLen - 2, false);
 						bAgain = true;
 					}
 				}
 
 				for (const TCHAR* Cast : Casts)
 				{
-					if (Temp.StartsWith(Cast))
+					if (Temp.StartsWith(Cast, ESearchCase::CaseSensitive))
 					{
-						Temp = Temp.RightChop(FCString::Strlen(Cast));
+						Temp.RightChopInline(FCString::Strlen(Cast), false);
 						bAgain = true;
 					}
 				}
@@ -5180,10 +5224,10 @@ FProperty* FHeaderParser::GetVarNameAndDim
 		if (Enum)
 		{
 			// set the ArraySizeEnum if applicable
-			VarProperty.MetaData.Add("ArraySizeEnum", Enum->GetPathName());
+			VarProperty.MetaData.Add(NAME_ArraySizeEnum, Enum->GetPathName());
 		}
 
-		MatchSymbol(TEXT("]"));
+		MatchSymbol(TEXT(']'));
 	}
 
 	// Try gathering metadata for member fields
@@ -5366,9 +5410,9 @@ bool FHeaderParser::CompileDeclaration(FClasses& AllClasses, TArray<UDelegateFun
 		{
 			FError::Throwf(TEXT("%s must occur inside the native interface definition"), Token.Identifier);
 		}
-		RequireSymbol(TEXT("("), Token.Identifier);
+		RequireSymbol(TEXT('('), Token.Identifier);
 		CompileVersionDeclaration(GetCurrentClass());
-		RequireSymbol(TEXT(")"), Token.Identifier);
+		RequireSymbol(TEXT(')'), Token.Identifier);
 
 		FClassMetaData* ClassData = GetCurrentClassData();
 		if (!ClassData)
@@ -5400,9 +5444,9 @@ bool FHeaderParser::CompileDeclaration(FClasses& AllClasses, TArray<UDelegateFun
 		{
 			FError::Throwf(TEXT("%s must occur inside the interface definition"), Token.Identifier);
 		}
-		RequireSymbol(TEXT("("), Token.Identifier);
+		RequireSymbol(TEXT('('), Token.Identifier);
 		CompileVersionDeclaration(GetCurrentClass());
-		RequireSymbol(TEXT(")"), Token.Identifier);
+		RequireSymbol(TEXT(')'), Token.Identifier);
 
 		FClassMetaData* ClassData = GetCurrentClassData();
 
@@ -5443,9 +5487,9 @@ bool FHeaderParser::CompileDeclaration(FClasses& AllClasses, TArray<UDelegateFun
 			CurrentAccessSpecifier = ACCESS_Public;
 		}
 
-		RequireSymbol(TEXT("("), Token.Identifier);
+		RequireSymbol(TEXT('('), Token.Identifier);
 		CompileVersionDeclaration(GetCurrentClass());
-		RequireSymbol(TEXT(")"), Token.Identifier);
+		RequireSymbol(TEXT(')'), Token.Identifier);
 
 		ClassData->SetGeneratedBodyLine(InputLine);
 
@@ -5514,14 +5558,14 @@ bool FHeaderParser::CompileDeclaration(FClasses& AllClasses, TArray<UDelegateFun
 		return true;
 	}
 
-	if (Token.Matches(TEXT("#")))
+	if (Token.Matches(TEXT('#')))
 	{
 		// Compiler directive.
 		CompileDirective(AllClasses);
 		return true;
 	}
 
-	if (bEncounteredNewStyleClass_UnmatchedBrackets && Token.Matches(TEXT("}")))
+	if (bEncounteredNewStyleClass_UnmatchedBrackets && Token.Matches(TEXT('}')))
 	{
 		if (ClassDefinitionRanges.Contains(GetCurrentClass()))
 		{
@@ -5576,7 +5620,7 @@ bool FHeaderParser::CompileDeclaration(FClasses& AllClasses, TArray<UDelegateFun
 		return true;
 	}
 
-	if (Token.Matches(TEXT(";")))
+	if (Token.Matches(TEXT(';')))
 	{
 		if (GetToken(Token))
 		{
@@ -5595,26 +5639,26 @@ bool FHeaderParser::CompileDeclaration(FClasses& AllClasses, TArray<UDelegateFun
 			FToken ConstructorToken = Token;
 
 			// Allow explicit constructors
-			bool bFoundExplicit = ConstructorToken.Matches(TEXT("explicit"));
+			bool bFoundExplicit = ConstructorToken.Matches(TEXT("explicit"), ESearchCase::CaseSensitive);
 			if (bFoundExplicit)
 			{
 				GetToken(ConstructorToken);
 			}
 
 			bool bSkippedAPIToken = false;
-			if (FString(ConstructorToken.Identifier).EndsWith("_API"))
+			if (FString(ConstructorToken.Identifier).EndsWith(TEXT("_API"), ESearchCase::CaseSensitive))
 			{
 				if (!bFoundExplicit)
 				{
 					// Explicit can come before or after an _API
-					MatchIdentifier(TEXT("explicit"));
+					MatchIdentifier(TEXT("explicit"), ESearchCase::CaseSensitive);
 				}
 
 				GetToken(ConstructorToken);
 				bSkippedAPIToken = true;
 			}
 
-			if (ConstructorToken.Matches(NameLookupCPP.GetNameCPP(Class)))
+			if (ConstructorToken.Matches(NameLookupCPP.GetNameCPP(Class), ESearchCase::IgnoreCase))
 			{
 				if (TryToMatchConstructorParameterList(ConstructorToken))
 				{
@@ -5638,88 +5682,76 @@ bool FHeaderParser::CompileDeclaration(FClasses& AllClasses, TArray<UDelegateFun
 	// Determine if this statement is a serialize function declaration
 	if (bEncounteredNewStyleClass_UnmatchedBrackets && IsInAClass() && TopNest->NestType == ENestType::Class)
 	{
-		static const FName NAME_Virtual(TEXT("virtual"));
-		static const FName NAME_Void(TEXT("void"));
-		static const FName NAME_Serialize(TEXT("Serialize"));
-		static const FName NAME_OpenBracket(TEXT("("));
-		static const FName NAME_CloseBracket(TEXT(")"));
-		static const FName NAME_FArchive(TEXT("FArchive"));
-		static const FName NAME_FStructuredArchive(TEXT("FStructuredArchive"));
-		static const FName NAME_FStructuredArchiveRecord(TEXT("FStructuredArchiveRecord"));
-		static const FName NAME_Reference(TEXT("&"));
-		static const FName NAME_ClassMember(TEXT("::"));
-		static const FName NAME_FRecord(TEXT("FRecord"));
-
-		while (Token.Matches(NAME_Virtual) || FString(Token.Identifier).EndsWith(TEXT("_API")))
+		while (Token.Matches(TEXT("virtual"), ESearchCase::CaseSensitive) || FString(Token.Identifier).EndsWith(TEXT("_API"), ESearchCase::CaseSensitive))
 		{
 			GetToken(Token);
 		}
 
-		if (Token.Identifier == NAME_Void)
+		if (Token.Matches(TEXT("void"), ESearchCase::CaseSensitive))
 		{
 			GetToken(Token);
-			if (Token.Identifier == NAME_Serialize)
+			if (Token.Matches(TEXT("Serialize"), ESearchCase::CaseSensitive))
 			{
 				GetToken(Token);
-				if (Token.Identifier == NAME_OpenBracket)
+				if (Token.Matches(TEXT('(')))
 				{
 					GetToken(Token);
 
 					ESerializerArchiveType ArchiveType = ESerializerArchiveType::None;
-					if (Token.Identifier == NAME_FArchive)
+					if (Token.Matches(TEXT("FArchive"), ESearchCase::CaseSensitive))
 					{
 						GetToken(Token);
-						if (Token.Identifier == NAME_Reference)
+						if (Token.Matches(TEXT('&')))
 						{
 							GetToken(Token);
 
 							// Allow the declaration to not define a name for the archive parameter
-							if (Token.Identifier != NAME_CloseBracket)
+							if (!Token.Matches(TEXT(')')))
 							{
 								GetToken(Token);
 							}
 
-							if (Token.Identifier == NAME_CloseBracket)
+							if (Token.Matches(TEXT(')')))
 							{
 								ArchiveType = ESerializerArchiveType::Archive;
 							}
 						}
 					}
-					else if (Token.Identifier == NAME_FStructuredArchive)
+					else if (Token.Matches(TEXT("FStructuredArchive"), ESearchCase::CaseSensitive))
 					{
 						GetToken(Token);
-						if (Token.Identifier == NAME_ClassMember)
+						if (Token.Matches(TEXT("::"), ESearchCase::CaseSensitive))
 						{
 							GetToken(Token);
 
-							if (Token.Identifier == NAME_FRecord)
+							if (Token.Matches(TEXT("FRecord"), ESearchCase::CaseSensitive))
 							{
 								GetToken(Token);
 
 								// Allow the declaration to not define a name for the slot parameter
-								if (Token.Identifier != NAME_CloseBracket)
+								if (!Token.Matches(TEXT(')')))
 								{
 									GetToken(Token);
 								}
 
-								if (Token.Identifier == NAME_CloseBracket)
+								if (Token.Matches(TEXT(')')))
 								{
 									ArchiveType = ESerializerArchiveType::StructuredArchiveRecord;
 								}
 							}
 						}
 					}
-					else if (Token.Identifier == NAME_FStructuredArchiveRecord)
+					else if (Token.Matches(TEXT("FStructuredArchiveRecord"), ESearchCase::CaseSensitive))
 					{
 						GetToken(Token);
 
 						// Allow the declaration to not define a name for the slot parameter
-						if (Token.Identifier != NAME_CloseBracket)
+						if (!Token.Matches(TEXT(')')))
 						{
 							GetToken(Token);
 						}
 
-						if (Token.Identifier == NAME_CloseBracket)
+						if (Token.Matches(TEXT(')')))
 						{
 							ArchiveType = ESerializerArchiveType::StructuredArchiveRecord;
 						}
@@ -5757,35 +5789,35 @@ bool FHeaderParser::SkipDeclaration(FToken& Token)
 	// Consume all tokens until the end of declaration/definition has been found.
 	int32 NestedScopes = 0;
 	// Check if this is a class/struct declaration in which case it can be followed by member variable declaration.	
-	bool bPossiblyClassDeclaration = Token.Matches(TEXT("class")) || Token.Matches(TEXT("struct"));
+	bool bPossiblyClassDeclaration = Token.Matches(TEXT("class"), ESearchCase::CaseSensitive) || Token.Matches(TEXT("struct"), ESearchCase::CaseSensitive);
 	// (known) macros can end without ; or } so use () to find the end of the declaration.
 	// However, we don't want to use it with DECLARE_FUNCTION, because we need it to be treated like a function.
-	bool bMacroDeclaration      = ProbablyAMacro(Token.Identifier) && !Token.Matches("DECLARE_FUNCTION");
+	bool bMacroDeclaration      = ProbablyAMacro(Token.Identifier) && !Token.Matches(TEXT("DECLARE_FUNCTION"), ESearchCase::CaseSensitive);
 	bool bEndOfDeclarationFound = false;
 	bool bDefinitionFound       = false;
-	const TCHAR* OpeningBracket = bMacroDeclaration ? TEXT("(") : TEXT("{");
-	const TCHAR* ClosingBracket = bMacroDeclaration ? TEXT(")") : TEXT("}");
+	TCHAR OpeningBracket = bMacroDeclaration ? TEXT('(') : TEXT('{');
+	TCHAR ClosingBracket = bMacroDeclaration ? TEXT(')') : TEXT('}');
 	bool bRetestCurrentToken = false;
 	while (bRetestCurrentToken || GetToken(Token))
 	{
 		// If we find parentheses at top-level and we think it's a class declaration then it's more likely
 		// to be something like: class UThing* GetThing();
-		if (bPossiblyClassDeclaration && NestedScopes == 0 && Token.Matches(TEXT("(")))
+		if (bPossiblyClassDeclaration && NestedScopes == 0 && Token.Matches(TEXT('(')))
 		{
 			bPossiblyClassDeclaration = false;
 		}
 
 		bRetestCurrentToken = false;
-		if (Token.Matches(TEXT(";")) && NestedScopes == 0)
+		if (Token.Matches(TEXT(';')) && NestedScopes == 0)
 		{
 			bEndOfDeclarationFound = true;
 			break;
 		}
 
-		if (!bMacroDeclaration && Token.Matches(TEXT("PURE_VIRTUAL")) && NestedScopes == 0)
+		if (!bMacroDeclaration && Token.Matches(TEXT("PURE_VIRTUAL"), ESearchCase::CaseSensitive) && NestedScopes == 0)
 		{
-			OpeningBracket = TEXT("(");
-			ClosingBracket = TEXT(")");
+			OpeningBracket = TEXT('(');
+			ClosingBracket = TEXT(')');
 		}
 
 		if (Token.Matches(OpeningBracket))
@@ -5809,7 +5841,7 @@ bool FHeaderParser::SkipDeclaration(FToken& Token)
 					GetToken(Token);
 
 					// If Strcmp returns 0, it is probably a class, else a macro.
-					bReallyEndDeclaration = FCString::Strcmp(PossibleBracketToken.Identifier, TEXT("{")) != 0;
+					bReallyEndDeclaration = !PossibleBracketToken.Matches(TEXT('{'));
 				}
 
 				if (bReallyEndDeclaration)
@@ -5827,8 +5859,8 @@ bool FHeaderParser::SkipDeclaration(FToken& Token)
 		else if (bMacroDeclaration && NestedScopes == 0)
 		{
 			bMacroDeclaration = false;
-			OpeningBracket = TEXT("{");
-			ClosingBracket = TEXT("}");
+			OpeningBracket = TEXT('{');
+			ClosingBracket = TEXT('}');
 			bRetestCurrentToken = true;
 		}
 	}
@@ -5849,14 +5881,14 @@ bool FHeaderParser::SkipDeclaration(FToken& Token)
 				// Not a variable name.
 				UngetToken(VariableName);
 			}
-			else if (!SafeMatchSymbol(TEXT(";")))
+			else if (!SafeMatchSymbol(TEXT(';')))
 			{
 				FError::Throwf(*FString::Printf(TEXT("Unexpected '%s'. Did you miss a semi-colon?"), VariableName.Identifier));
 			}
 		}
 
 		// C++ allows any number of ';' after member declaration/definition.
-		while (SafeMatchSymbol(TEXT(";")));
+		while (SafeMatchSymbol(TEXT(';')));
 	}
 
 	PrevComment = OldPrevComment;
@@ -5867,7 +5899,7 @@ bool FHeaderParser::SkipDeclaration(FToken& Token)
 	return NestedScopes == 0 && bEndOfDeclarationFound;
 }
 
-bool FHeaderParser::SafeMatchSymbol( const TCHAR* Match )
+bool FHeaderParser::SafeMatchSymbol( const TCHAR Match )
 {
 	FToken Token;
 
@@ -5877,7 +5909,7 @@ bool FHeaderParser::SafeMatchSymbol( const TCHAR* Match )
 
 	if (GetToken(Token, /*bNoConsts=*/ true))
 	{
-		if (Token.TokenType==TOKEN_Symbol && !FCString::Stricmp(Token.Identifier, Match))
+		if (Token.TokenType==TOKEN_Symbol && Token.Identifier[0] == Match && Token.Identifier[1] == 0)
 		{
 			return true;
 		}
@@ -5904,11 +5936,11 @@ FClass* FHeaderParser::ParseClassNameDeclaration(FClasses& AllClasses, FString&
 	bool bSpecifiesParentClass = false;
 
 	// Skip optional final keyword
-	MatchIdentifier(TEXT("final"));
+	MatchIdentifier(TEXT("final"), ESearchCase::CaseSensitive);
 
-	if (MatchSymbol(TEXT(":")))
+	if (MatchSymbol(TEXT(':')))
 	{
-		RequireIdentifier(TEXT("public"), TEXT("class inheritance"));
+		RequireIdentifier(TEXT("public"), ESearchCase::CaseSensitive, TEXT("class inheritance"));
 		bSpecifiesParentClass = true;
 	}
 
@@ -5939,9 +5971,9 @@ FClass* FHeaderParser::ParseClassNameDeclaration(FClasses& AllClasses, FString&
 		FoundClass->ClassCastFlags |= FoundClass->GetSuperClass()->ClassCastFlags;
 
 		// Handle additional inherited interface classes
-		while (MatchSymbol(TEXT(",")))
+		while (MatchSymbol(TEXT(',')))
 		{
-			RequireIdentifier(TEXT("public"), TEXT("Interface inheritance must be public"));
+			RequireIdentifier(TEXT("public"), ESearchCase::CaseSensitive, TEXT("Interface inheritance must be public"));
 
 			FToken Token;
 			if (!GetIdentifier(Token, true))
@@ -5950,7 +5982,7 @@ FClass* FHeaderParser::ParseClassNameDeclaration(FClasses& AllClasses, FString&
 			FString InterfaceName = Token.Identifier;
 
 			// Handle templated native classes
-			if (MatchSymbol(TEXT("<")))
+			if (MatchSymbol(TEXT('<')))
 			{
 				InterfaceName += TEXT('<');
 
@@ -5962,11 +5994,11 @@ FClass* FHeaderParser::ParseClassNameDeclaration(FClasses& AllClasses, FString&
 
 					if (Token.TokenType == TOKEN_Symbol)
 					{
-						if (!FCString::Strcmp(Token.Identifier, TEXT("<")))
+						if (Token.Matches(TEXT('<')))
 						{
 							++NestedScopes;
 						}
-						else if (!FCString::Strcmp(Token.Identifier, TEXT(">")))
+						else if (Token.Matches(TEXT('>')))
 						{
 							--NestedScopes;
 						}
@@ -6077,7 +6109,7 @@ UClass* FHeaderParser::CompileClassDeclaration(FClasses& AllClasses)
 	AddFormattedPrevCommentAsTooltipMetaData(MetaData);
 
 	// New style files have the class name / extends afterwards
-	RequireIdentifier(TEXT("class"), TEXT("Class declaration"));
+	RequireIdentifier(TEXT("class"), ESearchCase::CaseSensitive, TEXT("Class declaration"));
 
 	// alignas() can come before or after the deprecation macro.
 	// We can't have both, but the compiler will catch that anyway.
@@ -6133,13 +6165,13 @@ UClass* FHeaderParser::CompileClassDeclaration(FClasses& AllClasses)
 
 	// Class metadata
 	MetaData.Append(ClassDeclarationData->MetaData);
-	if (ClassDeclarationData->ClassGroupNames.Num()) { MetaData.Add("ClassGroupNames", FString::Join(ClassDeclarationData->ClassGroupNames, TEXT(" "))); }
-	if (ClassDeclarationData->AutoCollapseCategories.Num()) { MetaData.Add("AutoCollapseCategories", FString::Join(ClassDeclarationData->AutoCollapseCategories, TEXT(" "))); }
-	if (ClassDeclarationData->HideCategories.Num()) { MetaData.Add("HideCategories", FString::Join(ClassDeclarationData->HideCategories, TEXT(" "))); }
-	if (ClassDeclarationData->ShowSubCatgories.Num()) { MetaData.Add("ShowCategories", FString::Join(ClassDeclarationData->ShowSubCatgories, TEXT(" "))); }
-	if (ClassDeclarationData->SparseClassDataTypes.Num()) { MetaData.Add("SparseClassDataTypes", FString::Join(ClassDeclarationData->SparseClassDataTypes, TEXT(" "))); }
-	if (ClassDeclarationData->HideFunctions.Num()) { MetaData.Add("HideFunctions", FString::Join(ClassDeclarationData->HideFunctions, TEXT(" "))); }
-	if (ClassDeclarationData->AutoExpandCategories.Num()) { MetaData.Add("AutoExpandCategories", FString::Join(ClassDeclarationData->AutoExpandCategories, TEXT(" "))); }
+	if (ClassDeclarationData->ClassGroupNames.Num()) { MetaData.Add(NAME_ClassGroupNames, FString::Join(ClassDeclarationData->ClassGroupNames, TEXT(" "))); }
+	if (ClassDeclarationData->AutoCollapseCategories.Num()) { MetaData.Add(NAME_AutoCollapseCategories, FString::Join(ClassDeclarationData->AutoCollapseCategories, TEXT(" "))); }
+	if (ClassDeclarationData->HideCategories.Num()) { MetaData.Add(NAME_HideCategories, FString::Join(ClassDeclarationData->HideCategories, TEXT(" "))); }
+	if (ClassDeclarationData->ShowSubCatgories.Num()) { MetaData.Add(NAME_ShowCategories, FString::Join(ClassDeclarationData->ShowSubCatgories, TEXT(" "))); }
+	if (ClassDeclarationData->SparseClassDataTypes.Num()) { MetaData.Add(NAME_SparseClassDataTypes, FString::Join(ClassDeclarationData->SparseClassDataTypes, TEXT(" "))); }
+	if (ClassDeclarationData->HideFunctions.Num()) { MetaData.Add(NAME_HideFunctions, FString::Join(ClassDeclarationData->HideFunctions, TEXT(" "))); }
+	if (ClassDeclarationData->AutoExpandCategories.Num()) { MetaData.Add(NAME_AutoExpandCategories, FString::Join(ClassDeclarationData->AutoExpandCategories, TEXT(" "))); }
 
 	AddIncludePathToMetadata(Class, MetaData);
 	AddModuleRelativePathToMetadata(Class, MetaData);
@@ -6148,7 +6180,7 @@ UClass* FHeaderParser::CompileClassDeclaration(FClasses& AllClasses)
 	AddMetaDataToClassData(Class, MetaData);
 
 	// Handle the start of the rest of the class
-	RequireSymbol( TEXT("{"), TEXT("'Class'") );
+	RequireSymbol( TEXT('{'), TEXT("'Class'") );
 
 	// Make visible outside the package.
 	Class->ClearFlags(RF_Transient);
@@ -6197,13 +6229,13 @@ FClass* FHeaderParser::ParseInterfaceNameDeclaration(FClasses& AllClasses, FStri
 	}
 
 	// Get super interface
-	bool bSpecifiesParentClass = MatchSymbol(TEXT(":"));
+	bool bSpecifiesParentClass = MatchSymbol(TEXT(':'));
 	if (!bSpecifiesParentClass)
 	{
 		return FoundClass;
 	}
 
-	RequireIdentifier(TEXT("public"), TEXT("class inheritance"));
+	RequireIdentifier(TEXT("public"), ESearchCase::CaseSensitive, TEXT("class inheritance"));
 
 	// verify if our super class is an interface class
 	// the super class should have been marked as CLASS_Interface at the importing stage, if it were an interface
@@ -6231,8 +6263,6 @@ FClass* FHeaderParser::ParseInterfaceNameDeclaration(FClasses& AllClasses, FStri
 
 bool FHeaderParser::TryParseIInterfaceClass(FClasses& AllClasses)
 {
-	FString ErrorMsg(TEXT("C++ interface mix-in class declaration"));
-
 	// 'class' was already matched by the caller
 
 	// Get a class name
@@ -6243,7 +6273,7 @@ bool FHeaderParser::TryParseIInterfaceClass(FClasses& AllClasses)
 		return false;
 	}
 
-	if (MatchSymbol(TEXT(";")))
+	if (MatchSymbol(TEXT(';')))
 	{
 		// Forward declaration.
 		return false;
@@ -6261,7 +6291,7 @@ bool FHeaderParser::TryParseIInterfaceClass(FClasses& AllClasses)
 	}
 
 	// Continue parsing the second class as if it were a part of the first (for reflection data purposes, it is)
-	RequireSymbol(TEXT("{"), *ErrorMsg);
+	RequireSymbol(TEXT('{'), TEXT("C++ interface mix-in class declaration"));
 
 	// Push the interface class nesting again.
 	PushNest(ENestType::NativeInterface, FoundClass);
@@ -6292,7 +6322,7 @@ void FHeaderParser::CompileInterfaceDeclaration(FClasses& AllClasses)
 	int32 PrologFinishLine = InputLine;
 
 	// New style files have the interface name / extends afterwards
-	RequireIdentifier(TEXT("class"), TEXT("Interface declaration"));
+	RequireIdentifier(TEXT("class"), ESearchCase::CaseSensitive, TEXT("Interface declaration"));
 	FClass* InterfaceClass = ParseInterfaceNameDeclaration(AllClasses, /*out*/ DeclaredInterfaceName, /*out*/ RequiredAPIMacroIfPresent);
 	ClassDefinitionRanges.Add(InterfaceClass, ClassDefinitionRange(&Input[InputPos], nullptr));
 
@@ -6343,7 +6373,7 @@ void FHeaderParser::CompileInterfaceDeclaration(FClasses& AllClasses)
 
 			case EInterfaceSpecifier::ConversionRoot:
 			{
-				MetaData.Add(FName(TEXT("IsConversionRoot")), "true");
+				MetaData.Add(NAME_IsConversionRoot, TEXT("true"));
 			}
 			break;
 		}
@@ -6367,7 +6397,7 @@ void FHeaderParser::CompileInterfaceDeclaration(FClasses& AllClasses)
 	AddMetaDataToClassData(InterfaceClass, MetaData);
 
 	// Handle the start of the rest of the interface
-	RequireSymbol( TEXT("{"), TEXT("'Class'") );
+	RequireSymbol( TEXT('{'), TEXT("'Class'") );
 
 	// Make visible outside the package.
 	InterfaceClass->ClearFlags(RF_Transient);
@@ -6500,6 +6530,16 @@ void FHeaderParser::CompileRigVMMethodDeclaration(FClasses& AllClasses, UStruct*
 	StructRigVMInfo.Methods.Add(MethodInfo);
 }
 
+static const FName NAME_InputText(TEXT("Input"));
+static const FName NAME_OutputText(TEXT("Output"));
+static const FName NAME_ConstantText(TEXT("Constant"));
+static const FName NAME_MaxArraySizeText(TEXT("MaxArraySize"));
+
+static const TCHAR* TArrayText = TEXT("TArray");
+static const TCHAR* TArrayViewText = TEXT("TArrayView");
+static const TCHAR* GetRefText = TEXT("GetRef");
+static const TCHAR* GetArrayText = TEXT("GetArray");
+
 void FHeaderParser::ParseRigVMMethodParameters(UStruct* Struct)
 {
 	FRigVMStructInfo* StructRigVMInfo = StructRigVMMap.Find(Struct);
@@ -6508,20 +6548,10 @@ void FHeaderParser::ParseRigVMMethodParameters(UStruct* Struct)
 		return;
 	}
 
-	const TCHAR* InputText = TEXT("Input");
-	const TCHAR* OutputText = TEXT("Output");
-	const TCHAR* ConstantText = TEXT("Constant");
-	const TCHAR* MaxArraySizeText = TEXT("MaxArraySize");
-	const TCHAR* TArrayText = TEXT("TArray");
-	const TCHAR* TArrayViewText = TEXT("TArrayView");
-	const TCHAR* GetRefText = TEXT("GetRef");
-	const TCHAR* GetArrayText = TEXT("GetArray");
-
 	// validate the property types for this struct
 	for (TFieldIterator<FProperty> It(Struct); It; ++It)
 	{
 		FProperty const* const Prop = *It;
-		FString PropName = Prop->GetName();
 		FString MemberCPPType;
 		FString ExtendedCPPType;
 		MemberCPPType = Prop->GetCPPType(&ExtendedCPPType);
@@ -6529,10 +6559,10 @@ void FHeaderParser::ParseRigVMMethodParameters(UStruct* Struct)
 		FRigVMParameter Parameter;
 		Parameter.Name = Prop->GetName();
 		Parameter.Type = MemberCPPType + ExtendedCPPType;
-		Parameter.bConstant = Prop->HasMetaData(ConstantText);
-		Parameter.bInput = Prop->HasMetaData(InputText);
-		Parameter.bOutput = Prop->HasMetaData(OutputText);
-		Parameter.MaxArraySize = Prop->GetMetaData(MaxArraySizeText);
+		Parameter.bConstant = Prop->HasMetaData(NAME_ConstantText);
+		Parameter.bInput = Prop->HasMetaData(NAME_InputText);
+		Parameter.bOutput = Prop->HasMetaData(NAME_OutputText);
+		Parameter.MaxArraySize = Prop->GetMetaData(NAME_MaxArraySizeText);
 		Parameter.Getter = GetRefText;
 		Parameter.bEditorOnly = Prop->IsEditorOnlyProperty();
 
@@ -6557,7 +6587,7 @@ void FHeaderParser::ParseRigVMMethodParameters(UStruct* Struct)
 			}
 		}
 
-		if (MemberCPPType.StartsWith(TArrayText))
+		if (MemberCPPType.StartsWith(TArrayText, ESearchCase::CaseSensitive))
 		{
 			if (Parameter.IsConst() || !Parameter.MaxArraySize.IsEmpty())
 			{
@@ -6567,7 +6597,7 @@ void FHeaderParser::ParseRigVMMethodParameters(UStruct* Struct)
 			}
 		}
 
-		StructRigVMInfo->Members.Add(Parameter);
+		StructRigVMInfo->Members.Add(MoveTemp(Parameter));
 	}
 
 	if (StructRigVMInfo->Members.Num() > 64)
@@ -6579,8 +6609,7 @@ void FHeaderParser::ParseRigVMMethodParameters(UStruct* Struct)
 // Returns true if the token is a dynamic delegate declaration
 bool FHeaderParser::IsValidDelegateDeclaration(const FToken& Token) const
 {
-	FString TokenStr(Token.Identifier);
-	return (Token.TokenType == TOKEN_Identifier) && TokenStr.StartsWith(TEXT("DECLARE_DYNAMIC_"));
+	return (Token.TokenType == TOKEN_Identifier) && !FCString::Strncmp(Token.Identifier, TEXT("DECLARE_DYNAMIC_"), 16);
 }
 
 // Modify token to fix redirected types if needed
@@ -6599,7 +6628,7 @@ void FHeaderParser::RedirectTypeIdentifier(FToken& Token) const
 void FHeaderParser::ParseParameterList(FClasses& AllClasses, UFunction* Function, bool bExpectCommaBeforeName, TMap<FName, FString>* MetaData)
 {
 	// Get parameter list.
-	if (MatchSymbol(TEXT(")")))
+	if (MatchSymbol(TEXT(')')))
 	{
 		return;
 	}
@@ -6615,7 +6644,7 @@ void FHeaderParser::ParseParameterList(FClasses& AllClasses, UFunction* Function
 
 		if (bExpectCommaBeforeName)
 		{
-			RequireSymbol(TEXT(","), TEXT("Delegate definitions require a , between the parameter type and parameter name"));
+			RequireSymbol(TEXT(','), TEXT("Delegate definitions require a , between the parameter type and parameter name"));
 		}
 
 		FProperty* Prop = GetVarNameAndDim(Function, Property, VariableCategory);
@@ -6686,7 +6715,7 @@ void FHeaderParser::ParseParameterList(FClasses& AllClasses, UFunction* Function
 		}
 
 		// Default value.
-		if (MatchSymbol( TEXT("=") ))
+		if (MatchSymbol(TEXT('=')))
 		{
 			// Skip past the native specified default value; we make no attempt to parse it
 			FToken SkipToken;
@@ -6700,18 +6729,18 @@ void FHeaderParser::ParseParameterList(FClasses& AllClasses, UFunction* Function
 					StartPos = SkipToken.StartPos;
 				}
 				if ( ParenthesisNestCount == 0
-					&& (SkipToken.Matches(TEXT(")")) || SkipToken.Matches(TEXT(","))) )
+					&& (SkipToken.Matches(TEXT(')')) || SkipToken.Matches(TEXT(','))) )
 				{
 					EndPos = SkipToken.StartPos;
 					// went too far
 					UngetToken(SkipToken);
 					break;
 				}
-				if ( SkipToken.Matches(TEXT("(")) )
+				if ( SkipToken.Matches(TEXT('(')) )
 				{
 					ParenthesisNestCount++;
 				}
-				else if ( SkipToken.Matches(TEXT(")")) )
+				else if ( SkipToken.Matches(TEXT(')')) )
 				{
 					ParenthesisNestCount--;
 				}
@@ -6740,8 +6769,8 @@ void FHeaderParser::ParseParameterList(FClasses& AllClasses, UFunction* Function
 				}
 			}
 		}
-	} while( MatchSymbol(TEXT(",")) );
-	RequireSymbol( TEXT(")"), TEXT("parameter list") );
+	} while( MatchSymbol(TEXT(',')) );
+	RequireSymbol( TEXT(')'), TEXT("parameter list") );
 }
 UDelegateFunction* FHeaderParser::CompileDelegateDeclaration(FClasses& AllClasses, const TCHAR* DelegateIdentifier, EDelegateSpecifierAction::Type SpecifierAction)
 {
@@ -6784,10 +6813,10 @@ UDelegateFunction* FHeaderParser::CompileDelegateDeclaration(FClasses& AllClasse
 	}
 
 	// Break the delegate declaration macro down into parts
-	const bool bHasReturnValue = DelegateMacro.Contains(TEXT("_RetVal"));
-	const bool bDeclaredConst  = DelegateMacro.Contains(TEXT("_Const"));
-	const bool bIsMulticast    = DelegateMacro.Contains(TEXT("_MULTICAST"));
-	const bool bIsSparse       = DelegateMacro.Contains(TEXT("_SPARSE"));
+	const bool bHasReturnValue = DelegateMacro.Contains(TEXT("_RetVal"), ESearchCase::CaseSensitive);
+	const bool bDeclaredConst  = DelegateMacro.Contains(TEXT("_Const"), ESearchCase::CaseSensitive);
+	const bool bIsMulticast    = DelegateMacro.Contains(TEXT("_MULTICAST"), ESearchCase::CaseSensitive);
+	const bool bIsSparse       = DelegateMacro.Contains(TEXT("_SPARSE"), ESearchCase::CaseSensitive);
 
 	// Determine the parameter count
 	const FString* FoundParamCount = DelegateParameterCountStrings.FindByPredicate([&](const FString& Str){ return DelegateMacro.Contains(Str); });
@@ -6820,7 +6849,7 @@ UDelegateFunction* FHeaderParser::CompileDelegateDeclaration(FClasses& AllClasse
 	}
 
 	// Now parse the macro body
-	RequireSymbol(TEXT("("), CurrentScopeName);
+	RequireSymbol(TEXT('('), CurrentScopeName);
 
 	// Parse the return value type
 	FToken ReturnType( CPT_None );
@@ -6828,7 +6857,7 @@ UDelegateFunction* FHeaderParser::CompileDelegateDeclaration(FClasses& AllClasse
 	if (bHasReturnValue)
 	{
 		GetVarType(AllClasses, GetCurrentScope(), ReturnType, CPF_None, nullptr, EPropertyDeclarationStyle::None, EVariableCategory::Return);
-		RequireSymbol(TEXT(","), CurrentScopeName);
+		RequireSymbol(TEXT(','), CurrentScopeName);
 	}
 
 	// Skip whitespaces to get InputPos exactly on beginning of function name.
@@ -6850,12 +6879,12 @@ UDelegateFunction* FHeaderParser::CompileDelegateDeclaration(FClasses& AllClasse
 		// Remove the leading F
 		FString Name(FuncInfo.Function.Identifier);
 
-		if (!Name.StartsWith(TEXT("F")))
+		if (!Name.StartsWith(TEXT("F"), ESearchCase::CaseSensitive))
 		{
 			FError::Throwf(TEXT("Delegate type declarations must start with F"));
 		}
 
-		Name = Name.Mid(1);
+		Name.RightChopInline(1, false);
 
 		// Append the signature goo
 		Name += HEADER_GENERATED_DELEGATE_SIGNATURE_SUFFIX;
@@ -6889,13 +6918,13 @@ UDelegateFunction* FHeaderParser::CompileDelegateDeclaration(FClasses& AllClasse
 	{
 		FToken OwningClass;
 
-		RequireSymbol(TEXT(","), TEXT("Delegate Declaration"));
+		RequireSymbol(TEXT(','), TEXT("Delegate Declaration"));
 
 		if (!GetIdentifier(OwningClass))
 		{
 			FError::Throwf(TEXT("Missing OwningClass specifier."));
 		}
-		RequireSymbol(TEXT(","), TEXT("Delegate Declaration"));
+		RequireSymbol(TEXT(','), TEXT("Delegate Declaration"));
 		
 		FToken DelegateName;
 		if (!GetIdentifier(DelegateName))
@@ -6904,14 +6933,14 @@ UDelegateFunction* FHeaderParser::CompileDelegateDeclaration(FClasses& AllClasse
 		}
 
 		USparseDelegateFunction* SDF = CastChecked<USparseDelegateFunction>(DelegateSignatureFunction);
-		SDF->OwningClassName = *GetClassNameWithoutPrefix(OwningClass.TokenName.ToString());
+		SDF->OwningClassName = *GetClassNameWithoutPrefix(OwningClass.Identifier);
 		SDF->DelegateName = DelegateName.Identifier;
 	}
 
 	// Get parameter list.
 	if (FoundParamCount)
 	{
-		RequireSymbol(TEXT(","), CurrentScopeName);
+		RequireSymbol(TEXT(','), CurrentScopeName);
 
 		ParseParameterList(AllClasses, DelegateSignatureFunction, /*bExpectCommaBeforeName=*/ true);
 
@@ -6925,7 +6954,7 @@ UDelegateFunction* FHeaderParser::CompileDelegateDeclaration(FClasses& AllClasse
 	else
 	{
 		// Require the closing paren even with no parameter list
-		RequireSymbol(TEXT(")"), TEXT("Delegate Declaration"));
+		RequireSymbol(TEXT(')'), TEXT("Delegate Declaration"));
 	}
 
 	FuncInfo.MacroLine = InputLine;
@@ -7104,13 +7133,13 @@ void FHeaderParser::CompileFunctionDeclaration(FClasses& AllClasses)
 		FuncInfo.FunctionFlags |= FUNC_Const;
 	}
 
-	if (MatchIdentifier(TEXT("static")))
+	if (MatchIdentifier(TEXT("static"), ESearchCase::CaseSensitive))
 	{
 		FuncInfo.FunctionFlags |= FUNC_Static;
 		FuncInfo.FunctionExportFlags |= FUNCEXPORT_CppStatic;
 	}
 
-	if (MetaData.Contains("CppFromBpEvent"))
+	if (MetaData.Contains(NAME_CppFromBpEvent))
 	{
 		FuncInfo.FunctionFlags |= FUNC_Event;
 	}
@@ -7124,9 +7153,9 @@ void FHeaderParser::CompileFunctionDeclaration(FClasses& AllClasses)
 
 	const bool bClassGeneratedFromBP = FClass::IsDynamic(GetCurrentClass());
 
-	if ((0 != (FuncInfo.FunctionExportFlags & FUNCEXPORT_CustomThunk)) && !MetaData.Contains("CustomThunk"))
+	if ((0 != (FuncInfo.FunctionExportFlags & FUNCEXPORT_CustomThunk)) && !MetaData.Contains(NAME_CustomThunk))
 	{
-		MetaData.Add(TEXT("CustomThunk"), TEXT("true"));
+		MetaData.Add(NAME_CustomThunk, TEXT("true"));
 	}
 
 	if ((FuncInfo.FunctionFlags & FUNC_BlueprintPure) && GetCurrentClass()->HasAnyClassFlags(CLASS_Interface))
@@ -7150,20 +7179,20 @@ void FHeaderParser::CompileFunctionDeclaration(FClasses& AllClasses)
 
 	bool bSawVirtual = false;
 
-	if (MatchIdentifier(TEXT("virtual")))
+	if (MatchIdentifier(TEXT("virtual"), ESearchCase::CaseSensitive))
 	{
 		bSawVirtual = true;
 	}
 
-	FString*   InternalPtr = MetaData.Find("BlueprintInternalUseOnly"); // FBlueprintMetadata::MD_BlueprintInternalUseOnly
+	FString*   InternalPtr = MetaData.Find(NAME_BlueprintInternalUseOnly); // FBlueprintMetadata::MD_BlueprintInternalUseOnly
 	const bool bInternalOnly = InternalPtr && *InternalPtr == TEXT("true");
 
 	// If this function is blueprint callable or blueprint pure, require a category 
 	if ((FuncInfo.FunctionFlags & (FUNC_BlueprintCallable | FUNC_BlueprintPure)) != 0) 
 	{ 
-		const bool bDeprecated = MetaData.Contains("DeprecatedFunction");       // FBlueprintMetadata::MD_DeprecatedFunction
-		const bool bBlueprintAccessor = MetaData.Contains("BlueprintSetter") || MetaData.Contains("BlueprintGetter"); // FBlueprintMetadata::MD_BlueprintSetter, // FBlueprintMetadata::MD_BlueprintGetter
-		const bool bHasMenuCategory = MetaData.Contains("Category");                 // FBlueprintMetadata::MD_FunctionCategory
+		const bool bDeprecated = MetaData.Contains(NAME_DeprecatedFunction);       // FBlueprintMetadata::MD_DeprecatedFunction
+		const bool bBlueprintAccessor = MetaData.Contains(NAME_BlueprintSetter) || MetaData.Contains(NAME_BlueprintGetter); // FBlueprintMetadata::MD_BlueprintSetter, // FBlueprintMetadata::MD_BlueprintGetter
+		const bool bHasMenuCategory = MetaData.Contains(NAME_Category);                 // FBlueprintMetadata::MD_FunctionCategory
 
 		if (!bHasMenuCategory && !bInternalOnly && !bDeprecated && !bBlueprintAccessor) 
 		{ 
@@ -7178,11 +7207,12 @@ void FHeaderParser::CompileFunctionDeclaration(FClasses& AllClasses)
 	// Verify interfaces with respect to their blueprint accessible functions
 	if (GetCurrentClass()->HasAnyClassFlags(CLASS_Interface))
 	{
-		const bool bCanImplementInBlueprints = !GetCurrentClass()->HasMetaData(TEXT("CannotImplementInterfaceInBlueprint"));  //FBlueprintMetadata::MD_CannotImplementInterfaceInBlueprint
-		if((FuncInfo.FunctionFlags & FUNC_BlueprintEvent) != 0)
+		if((FuncInfo.FunctionFlags & FUNC_BlueprintEvent) != 0 && !bInternalOnly)
 		{
+			const bool bCanImplementInBlueprints = !GetCurrentClass()->HasMetaData(NAME_CannotImplementInterfaceInBlueprint);  //FBlueprintMetadata::MD_CannotImplementInterfaceInBlueprint
+
 			// Ensure that blueprint events are only allowed in implementable interfaces. Internal only functions allowed
-			if (!bCanImplementInBlueprints && !bInternalOnly)
+			if (!bCanImplementInBlueprints)
 			{
 				UE_LOG_ERROR_UHT(TEXT("Interfaces that are not implementable in blueprints cannot have BlueprintImplementableEvent members."));
 			}
@@ -7190,6 +7220,8 @@ void FHeaderParser::CompileFunctionDeclaration(FClasses& AllClasses)
 		
 		if (((FuncInfo.FunctionFlags & FUNC_BlueprintCallable) != 0) && (((~FuncInfo.FunctionFlags) & FUNC_BlueprintEvent) != 0))
 		{
+			const bool bCanImplementInBlueprints = !GetCurrentClass()->HasMetaData(NAME_CannotImplementInterfaceInBlueprint);  //FBlueprintMetadata::MD_CannotImplementInterfaceInBlueprint
+
 			// Ensure that if this interface contains blueprint callable functions that are not blueprint defined, that it must be implemented natively
 			if (bCanImplementInBlueprints)
 			{
@@ -7208,7 +7240,7 @@ void FHeaderParser::CompileFunctionDeclaration(FClasses& AllClasses)
 			if (Token.TokenType == TOKEN_Identifier)
 			{
 				FString RequiredAPIMacroIfPresent(Token.Identifier);
-				if (RequiredAPIMacroIfPresent.EndsWith(TEXT("_API")))
+				if (RequiredAPIMacroIfPresent.EndsWith(TEXT("_API"), ESearchCase::CaseSensitive))
 				{
 					//@TODO: Validate the module name for RequiredAPIMacroIfPresent
 					bThrowTokenBack = false;
@@ -7232,13 +7264,13 @@ void FHeaderParser::CompileFunctionDeclaration(FClasses& AllClasses)
 	}
 
 	// Look for static again, in case there was an ENGINE_API token first
-	if (!APIMacroIfPresent.IsEmpty() && MatchIdentifier(TEXT("static")))
+	if (!APIMacroIfPresent.IsEmpty() && MatchIdentifier(TEXT("static"), ESearchCase::CaseSensitive))
 	{
 		FError::Throwf(TEXT("Unexpected API macro '%s'. Did you mean to put '%s' after the static keyword?"), *APIMacroIfPresent, *APIMacroIfPresent);
 	}
 
 	// Look for virtual again, in case there was an ENGINE_API token first
-	if (MatchIdentifier(TEXT("virtual")))
+	if (MatchIdentifier(TEXT("virtual"), ESearchCase::CaseSensitive))
 	{
 		bSawVirtual = true;
 	}
@@ -7295,7 +7327,7 @@ void FHeaderParser::CompileFunctionDeclaration(FClasses& AllClasses)
 	FToken ReturnType( CPT_None );
 
 	// C++ style functions always have a return value type, even if it's void
-	bool bHasReturnValue = !MatchIdentifier(TEXT("void"));
+	bool bHasReturnValue = !MatchIdentifier(TEXT("void"), ESearchCase::CaseSensitive);
 	if (bHasReturnValue)
 	{
 		GetVarType(AllClasses, GetCurrentScope(), ReturnType, CPF_None, nullptr, EPropertyDeclarationStyle::None, EVariableCategory::Return);
@@ -7312,7 +7344,7 @@ void FHeaderParser::CompileFunctionDeclaration(FClasses& AllClasses)
 		FError::Throwf(TEXT("Missing %s name"), TypeOfFunction);
 	}
 
-	if ( !MatchSymbol(TEXT("(")) )
+	if ( !MatchSymbol(TEXT('(')) )
 	{
 		FError::Throwf(TEXT("Bad %s definition"), TypeOfFunction);
 	}
@@ -7427,7 +7459,7 @@ void FHeaderParser::CompileFunctionDeclaration(FClasses& AllClasses)
 
 
 	// determine whether this function should be 'const'
-	if ( MatchIdentifier(TEXT("const")) )
+	if ( MatchIdentifier(TEXT("const"), ESearchCase::CaseSensitive) )
 	{
 		if( (TopFunction->FunctionFlags & (FUNC_Native)) == 0 )
 		{
@@ -7457,15 +7489,15 @@ void FHeaderParser::CompileFunctionDeclaration(FClasses& AllClasses)
 	AddMetaDataToClassData(TopFunction, MetaData);
 
 	// 'final' and 'override' can appear in any order before an optional '= 0' pure virtual specifier
-	bool bFoundFinal    = MatchIdentifier(TEXT("final"));
-	bool bFoundOverride = MatchIdentifier(TEXT("override"));
+	bool bFoundFinal    = MatchIdentifier(TEXT("final"), ESearchCase::CaseSensitive);
+	bool bFoundOverride = MatchIdentifier(TEXT("override"), ESearchCase::CaseSensitive);
 	if (!bFoundFinal && bFoundOverride)
 	{
-		bFoundFinal = MatchIdentifier(TEXT("final"));
+		bFoundFinal = MatchIdentifier(TEXT("final"), ESearchCase::CaseSensitive);
 	}
 
 	// Handle C++ style functions being declared as abstract
-	if (MatchSymbol(TEXT("=")))
+	if (MatchSymbol(TEXT('=')))
 	{
 		int32 ZeroValue = 1;
 		bool bGotZero = GetConstInt(/*out*/ZeroValue);
@@ -7548,11 +7580,11 @@ void FHeaderParser::CompileFunctionDeclaration(FClasses& AllClasses)
 
 	// Optionally consume a semicolon
 	// This is optional to allow inline function definitions
-	if (Token.TokenType == TOKEN_Symbol && !FCString::Stricmp(Token.Identifier, TEXT(";")))
+	if (Token.TokenType == TOKEN_Symbol && Token.Matches(TEXT(';')))
 	{
 		// Do nothing (consume it)
 	}
-	else if (Token.TokenType == TOKEN_Symbol && !FCString::Stricmp(Token.Identifier, TEXT("{")))
+	else if (Token.TokenType == TOKEN_Symbol && Token.Matches(TEXT('{')))
 	{
 		// Skip inline function bodies
 		UngetToken(Token);
@@ -7573,15 +7605,17 @@ void FHeaderParser::ParseFieldMetaData(TMap<FName, FString>& MetaData, const TCH
 {
 	FToken PropertyMetaData;
 	bool bMetadataPresent = false;
-	 if (MatchIdentifier(TEXT("UMETA")))
+	if (MatchIdentifier(TEXT("UMETA"), ESearchCase::CaseSensitive))
 	{
+		auto ErrorMessageGetter = [FieldName]() { return FString::Printf(TEXT("' %s metadata'"), FieldName); };
+
 		bMetadataPresent = true;
-		RequireSymbol( TEXT("("),*FString::Printf(TEXT("' %s metadata'"), FieldName) );
+		RequireSymbol( TEXT('('), ErrorMessageGetter );
 		if (!GetRawTokenRespectingQuotes(PropertyMetaData, TCHAR(')')))
 		{
 			FError::Throwf(TEXT("'%s': No metadata specified"), FieldName);
 		}
-		RequireSymbol( TEXT(")"),*FString::Printf(TEXT("' %s metadata'"), FieldName) );
+		RequireSymbol( TEXT(')'), ErrorMessageGetter);
 	}
 
 	if (bMetadataPresent)
@@ -7615,28 +7649,33 @@ void FHeaderParser::ParseFieldMetaData(TMap<FName, FString>& MetaData, const TCH
 
 		if (LastStartIndex != CharIndex)
 		{
-			Pairs.Add(NewData.Mid(LastStartIndex, CharIndex - LastStartIndex));
+			Pairs.Add(MoveTemp(NewData).Mid(LastStartIndex, CharIndex - LastStartIndex));
 		}
 
 		// go over all pairs
 		for (int32 PairIndex = 0; PairIndex < Pairs.Num(); PairIndex++)
 		{
 			// break the pair into a key and a value
-			FString Token = Pairs[PairIndex];
-			FString Key = Token;
+			FString Token = MoveTemp(Pairs[PairIndex]);
+			FString Key;
 			// by default, not value, just a key (allowed)
 			FString Value;
 
 			// look for a value after an =
-			int32 Equals = Token.Find(TEXT("="));
+			const int32 Equals = Token.Find(TEXT("="), ESearchCase::CaseSensitive);
 			// if we have an =, break up the string
-			if (Equals != -1)
+			if (Equals != INDEX_NONE)
 			{
 				Key = Token.Left(Equals);
-				Value = Token.Right((Token.Len() - Equals) - 1);
+				Value = MoveTemp(Token);
+				Value.RightInline((Value.Len() - Equals) - 1, false);
+			}
+			else
+			{
+				Key = MoveTemp(Token);
 			}
 
-			InsertMetaDataPair(MetaData, Key, Value);
+			InsertMetaDataPair(MetaData, MoveTemp(Key), MoveTemp(Value));
 		}
 	}
 }
@@ -7653,7 +7692,7 @@ bool FHeaderParser::IsBitfieldProperty()
 		FToken Token;
 		if (GetToken(Token, /*bNoConsts=*/ true))
 		{
-			if (Token.TokenType == TOKEN_Symbol && FCString::Stricmp(Token.Identifier, TEXT(":")) == 0)
+			if (Token.TokenType == TOKEN_Symbol && Token.Matches(TEXT(':')))
 			{
 				bIsBitfield = true;
 			}
@@ -7728,15 +7767,14 @@ void FHeaderParser::CompileVariableDeclaration(FClasses& AllClasses, UStruct* St
 	GetVarType( AllClasses, &FScope::GetTypeScope(Struct).Get(), OriginalProperty, DisallowFlags, /*OuterPropertyType=*/ NULL, EPropertyDeclarationStyle::UPROPERTY, EVariableCategory::Member, &TypeRange );
 	OriginalProperty.PropertyFlags |= EdFlags;
 
-	FString* Category = OriginalProperty.MetaData.Find("Category");
+	FString* Category = OriginalProperty.MetaData.Find(NAME_Category);
 
 	// First check if the category was specified at all and if the property was exposed to the editor.
 	if (!Category && (OriginalProperty.PropertyFlags & (CPF_Edit|CPF_BlueprintVisible)))
 	{
 		if ((Struct->GetOutermost() != nullptr) && !bIsCurrentModulePartOfEngine)
 		{
-			OriginalProperty.MetaData.Add("Category", Struct->GetFName().ToString());
-			Category = OriginalProperty.MetaData.Find("Category");
+			Category = &OriginalProperty.MetaData.Add(NAME_Category, Struct->GetName());
 		}
 		else
 		{
@@ -7778,7 +7816,7 @@ void FHeaderParser::CompileVariableDeclaration(FClasses& AllClasses, UStruct* St
 		FError::Throwf(TEXT("Illegal type modifiers in member variable declaration") );
 	}
 
-	if (FString* ExposeOnSpawnValue = OriginalProperty.MetaData.Find(TEXT("ExposeOnSpawn")))
+	if (FString* ExposeOnSpawnValue = OriginalProperty.MetaData.Find(NAME_ExposeOnSpawn))
 	{
 		if ((*ExposeOnSpawnValue == TEXT("true")) && !FExposeOnSpawnValidator::IsSupported(OriginalProperty))
 		{
@@ -7794,7 +7832,7 @@ void FHeaderParser::CompileVariableDeclaration(FClasses& AllClasses, UStruct* St
 		FProperty* NewProperty = GetVarNameAndDim(Struct, Property, EVariableCategory::Member);
 
 		// Optionally consume the :1 at the end of a bitfield boolean declaration
-		if (Property.IsBool() && MatchSymbol(TEXT(":")))
+		if (Property.IsBool() && MatchSymbol(TEXT(':')))
 		{
 			int32 BitfieldSize = 0;
 			if (!GetConstInt(/*out*/ BitfieldSize) || (BitfieldSize != 1))
@@ -7851,16 +7889,16 @@ void FHeaderParser::CompileVariableDeclaration(FClasses& AllClasses, UStruct* St
 			}
 		}
 
-	} while( MatchSymbol(TEXT(",")) );
+	} while( MatchSymbol(TEXT(',')) );
 
 	// Optional member initializer.
-	if (MatchSymbol(TEXT("=")))
+	if (MatchSymbol(TEXT('=')))
 	{
 		// Skip past the specified member initializer; we make no attempt to parse it
 		FToken SkipToken;
 		while (GetToken(SkipToken))
 		{
-			if (SkipToken.Matches(TEXT(";")))
+			if (SkipToken.Matches(TEXT(';')))
 			{
 				// went too far
 				UngetToken(SkipToken);
@@ -7869,17 +7907,17 @@ void FHeaderParser::CompileVariableDeclaration(FClasses& AllClasses, UStruct* St
 		}
 	}
 	// Using Brace Initialization
-	else if (MatchSymbol(TEXT("{")))
+	else if (MatchSymbol(TEXT('{')))
 	{
 		FToken SkipToken;
 		int BraceLevel = 1;
 		while (GetToken(SkipToken))
 		{
-			if (SkipToken.Matches(TEXT("{")))
+			if (SkipToken.Matches(TEXT('{')))
 			{
 				++BraceLevel;
 			}
-			else if (SkipToken.Matches(TEXT("}")))
+			else if (SkipToken.Matches(TEXT('}')))
 			{
 				--BraceLevel;
 				if (BraceLevel == 0)
@@ -7891,7 +7929,7 @@ void FHeaderParser::CompileVariableDeclaration(FClasses& AllClasses, UStruct* St
 	}
 
 	// Expect a semicolon.
-	RequireSymbol( TEXT(";"), TEXT("'variable declaration'") );
+	RequireSymbol( TEXT(';'), TEXT("'variable declaration'") );
 
 	// Skip redundant semi-colons
 	for (;;)
@@ -7905,7 +7943,7 @@ void FHeaderParser::CompileVariableDeclaration(FClasses& AllClasses, UStruct* St
 			break;
 		}
 
-		if (Token.TokenType != TOKEN_Symbol || FCString::Stricmp(Token.Identifier, TEXT(";")))
+		if (Token.TokenType != TOKEN_Symbol || !Token.Matches(TEXT(';')))
 		{
 			InputPos  = CurrInputPos;
 			InputLine = CurrInputLine;
@@ -8003,15 +8041,15 @@ void FHeaderParser::SkipStatements( int32 NestCount, const TCHAR* ErrorTag  )
 
 	while( GetToken( Token, true ) )
 	{
-		if ( Token.Matches(TEXT("{")) )
+		if ( Token.Matches(TEXT('{')) )
 		{
 			NestCount++;
 		}
-		else if	( Token.Matches(TEXT("}")) )
+		else if	( Token.Matches(TEXT('}')) )
 		{
 			NestCount--;
 		}
-		else if ( Token.Matches(TEXT(";")) && OriginalNestCount == 0 )
+		else if ( Token.Matches(TEXT(';')) && OriginalNestCount == 0 )
 		{
 			break;
 		}
@@ -8260,6 +8298,8 @@ ECompilationResult::Type FHeaderParser::ParseRestOfModulesSourceFiles(FClasses&
 }
 
 // Parse Class's annotated headers and optionally its child classes.
+static const FString ObjectHeader(TEXT("NoExportTypes.h"));
+
 ECompilationResult::Type FHeaderParser::ParseHeaders(FClasses& AllClasses, FHeaderParser& HeaderParser, FUnrealSourceFile* SourceFile)
 {
 	ECompilationResult::Type Result = ECompilationResult::Succeeded;
@@ -8273,7 +8313,6 @@ ECompilationResult::Type FHeaderParser::ParseHeaders(FClasses& AllClasses, FHead
 
 	TArray<FUnrealSourceFile*> SourceFilesRequired;
 
-	static const FString ObjectHeader = FString(TEXT("NoExportTypes.h"));
 	for (FHeaderProvider& Include : SourceFile->GetIncludes())
 	{
 		if (Include.GetId() == ObjectHeader)
@@ -8332,11 +8371,11 @@ ECompilationResult::Type FHeaderParser::ParseHeaders(FClasses& AllClasses, FHead
 bool FHeaderParser::DependentClassNameFromHeader(const TCHAR* HeaderFilename, FString& OutClassName)
 {
 	FString DependentClassName(HeaderFilename);
-	const int32 ExtensionIndex = DependentClassName.Find(TEXT("."));
+	const int32 ExtensionIndex = DependentClassName.Find(TEXT("."), ESearchCase::CaseSensitive);
 	if (ExtensionIndex != INDEX_NONE)
 	{
 		// Generate UHeaderName name for this header.
-		OutClassName = FString(TEXT("U")) + FPaths::GetBaseFilename(*DependentClassName);
+		OutClassName = TEXT("U") + FPaths::GetBaseFilename(MoveTemp(DependentClassName));
 		return true;
 	}
 	return false;
@@ -8398,10 +8437,6 @@ void FHeaderParser::ExportNativeHeaders(
 	const FManifestModule& Module
 )
 {
-	// Build a list of header filenames
-	TArray<FString>	ClassHeaderFilenames;
-	new (ClassHeaderFilenames) FString();
-
 	TSet<FUnrealSourceFile*> SourceFiles = GetSourceFilesWithInheritanceOrdering(CurrentPackage, AllClasses);
 	if (SourceFiles.Num() > 0)
 	{
@@ -8487,7 +8522,7 @@ FHeaderParser::FHeaderParser(FFeedbackContext* InWarn, const FManifestModule& In
 					FParse::Value(*It.Value().GetValue(), TEXT("OldType="), OldType);
 					FParse::Value(*It.Value().GetValue(), TEXT("NewType="), NewType);
 
-					TypeRedirectMap.Add(OldType, NewType);
+					TypeRedirectMap.Add(MoveTemp(OldType), MoveTemp(NewType));
 				}
 				else if (It.Key() == StructsWithNoPrefixKey)
 				{
@@ -8687,7 +8722,7 @@ ECompilationResult::Type FHeaderParser::ParseAllHeadersInside(
  * @param InNameToCheck - Name w/ potential prefix to check
  * @param OriginalClassName - Name of class w/ no prefix to check against
  */
-bool FHeaderParser::ClassNameHasValidPrefix(const FString InNameToCheck, const FString OriginalClassName)
+bool FHeaderParser::ClassNameHasValidPrefix(const FString& InNameToCheck, const FString& OriginalClassName)
 {
 	bool bIsLabledDeprecated;
 	const FString ClassPrefix = GetClassPrefix( InNameToCheck, bIsLabledDeprecated );
@@ -9173,10 +9208,10 @@ void FHeaderParser::SimplifiedClassParse(const TCHAR* Filename, const TCHAR* InB
 					{
 						if (StrBegin == INDEX_NONE || Pos < StrBegin || Pos > StrEnd)
 						{
-							StrLine = StrLine.Left(Pos);
+							StrLine.LeftInline(Pos, false);
 						}
 
-						if (StrLine == TEXT(""))
+						if (StrLine.IsEmpty())
 						{
 							continue;
 						}
@@ -9196,7 +9231,7 @@ void FHeaderParser::SimplifiedClassParse(const TCHAR* Filename, const TCHAR* InB
 							}
 							else
 							{
-								StrLine = StrLine.Left(Pos);
+								StrLine.LeftInline(Pos, false);
 								CommentDim++;
 							}
 						}
@@ -9207,7 +9242,7 @@ void FHeaderParser::SimplifiedClassParse(const TCHAR* Filename, const TCHAR* InB
 					{
 						if (StrBegin == INDEX_NONE || EndPos < StrBegin || EndPos > StrEnd)
 						{
-							StrLine = StrLine.Mid(EndPos + 2);
+							StrLine.MidInline(EndPos + 2, MAX_int32, false);
 							CommentDim--;
 						}
 
@@ -9217,7 +9252,7 @@ void FHeaderParser::SimplifiedClassParse(const TCHAR* Filename, const TCHAR* InB
 			}
 
 			StrLine.TrimStartInline();
-			if (!bProcess || StrLine == TEXT(""))
+			if (!bProcess || StrLine.IsEmpty())
 			{
 				continue;
 			}
@@ -9284,12 +9319,12 @@ void FHeaderParser::SimplifiedClassParse(const TCHAR* Filename, const TCHAR* InB
 
 void FHeaderPreParser::ParseClassDeclaration(const TCHAR* Filename, const TCHAR* InputText, int32 InLineNumber, const TCHAR* StartingMatchID, FName& out_StrippedClassName, FString& out_ClassName, FString& out_BaseClassName, TArray<FHeaderProvider>& out_RequiredIncludes, const TArray<FSimplifiedParsingClassInfo>& ParsedClassArray)
 {
-	FString ErrorMsg = TEXT("Class declaration");
+	const TCHAR* ErrorMsg = TEXT("Class declaration");
 
 	ResetParser(InputText, InLineNumber);
 
 	// Require 'UCLASS' or 'UINTERFACE'
-	RequireIdentifier(StartingMatchID, *ErrorMsg);
+	RequireIdentifier(StartingMatchID, ESearchCase::CaseSensitive, ErrorMsg);
 
 	// New-style UCLASS() syntax
 	TMap<FName, FString> MetaData;
@@ -9297,7 +9332,7 @@ void FHeaderPreParser::ParseClassDeclaration(const TCHAR* Filename, const TCHAR*
 	ReadSpecifierSetInsideMacro(SpecifiersFound, ErrorMsg, MetaData);
 
 	// Require 'class'
-	RequireIdentifier(TEXT("class"), *ErrorMsg);
+	RequireIdentifier(TEXT("class"), ESearchCase::CaseSensitive, ErrorMsg);
 
 	// alignas() can come before or after the deprecation macro.
 	// We can't have both, but the compiler will catch that anyway.
@@ -9316,19 +9351,19 @@ void FHeaderPreParser::ParseClassDeclaration(const TCHAR* Filename, const TCHAR*
 	{
 		// Add class declaration meta data so that we can access class flags before the class is fully parsed
 		TSharedRef<FClassDeclarationMetaData> DeclarationData = MakeShareable(new FClassDeclarationMetaData());
-		DeclarationData->MetaData = MetaData;
+		DeclarationData->MetaData = MoveTemp(MetaData);
 		DeclarationData->ParseClassProperties(MoveTemp(SpecifiersFound), RequiredAPIMacroIfPresent);
 		GClassDeclarations.Add(out_StrippedClassName, DeclarationData);
 	}
 
 	// Skip optional final keyword
-	MatchIdentifier(TEXT("final"));
+	MatchIdentifier(TEXT("final"), ESearchCase::CaseSensitive);
 
 	// Handle inheritance
-	if (MatchSymbol(TEXT(":")))
+	if (MatchSymbol(TEXT(':')))
 	{
 		// Require 'public'
-		RequireIdentifier(TEXT("public"), *ErrorMsg);
+		RequireIdentifier(TEXT("public"), ESearchCase::CaseSensitive, ErrorMsg);
 
 		// Inherits from something
 		FToken BaseClassNameToken;
@@ -9366,10 +9401,10 @@ void FHeaderPreParser::ParseClassDeclaration(const TCHAR* Filename, const TCHAR*
 		AddDependencyIfNeeded(out_BaseClassName);
 
 		// Get additional inheritance links and rack them up as dependencies if they're UObject derived
-		while (MatchSymbol(TEXT(",")))
+		while (MatchSymbol(TEXT(',')))
 		{
 			// Require 'public'
-			RequireIdentifier(TEXT("public"), *ErrorMsg);
+			RequireIdentifier(TEXT("public"), ESearchCase::CaseSensitive, ErrorMsg);
 
 			FToken InterfaceClassNameToken;
 			if (!GetIdentifier(InterfaceClassNameToken, true))
@@ -9746,7 +9781,7 @@ bool FHeaderParser::TryToMatchConstructorParameterList(FToken Token)
 		return false;
 	}
 
-	if (!PotentialParenthesisToken.Matches(TEXT("(")))
+	if (!PotentialParenthesisToken.Matches(TEXT('(')))
 	{
 		UngetToken(PotentialParenthesisToken);
 		return false;
@@ -9758,7 +9793,7 @@ bool FHeaderParser::TryToMatchConstructorParameterList(FToken Token)
 	bool bOICtor = false;
 	bool bVTCtor = false;
 
-	if (!ClassData->bDefaultConstructorDeclared && MatchSymbol(TEXT(")")))
+	if (!ClassData->bDefaultConstructorDeclared && MatchSymbol(TEXT(')')))
 	{
 		ClassData->bDefaultConstructorDeclared = true;
 	}
@@ -9775,45 +9810,45 @@ bool FHeaderParser::TryToMatchConstructorParameterList(FToken Token)
 		while (ParenthesesNestingLevel && GetToken(ObjectInitializerParamParsingToken))
 		{
 			// Template instantiation or additional parameter excludes ObjectInitializer constructor.
-			if (ObjectInitializerParamParsingToken.Matches(TEXT(",")) || ObjectInitializerParamParsingToken.Matches(TEXT("<")))
+			if (ObjectInitializerParamParsingToken.Matches(TEXT(',')) || ObjectInitializerParamParsingToken.Matches(TEXT('<')))
 			{
 				bOICtor = false;
 				bVTCtor = false;
 				break;
 			}
 
-			if (ObjectInitializerParamParsingToken.Matches(TEXT("(")))
+			if (ObjectInitializerParamParsingToken.Matches(TEXT('(')))
 			{
 				ParenthesesNestingLevel++;
 				continue;
 			}
 
-			if (ObjectInitializerParamParsingToken.Matches(TEXT(")")))
+			if (ObjectInitializerParamParsingToken.Matches(TEXT(')')))
 			{
 				ParenthesesNestingLevel--;
 				continue;
 			}
 
-			if (ObjectInitializerParamParsingToken.Matches(TEXT("const")))
+			if (ObjectInitializerParamParsingToken.Matches(TEXT("const"), ESearchCase::CaseSensitive))
 			{
 				bIsConst = true;
 				continue;
 			}
 
-			if (ObjectInitializerParamParsingToken.Matches(TEXT("&")))
+			if (ObjectInitializerParamParsingToken.Matches(TEXT('&')))
 			{
 				bIsRef = true;
 				continue;
 			}
 
-			if (ObjectInitializerParamParsingToken.Matches(TEXT("FObjectInitializer"))
-				|| ObjectInitializerParamParsingToken.Matches(TEXT("FPostConstructInitializeProperties")) // Deprecated, but left here, so it won't break legacy code.
+			if (ObjectInitializerParamParsingToken.Matches(TEXT("FObjectInitializer"), ESearchCase::CaseSensitive)
+				|| ObjectInitializerParamParsingToken.Matches(TEXT("FPostConstructInitializeProperties"), ESearchCase::CaseSensitive) // Deprecated, but left here, so it won't break legacy code.
 				)
 			{
 				bOICtor = true;
 			}
 
-			if (ObjectInitializerParamParsingToken.Matches(TEXT("FVTableHelper")))
+			if (ObjectInitializerParamParsingToken.Matches(TEXT("FVTableHelper"), ESearchCase::CaseSensitive))
 			{
 				bVTCtor = true;
 			}
@@ -9822,13 +9857,13 @@ bool FHeaderParser::TryToMatchConstructorParameterList(FToken Token)
 		// Parse until finish.
 		while (ParenthesesNestingLevel && GetToken(ObjectInitializerParamParsingToken))
 		{
-			if (ObjectInitializerParamParsingToken.Matches(TEXT("(")))
+			if (ObjectInitializerParamParsingToken.Matches(TEXT('(')))
 			{
 				ParenthesesNestingLevel++;
 				continue;
 			}
 
-			if (ObjectInitializerParamParsingToken.Matches(TEXT(")")))
+			if (ObjectInitializerParamParsingToken.Matches(TEXT(')')))
 			{
 				ParenthesesNestingLevel--;
 				continue;
@@ -9842,7 +9877,7 @@ bool FHeaderParser::TryToMatchConstructorParameterList(FToken Token)
 	ClassData->bConstructorDeclared = ClassData->bConstructorDeclared || !bVTCtor;
 
 	// Optionally match semicolon.
-	if (!MatchSymbol(TEXT(";")))
+	if (!MatchSymbol(TEXT(';')))
 	{
 		// If not matched a semicolon, this is inline constructor definition. We have to skip it.
 		UngetToken(Token); // Resets input stream to the initial token.
@@ -9874,7 +9909,7 @@ void FHeaderParser::CompileVersionDeclaration(UStruct* Struct)
 	}
 
 	if (Token.TokenType == ETokenType::TOKEN_Symbol
-		&& !FCString::Stricmp(Token.Identifier, TEXT(")")))
+		&& Token.Matches(TEXT(')')))
 	{
 		CurrentSrcFile->GetGeneratedCodeVersions().FindOrAdd(Struct) = Version;
 		UngetToken(Token);
@@ -9895,8 +9930,7 @@ void FHeaderParser::ResetClassData()
 	// Set class flags and within.
 	CurrentClass->ClassFlags &= ~CLASS_RecompilerClear;
 
-	UClass* SuperClass = CurrentClass->GetSuperClass();
-	if (SuperClass != NULL)
+	if (UClass* SuperClass = CurrentClass->GetSuperClass())
 	{
 		CurrentClass->ClassFlags |= (SuperClass->ClassFlags) & CLASS_ScriptInherit;
 		CurrentClass->ClassConfigName = SuperClass->ClassConfigName;
@@ -9907,29 +9941,29 @@ void FHeaderParser::ResetClassData()
 		}
 
 		// Copy special categories from parent
-		if (SuperClass->HasMetaData(TEXT("HideCategories")))
+		if (SuperClass->HasMetaData(NAME_HideCategories))
 		{
-			CurrentClass->SetMetaData(TEXT("HideCategories"), *SuperClass->GetMetaData("HideCategories"));
+			CurrentClass->SetMetaData(NAME_HideCategories, *SuperClass->GetMetaData(NAME_HideCategories));
 		}
-		if (SuperClass->HasMetaData(TEXT("ShowCategories")))
+		if (SuperClass->HasMetaData(NAME_ShowCategories))
 		{
-			CurrentClass->SetMetaData(TEXT("ShowCategories"), *SuperClass->GetMetaData("ShowCategories"));
+			CurrentClass->SetMetaData(NAME_ShowCategories, *SuperClass->GetMetaData(NAME_ShowCategories));
 		}
-		if (SuperClass->HasMetaData(TEXT("SparseClassDataTypes")))
+		if (SuperClass->HasMetaData(NAME_SparseClassDataTypes))
 		{
-			CurrentClass->SetMetaData(TEXT("SparseClassDataTypes"), *SuperClass->GetMetaData("SparseClassDataTypes"));
+			CurrentClass->SetMetaData(NAME_SparseClassDataTypes, *SuperClass->GetMetaData(NAME_SparseClassDataTypes));
 		}
-		if (SuperClass->HasMetaData(TEXT("HideFunctions")))
+		if (SuperClass->HasMetaData(NAME_HideFunctions))
 		{
-			CurrentClass->SetMetaData(TEXT("HideFunctions"), *SuperClass->GetMetaData("HideFunctions"));
+			CurrentClass->SetMetaData(NAME_HideFunctions, *SuperClass->GetMetaData(NAME_HideFunctions));
 		}
-		if (SuperClass->HasMetaData(TEXT("AutoExpandCategories")))
+		if (SuperClass->HasMetaData(NAME_AutoExpandCategories))
 		{
-			CurrentClass->SetMetaData(TEXT("AutoExpandCategories"), *SuperClass->GetMetaData("AutoExpandCategories"));
+			CurrentClass->SetMetaData(NAME_AutoExpandCategories, *SuperClass->GetMetaData(NAME_AutoExpandCategories));
 		}
-		if (SuperClass->HasMetaData(TEXT("AutoCollapseCategories")))
+		if (SuperClass->HasMetaData(NAME_AutoCollapseCategories))
 		{
-			CurrentClass->SetMetaData(TEXT("AutoCollapseCategories"), *SuperClass->GetMetaData("AutoCollapseCategories"));
+			CurrentClass->SetMetaData(NAME_AutoCollapseCategories, *SuperClass->GetMetaData(NAME_AutoCollapseCategories));
 		}
 	}
 
@@ -9956,7 +9990,7 @@ void FHeaderParser::PostPopNestClass(UClass* CurrentClass)
 			// So iterate over all functions this interface declares
 			for (UFunction* InterfaceFunction : TFieldRange<UFunction>(Interface, EFieldIteratorFlags::ExcludeSuper))
 			{
-				bool Implemented = false;
+				bool bImplemented = false;
 
 				// And try to find one that matches
 				for (UFunction* ClassFunction : TFieldRange<UFunction>(CurrentClass))
@@ -9977,7 +10011,7 @@ void FHeaderParser::PostPopNestClass(UClass* CurrentClass)
 					}
 
 					// Making sure all the parameters match up correctly
-					Implemented = true;
+					bImplemented = true;
 
 					if (ClassFunction->NumParms != InterfaceFunction->NumParms)
 					{
@@ -10004,14 +10038,14 @@ void FHeaderParser::PostPopNestClass(UClass* CurrentClass)
 				// Delegate signature functions are simple stubs and aren't required to be implemented (they are not callable)
 				if (InterfaceFunction->FunctionFlags & FUNC_Delegate)
 				{
-					Implemented = true;
+					bImplemented = true;
 				}
 
 				// Verify that if this has blueprint-callable functions that are not implementable events, we've implemented them as a UFunction in the target class
-				if (!Implemented
-					&& !Interface->HasMetaData(TEXT("CannotImplementInterfaceInBlueprint"))  // FBlueprintMetadata::MD_CannotImplementInterfaceInBlueprint
+				if (!bImplemented
 					&& InterfaceFunction->HasAnyFunctionFlags(FUNC_BlueprintCallable)
-					&& !InterfaceFunction->HasAnyFunctionFlags(FUNC_BlueprintEvent))
+					&& !InterfaceFunction->HasAnyFunctionFlags(FUNC_BlueprintEvent)
+					&& !Interface->HasMetaData(NAME_CannotImplementInterfaceInBlueprint))  // FBlueprintMetadata::MD_CannotImplementInterfaceInBlueprint
 				{
 					FError::Throwf(TEXT("Missing UFunction implementation of function '%s' from interface '%s'.  This function needs a UFUNCTION() declaration."), *InterfaceFunction->GetName(), *Interface->GetName());
 				}
@@ -10045,7 +10079,7 @@ void FHeaderParser::PostPopNestInterface(FClasses& AllClasses, UClass* CurrentIn
 FDocumentationPolicy FHeaderParser::GetDocumentationPolicyFromName(const FString& PolicyName)
 {
 	FDocumentationPolicy DocumentationPolicy;
-	if (PolicyName.Equals(TEXT("Strict")))
+	if (FCString::Strcmp(*PolicyName, TEXT("Strict")) == 0)
 	{
 		DocumentationPolicy.bClassOrStructCommentRequired = true;
 		DocumentationPolicy.bFunctionToolTipsRequired = true;
@@ -10103,7 +10137,7 @@ void FHeaderParser::CheckDocumentationPolicyForEnum(UEnum* Enum, const TMap<FNam
 	TMap<FString, FString> ToolTipToEntry;
 	for (const TMap<FName, FString>& Entry : Entries)
 	{
-		const FString* EntryName = Entry.Find(TEXT("Name"));
+		const FString* EntryName = Entry.Find(NAME_Name);
 		if (EntryName == nullptr)
 		{
 			continue;
@@ -10163,7 +10197,7 @@ void FHeaderParser::CheckDocumentationPolicyForStruct(UStruct* Struct, const TMa
 			{
 				UE_LOG_ERROR_UHT(TEXT("Property '%s::%s' and '%s::%s' are using identical tooltips (DocumentationPolicy)."), *Struct->GetName(), *ExistingPropertyName->ToString(), *Struct->GetName(), *Property->GetName());
 			}
-			ToolTipToPropertyName.Add(ToolTip, Property->GetFName());
+			ToolTipToPropertyName.Add(MoveTemp(ToolTip), Property->GetFName());
 		}
 	}
 
@@ -10173,8 +10207,8 @@ void FHeaderParser::CheckDocumentationPolicyForStruct(UStruct* Struct, const TMa
 		{
 			if(DoesCPPTypeRequireDocumentation(Property->GetCPPType()))
 			{
-				const FString& UIMin = Property->GetMetaData(TEXT("UIMin"));
-				const FString& UIMax = Property->GetMetaData(TEXT("UIMax"));
+				const FString& UIMin = Property->GetMetaData(NAME_UIMin);
+				const FString& UIMax = Property->GetMetaData(NAME_UIMax);
 
 				if(!CheckUIMinMaxRangeFromMetaData(UIMin, UIMax))
 				{
@@ -10190,7 +10224,7 @@ void FHeaderParser::CheckDocumentationPolicyForStruct(UStruct* Struct, const TMa
 		UClass* Class = Cast<UClass>(Struct);
 		if (Class != nullptr)
 		{
-			TMap<FString, FString> ToolTipToFunc;
+			TMap<FString, FName> ToolTipToFunc;
 			for (UFunction* Func : TFieldRange<UFunction>(Class, EFieldIteratorFlags::ExcludeSuper))
 			{
 				FString ToolTip = Func->GetToolTipText().ToString();
@@ -10199,12 +10233,12 @@ void FHeaderParser::CheckDocumentationPolicyForStruct(UStruct* Struct, const TMa
 					UE_LOG_ERROR_UHT(TEXT("Function '%s::%s' does not provide a tooltip / comment (DocumentationPolicy)."), *Class->GetName(), *Func->GetName());
 					continue;
 				}
-				const FString* ExistingFuncName = ToolTipToFunc.Find(ToolTip);
+				const FName* ExistingFuncName = ToolTipToFunc.Find(ToolTip);
 				if (ExistingFuncName != nullptr)
 				{
-					UE_LOG_ERROR_UHT(TEXT("Functions '%s::%s' and '%s::%s' uses identical tooltips / comments (DocumentationPolicy)."), *Class->GetName(), *(*ExistingFuncName), *Class->GetName(), *Func->GetName());
+					UE_LOG_ERROR_UHT(TEXT("Functions '%s::%s' and '%s::%s' uses identical tooltips / comments (DocumentationPolicy)."), *Class->GetName(), *(*ExistingFuncName).ToString(), *Class->GetName(), *Func->GetName());
 				}
-				ToolTipToFunc.Add(ToolTip, Func->GetName());
+				ToolTipToFunc.Add(MoveTemp(ToolTip), Func->GetFName());
 			}
 		}
 	}
@@ -10246,7 +10280,7 @@ void FHeaderParser::CheckDocumentationPolicyForFunc(UClass* Class, UFunction* Fu
 		bool HasAnyParamToolTips = ParamToolTips.Num() > 0;
 		if (ParamToolTips.Num() == 0)
 		{
-			const FString* ReturnValueToolTip = ParamToolTips.Find(TEXT("ReturnValue"));
+			const FString* ReturnValueToolTip = ParamToolTips.Find(NAME_ReturnValue);
 			if (ReturnValueToolTip != nullptr)
 			{
 				HasAnyParamToolTips = false;
@@ -10261,7 +10295,7 @@ void FHeaderParser::CheckDocumentationPolicyForFunc(UClass* Class, UFunction* Fu
 			for (FProperty* Property : TFieldRange<FProperty>(Func))
 			{
 				FName ParamName = Property->GetFName();
-				if (ParamName.IsEqual(TEXT("ReturnValue")))
+				if (ParamName == NAME_ReturnValue)
 				{
 					continue;
 				}
@@ -10277,7 +10311,7 @@ void FHeaderParser::CheckDocumentationPolicyForFunc(UClass* Class, UFunction* Fu
 			for (TPair<FName, FString>& Pair : ParamToolTips)
 			{
 				const FName& ParamName = Pair.Key;
-				if (ParamName.IsEqual(TEXT("ReturnValue")))
+				if (ParamName == NAME_ReturnValue)
 				{
 					continue;
 				}
@@ -10292,7 +10326,7 @@ void FHeaderParser::CheckDocumentationPolicyForFunc(UClass* Class, UFunction* Fu
 			for (TPair<FName, FString>& Pair : ParamToolTips)
 			{
 				const FName& ParamName = Pair.Key;
-				if (ParamName.IsEqual(TEXT("ReturnValue")))
+				if (ParamName == NAME_ReturnValue)
 				{
 					continue;
 				}
@@ -10301,7 +10335,7 @@ void FHeaderParser::CheckDocumentationPolicyForFunc(UClass* Class, UFunction* Fu
 				{
 					UE_LOG_ERROR_UHT(TEXT("Function '%s::%s' uses identical tooltips for parameters '%s' and '%s' (DocumentationPolicy)."), *Class->GetName(), *Func->GetName(), *ExistingParam->ToString(), *Pair.Key.ToString());
 				}
-				ToolTipToParam.Add(Pair.Value, Pair.Key);
+				ToolTipToParam.Add(MoveTemp(Pair.Value), Pair.Key);
 			}
 		}
 	}
diff --git a/Engine/Source/Programs/UnrealHeaderTool/Private/HeaderParser.h b/Engine/Source/Programs/UnrealHeaderTool/Private/HeaderParser.h
--- a/Engine/Source/Programs/UnrealHeaderTool/Private/HeaderParser.h
+++ b/Engine/Source/Programs/UnrealHeaderTool/Private/HeaderParser.h
@@ -433,7 +433,7 @@ public:
 	 * @param InNameToCheck - Name w/ potential prefix to check
 	 * @param OriginalClassName - Name of class w/ no prefix to check against
 	 */
-	static bool ClassNameHasValidPrefix(const FString InNameToCheck, const FString OriginalClassName);
+	static bool ClassNameHasValidPrefix(const FString& InNameToCheck, const FString& OriginalClassName);
 
 	/**
 	 * Tries to convert the header file name to a class name (with 'U' prefix)
@@ -768,7 +768,7 @@ protected:
 	/** Skip C++ (noexport) declaration. */
 	bool SkipDeclaration(FToken& Token);
 	/** Similar to MatchSymbol() but will return to the exact location as on entry if the symbol was not found. */
-	bool SafeMatchSymbol(const TCHAR* Match);
+	bool SafeMatchSymbol(const TCHAR Match);
 	void HandleOneInheritedClass(FClasses& AllClasses, UClass* Class, FString InterfaceName);
 	FClass* ParseClassNameDeclaration(FClasses& AllClasses, FString& DeclaredClassName, FString& RequiredAPIMacroIfPresent);
 
diff --git a/Engine/Source/Programs/UnrealHeaderTool/Private/ParserClass.cpp b/Engine/Source/Programs/UnrealHeaderTool/Private/ParserClass.cpp
--- a/Engine/Source/Programs/UnrealHeaderTool/Private/ParserClass.cpp
+++ b/Engine/Source/Programs/UnrealHeaderTool/Private/ParserClass.cpp
@@ -5,6 +5,11 @@
 #include "UObject/Package.h"
 #include "Templates/Casts.h"
 
+static const FName NAME_HideCategories(TEXT("HideCategories"));
+static const FName NAME_ShowCategories(TEXT("ShowCategories"));
+static const FName NAME_SparseClassDataTypes(TEXT("SparseClassDataTypes"));
+const FName FClass::NAME_ReplaceConverted(TEXT("ReplaceConverted"));
+
 FString FClass::GetNameWithPrefix(EEnforceInterfacePrefix::Type EnforceInterfacePrefix) const
 {
 	const TCHAR* Prefix = 0;
@@ -60,7 +65,7 @@ FClass* FClass::GetClassWithin() const
 TArray<FClass*> FClass::GetInterfaceTypes() const
 {
 	TArray<FClass*> Result;
-	for (auto& i : Interfaces)
+	for (const FImplementedInterface& i : Interfaces)
 	{
 		Result.Add((FClass*)i.Class);
 	}
@@ -69,7 +74,6 @@ TArray<FClass*> FClass::GetInterfaceTypes() const
 
 void FClass::GetHideCategories(TArray<FString>& OutHideCategories) const
 {
-	static const FName NAME_HideCategories(TEXT("HideCategories"));
 	if (HasMetaData(NAME_HideCategories))
 	{
 		const FString& HideCategories = GetMetaData(NAME_HideCategories);
@@ -79,7 +83,6 @@ void FClass::GetHideCategories(TArray<FString>& OutHideCategories) const
 
 void FClass::GetShowCategories(TArray<FString>& OutShowCategories) const
 {
-	static const FName NAME_ShowCategories(TEXT("ShowCategories"));
 	if (HasMetaData(NAME_ShowCategories))
 	{
 		const FString& ShowCategories = GetMetaData(NAME_ShowCategories);
@@ -89,7 +92,6 @@ void FClass::GetShowCategories(TArray<FString>& OutShowCategories) const
 
 void FClass::GetSparseClassDataTypes(TArray<FString>& OutSparseClassDataTypes) const
 {
-	static const FName NAME_SparseClassDataTypes(TEXT("SparseClassDataTypes"));
 	if (HasMetaData(NAME_SparseClassDataTypes))
 	{
 		const FString& SparseClassDataTypes = GetMetaData(NAME_SparseClassDataTypes);
@@ -120,7 +122,45 @@ bool FClass::IsOwnedByDynamicType(const FField* Field)
 		else if (IsDynamic(Owner.ToField()))
 		{
 			return true;
-	}
+		}
 	}
 	return false;
 }
+
+static TMap<const UField*, FString> UFieldTypePackageNames;
+static TMap<const FField*, FString> FFieldTypePackageNames;
+
+template <typename T>
+const FString& GetTypePackageName_Inner(const T* Field, TMap<const T*, FString>& TypePackageNames)
+{
+	FString* TypePackageName = TypePackageNames.Find(Field);
+	if (TypePackageName == nullptr)
+	{
+		FString PackageName = Field->GetMetaData(FClass::NAME_ReplaceConverted);
+		if (PackageName.Len())
+		{
+			int32 ObjectDotIndex = INDEX_NONE;
+			// Strip the object name
+			if (PackageName.FindChar(TEXT('.'), ObjectDotIndex))
+			{
+				PackageName.MidInline(0, ObjectDotIndex, false);
+			}
+		}
+		else
+		{
+			PackageName = Field->GetOutermost()->GetName();
+		}
+		TypePackageName = &TypePackageNames.Add(Field, MoveTemp(PackageName));
+	}
+	return *TypePackageName;
+}
+
+const FString& FClass::GetTypePackageName(const UField* Field)
+{
+	return GetTypePackageName_Inner(Field, UFieldTypePackageNames);
+}
+
+const FString& FClass::GetTypePackageName(const FField* Field)
+{
+	return GetTypePackageName_Inner(Field, FFieldTypePackageNames);
+}
diff --git a/Engine/Source/Programs/UnrealHeaderTool/Private/ParserClass.h b/Engine/Source/Programs/UnrealHeaderTool/Private/ParserClass.h
--- a/Engine/Source/Programs/UnrealHeaderTool/Private/ParserClass.h
+++ b/Engine/Source/Programs/UnrealHeaderTool/Private/ParserClass.h
@@ -56,7 +56,6 @@ public:
 	template <typename T>
 	static bool IsDynamic(const T* Field)
 	{
-		static const FName NAME_ReplaceConverted(TEXT("ReplaceConverted"));
 		return Field->HasMetaData(NAME_ReplaceConverted);
 	}
 
@@ -65,24 +64,8 @@ public:
 	static bool IsOwnedByDynamicType(const FField* Field);
 
 	/** Helper function to get the source replaced package name */
-	template <typename T>
-	static FString GetTypePackageName(const T* Field)
-	{
-		static const FName NAME_ReplaceConverted(TEXT("ReplaceConverted"));
-		FString PackageName = Field->GetMetaData(NAME_ReplaceConverted);
-		if (PackageName.Len())
-		{
-			int32 ObjectDotIndex = INDEX_NONE;
-			// Strip the object name
-			if (PackageName.FindChar(TEXT('.'), ObjectDotIndex))
-			{
-				PackageName = PackageName.Mid(0, ObjectDotIndex);
-			}
-		}
-		else
-		{
-			PackageName = Field->GetOutermost()->GetName();
-		}
-		return PackageName;
-	}
+	static const FString& GetTypePackageName(const UField* Field);
+	static const FString& GetTypePackageName(const FField* Field);
+
+	static const FName NAME_ReplaceConverted;
 };
diff --git a/Engine/Source/Programs/UnrealHeaderTool/Private/ParserHelper.h b/Engine/Source/Programs/UnrealHeaderTool/Private/ParserHelper.h
--- a/Engine/Source/Programs/UnrealHeaderTool/Private/ParserHelper.h
+++ b/Engine/Source/Programs/UnrealHeaderTool/Private/ParserHelper.h
@@ -731,8 +731,14 @@ public:
 	//@{
 	/** Type of token. */
 	ETokenType TokenType;
-	/** Name of token. */
-	FName TokenName;
+
+private:
+	/** Whether TokenName has been looked up */
+	bool bTokenNameInitialized;
+	/** Name of token, lazily initialized */
+	mutable FName TokenName;
+
+public:
 	/** Starting position in script where this token came from. */
 	int32 StartPos;
 	/** Starting line in script. */
@@ -742,6 +748,9 @@ public:
 	/** property that corresponds to this FToken - null if this Token doesn't correspond to a FProperty */
 	FProperty* TokenProperty;
 
+
+public:
+
 	union
 	{
 		// TOKEN_Const values.
@@ -834,22 +843,30 @@ public:
 		*Identifier		= 0;
 		FMemory::Memzero(String);
 	}
-	bool Matches( const TCHAR* Str, ESearchCase::Type SearchCase = ESearchCase::IgnoreCase ) const
+	bool Matches(const TCHAR Ch) const
+	{
+		return TokenType==TOKEN_Symbol && Identifier[0] == Ch && Identifier[1] == 0;
+	}
+	bool Matches( const TCHAR* Str, ESearchCase::Type SearchCase ) const
 	{
 		return (TokenType==TOKEN_Identifier || TokenType==TOKEN_Symbol) && ((SearchCase == ESearchCase::CaseSensitive) ? !FCString::Strcmp(Identifier, Str) : !FCString::Stricmp(Identifier, Str));
 	}
-	bool Matches( const FName& Name ) const
+	bool IsBool() const
 	{
-		return TokenType==TOKEN_Identifier && TokenName==Name;
+		return Type == CPT_Bool || Type == CPT_Bool8 || Type == CPT_Bool16 || Type == CPT_Bool32 || Type == CPT_Bool64;
 	}
-	bool StartsWith( const TCHAR* Str, bool bCaseSensitive = false ) const
+	FName GetTokenName() const
 	{
-		const int32 StrLength = FCString::Strlen(Str);
-		return (TokenType==TOKEN_Identifier || TokenType==TOKEN_Symbol) && (bCaseSensitive ? (!FCString::Strncmp(Identifier, Str, StrLength)) : (!FCString::Strnicmp(Identifier, Str, StrLength)));
+		if (!bTokenNameInitialized)
+		{
+			TokenName = FName(Identifier, FNAME_Find);
+		}
+		return TokenName;
 	}
-	bool IsBool() const
+	void ClearTokenName()
 	{
-		return Type == CPT_Bool || Type == CPT_Bool8 || Type == CPT_Bool16 || Type == CPT_Bool32 || Type == CPT_Bool64;
+		bTokenNameInitialized = false;
+		TokenName = NAME_None;
 	}
 
 	// Setters.
@@ -859,7 +876,7 @@ public:
 		InitToken(CPT_None);
 		TokenType = TOKEN_Identifier;
 		FCString::Strncpy(Identifier, InString, NAME_SIZE);
-		TokenName = FName(Identifier, FNAME_Find);
+		bTokenNameInitialized = false;
 	}
 
 	void SetConstInt64( int64 InInt64 )
@@ -984,7 +1001,7 @@ public:
 		return FString::Printf(
 			TEXT("Property:%s  Type:%s  TokenName:%s  ConstValue:%s  Struct:%s  Flags:%lli  Implied:%lli"),
 			TokenProperty!=NULL?*TokenProperty->GetName():TEXT("NULL"),
-			GetPropertyTypeText(Type), *TokenName.ToString(), *GetConstantValue(),
+			GetPropertyTypeText(Type), *GetTokenName().ToString(), *GetConstantValue(),
 			Struct!=NULL?*Struct->GetName():TEXT("NULL"),
 			PropertyFlags,
 			ImpliedPropertyFlags
@@ -1093,7 +1110,7 @@ struct FFuncInfo
 		FString FunctionName = FunctionReference->GetName();
 		if( FunctionReference->HasAnyFunctionFlags( FUNC_Delegate ) )
 		{
-			FunctionName = FunctionName.LeftChop( FString(HEADER_GENERATED_DELEGATE_SIGNATURE_SUFFIX).Len() );
+			FunctionName.LeftChopInline( FString(HEADER_GENERATED_DELEGATE_SIGNATURE_SUFFIX).Len(), false );
 		}
 		UnMarshallAndCallName = FString(TEXT("exec")) + FunctionName;
 		
