Fix duplicate FName declarations in unity builds
#rnx

diff --git a/Engine/Source/Programs/UnrealHeaderTool/Private/HeaderParser.cpp b/Engine/Source/Programs/UnrealHeaderTool/Private/HeaderParser.cpp
--- a/Engine/Source/Programs/UnrealHeaderTool/Private/HeaderParser.cpp
+++ b/Engine/Source/Programs/UnrealHeaderTool/Private/HeaderParser.cpp
@@ -66,36 +66,39 @@ private:
 
 enum {MAX_ARRAY_SIZE=2048};
 
-static const FName NAME_Comment(TEXT("Comment"));
-static const FName NAME_ToolTip(TEXT("ToolTip"));
-static const FName NAME_DocumentationPolicy(TEXT("DocumentationPolicy"));
-static const FName NAME_AllowPrivateAccess(TEXT("AllowPrivateAccess"));
-static const FName NAME_ExposeOnSpawn(TEXT("ExposeOnSpawn"));
-static const FName NAME_NativeConst(TEXT("NativeConst"));
-static const FName NAME_NativeConstTemplateArg(TEXT("NativeConstTemplateArg"));
-static const FName NAME_BlueprintInternalUseOnly(TEXT("BlueprintInternalUseOnly"));
-static const FName NAME_DeprecatedFunction(TEXT("DeprecatedFunction"));
-static const FName NAME_BlueprintSetter(TEXT("BlueprintSetter"));
-static const FName NAME_BlueprintGetter(TEXT("BlueprintGetter"));
-static const FName NAME_Category(TEXT("Category"));
-static const FName NAME_ReturnValue(TEXT("ReturnValue"));
-static const FName NAME_CppFromBpEvent(TEXT("CppFromBpEvent"));
-static const FName NAME_CustomThunk(TEXT("CustomThunk"));
-static const FName NAME_ArraySizeEnum(TEXT("ArraySizeEnum"));
-static const FName NAME_ClassGroupNames(TEXT("ClassGroupNames"));
-static const FName NAME_AutoCollapseCategories(TEXT("AutoCollapseCategories"));
-static const FName NAME_HideCategories(TEXT("HideCategories"));
-static const FName NAME_ShowCategories(TEXT("ShowCategories"));
-static const FName NAME_SparseClassDataTypes(TEXT("SparseClassDataTypes"));
-static const FName NAME_HideFunctions(TEXT("HideFunctions"));
-static const FName NAME_AutoExpandCategories(TEXT("AutoExpandCategories"));
-static const FName NAME_IsConversionRoot(TEXT("IsConversionRoot"));
-static const FName NAME_EditInline(TEXT("EditInline"));
-static const FName NAME_IncludePath(TEXT("IncludePath"));
-static const FName NAME_ModuleRelativePath(TEXT("ModuleRelativePath"));
-static const FName NAME_CannotImplementInterfaceInBlueprint(TEXT("CannotImplementInterfaceInBlueprint"));
-static const FName NAME_UIMin(TEXT("UIMin"));
-static const FName NAME_UIMax(TEXT("UIMax"));
+namespace
+{
+	static const FName NAME_Comment(TEXT("Comment"));
+	static const FName NAME_ToolTip(TEXT("ToolTip"));
+	static const FName NAME_DocumentationPolicy(TEXT("DocumentationPolicy"));
+	static const FName NAME_AllowPrivateAccess(TEXT("AllowPrivateAccess"));
+	static const FName NAME_ExposeOnSpawn(TEXT("ExposeOnSpawn"));
+	static const FName NAME_NativeConst(TEXT("NativeConst"));
+	static const FName NAME_NativeConstTemplateArg(TEXT("NativeConstTemplateArg"));
+	static const FName NAME_BlueprintInternalUseOnly(TEXT("BlueprintInternalUseOnly"));
+	static const FName NAME_DeprecatedFunction(TEXT("DeprecatedFunction"));
+	static const FName NAME_BlueprintSetter(TEXT("BlueprintSetter"));
+	static const FName NAME_BlueprintGetter(TEXT("BlueprintGetter"));
+	static const FName NAME_Category(TEXT("Category"));
+	static const FName NAME_ReturnValue(TEXT("ReturnValue"));
+	static const FName NAME_CppFromBpEvent(TEXT("CppFromBpEvent"));
+	static const FName NAME_CustomThunk(TEXT("CustomThunk"));
+	static const FName NAME_ArraySizeEnum(TEXT("ArraySizeEnum"));
+	static const FName NAME_ClassGroupNames(TEXT("ClassGroupNames"));
+	static const FName NAME_AutoCollapseCategories(TEXT("AutoCollapseCategories"));
+	static const FName NAME_HideCategories(TEXT("HideCategories"));
+	static const FName NAME_ShowCategories(TEXT("ShowCategories"));
+	static const FName NAME_SparseClassDataTypes(TEXT("SparseClassDataTypes"));
+	static const FName NAME_HideFunctions(TEXT("HideFunctions"));
+	static const FName NAME_AutoExpandCategories(TEXT("AutoExpandCategories"));
+	static const FName NAME_IsConversionRoot(TEXT("IsConversionRoot"));
+	static const FName NAME_EditInline(TEXT("EditInline"));
+	static const FName NAME_IncludePath(TEXT("IncludePath"));
+	static const FName NAME_ModuleRelativePath(TEXT("ModuleRelativePath"));
+	static const FName NAME_CannotImplementInterfaceInBlueprint(TEXT("CannotImplementInterfaceInBlueprint"));
+	static const FName NAME_UIMin(TEXT("UIMin"));
+	static const FName NAME_UIMax(TEXT("UIMax"));
+}
 
 EGeneratedCodeVersion FHeaderParser::DefaultGeneratedCodeVersion = EGeneratedCodeVersion::V1;
 TArray<FString> FHeaderParser::StructsWithNoPrefix;
diff --git a/Engine/Source/Programs/UnrealHeaderTool/Private/ParserClass.cpp b/Engine/Source/Programs/UnrealHeaderTool/Private/ParserClass.cpp
--- a/Engine/Source/Programs/UnrealHeaderTool/Private/ParserClass.cpp
+++ b/Engine/Source/Programs/UnrealHeaderTool/Private/ParserClass.cpp
@@ -5,9 +5,12 @@
 #include "UObject/Package.h"
 #include "Templates/Casts.h"
 
-static const FName NAME_HideCategories(TEXT("HideCategories"));
-static const FName NAME_ShowCategories(TEXT("ShowCategories"));
-static const FName NAME_SparseClassDataTypes(TEXT("SparseClassDataTypes"));
+namespace
+{
+	static const FName NAME_HideCategories(TEXT("HideCategories"));
+	static const FName NAME_ShowCategories(TEXT("ShowCategories"));
+	static const FName NAME_SparseClassDataTypes(TEXT("SparseClassDataTypes"));
+}
 const FName FClass::NAME_ReplaceConverted(TEXT("ReplaceConverted"));
 
 FString FClass::GetNameWithPrefix(EEnforceInterfacePrefix::Type EnforceInterfacePrefix) const
