Sort out duplicate symbol names in unity builds
#rnx

diff --git a/Engine/Source/Programs/UnrealHeaderTool/Private/ClassDeclarationMetaData.cpp b/Engine/Source/Programs/UnrealHeaderTool/Private/ClassDeclarationMetaData.cpp
--- a/Engine/Source/Programs/UnrealHeaderTool/Private/ClassDeclarationMetaData.cpp
+++ b/Engine/Source/Programs/UnrealHeaderTool/Private/ClassDeclarationMetaData.cpp
@@ -15,7 +15,6 @@
 namespace
 {
 	static const FName NAME_IgnoreCategoryKeywordsInSubclasses(TEXT("IgnoreCategoryKeywordsInSubclasses"));
-	static const FName NAME_IsConversionRoot(TEXT("IsConversionRoot"));
 	bool IsActorClass(UClass *Class)
 	{
 		while (Class)
@@ -293,7 +292,7 @@ void FClassDeclarationMetaData::ParseClassProperties(TArray<FPropertySpecifier>&
 
 			case EClassMetadataSpecifier::ConversionRoot:
 
-				MetaData.Add(NAME_IsConversionRoot, TEXT("true"));
+				MetaData.Add(FHeaderParserNames::NAME_IsConversionRoot, TEXT("true"));
 				break;
 
 			default:
diff --git a/Engine/Source/Programs/UnrealHeaderTool/Private/HeaderParser.cpp b/Engine/Source/Programs/UnrealHeaderTool/Private/HeaderParser.cpp
--- a/Engine/Source/Programs/UnrealHeaderTool/Private/HeaderParser.cpp
+++ b/Engine/Source/Programs/UnrealHeaderTool/Private/HeaderParser.cpp
@@ -86,20 +86,22 @@ namespace
 	static const FName NAME_ArraySizeEnum(TEXT("ArraySizeEnum"));
 	static const FName NAME_ClassGroupNames(TEXT("ClassGroupNames"));
 	static const FName NAME_AutoCollapseCategories(TEXT("AutoCollapseCategories"));
-	static const FName NAME_HideCategories(TEXT("HideCategories"));
-	static const FName NAME_ShowCategories(TEXT("ShowCategories"));
-	static const FName NAME_SparseClassDataTypes(TEXT("SparseClassDataTypes"));
 	static const FName NAME_HideFunctions(TEXT("HideFunctions"));
 	static const FName NAME_AutoExpandCategories(TEXT("AutoExpandCategories"));
-	static const FName NAME_IsConversionRoot(TEXT("IsConversionRoot"));
 	static const FName NAME_EditInline(TEXT("EditInline"));
 	static const FName NAME_IncludePath(TEXT("IncludePath"));
 	static const FName NAME_ModuleRelativePath(TEXT("ModuleRelativePath"));
 	static const FName NAME_CannotImplementInterfaceInBlueprint(TEXT("CannotImplementInterfaceInBlueprint"));
 	static const FName NAME_UIMin(TEXT("UIMin"));
 	static const FName NAME_UIMax(TEXT("UIMax"));
+	static const FName NAME_BlueprintType(TEXT("BlueprintType"));
 }
 
+const FName FHeaderParserNames::NAME_HideCategories(TEXT("HideCategories"));
+const FName FHeaderParserNames::NAME_ShowCategories(TEXT("ShowCategories"));
+const FName FHeaderParserNames::NAME_SparseClassDataTypes(TEXT("SparseClassDataTypes"));
+const FName FHeaderParserNames::NAME_IsConversionRoot(TEXT("IsConversionRoot"));
+
 EGeneratedCodeVersion FHeaderParser::DefaultGeneratedCodeVersion = EGeneratedCodeVersion::V1;
 TArray<FString> FHeaderParser::StructsWithNoPrefix;
 TArray<FString> FHeaderParser::StructsWithTPrefix;
@@ -1213,7 +1215,7 @@ namespace
 		}
 		else if (const FStructProperty* StructProperty = CastField<const FStructProperty>(Property))
 		{
-			return (StructProperty->Struct->GetBoolMetaDataHierarchical(TEXT("BlueprintType")));
+			return (StructProperty->Struct->GetBoolMetaDataHierarchical(NAME_BlueprintType));
 		}
 
 		const bool bSupportedType = Property->IsA<FInterfaceProperty>()
@@ -1692,8 +1694,7 @@ UEnum* FHeaderParser::CompileEnum()
 		GEnumUnderlyingTypes.Add(Enum, UnderlyingType);
 	}
 
-	static const FName BlueprintTypeName = TEXT("BlueprintType");
-	if (UnderlyingType != EUnderlyingEnumType::uint8 && EnumToken.MetaData.Contains(BlueprintTypeName))
+	if (UnderlyingType != EUnderlyingEnumType::uint8 && EnumToken.MetaData.Contains(NAME_BlueprintType))
 	{
 		FError::Throwf(TEXT("Invalid BlueprintType enum base - currently only uint8 supported"));
 	}
@@ -1890,7 +1891,7 @@ UEnum* FHeaderParser::CompileEnum()
 
 	CheckDocumentationPolicyForEnum(Enum, EnumValueMetaData, EntryMetaData);
 
-	if (!Enum->IsValidEnumValue(0) && EnumToken.MetaData.Contains(BlueprintTypeName))
+	if (!Enum->IsValidEnumValue(0) && EnumToken.MetaData.Contains(NAME_BlueprintType))
 	{
 		UE_LOG_WARNING_UHT(TEXT("'%s' does not have a 0 entry! (This is a problem when the enum is initalized by default)"), *Enum->GetName());
 	}
@@ -3144,14 +3145,14 @@ void FHeaderParser::CheckSparseClassData(const UStruct* StructToCheck)
 	if (!ClassToCheck)
 	{
 		// make sure we don't try to have sparse class data inside of a struct instead of a class
-		if (StructToCheck->HasMetaData(NAME_SparseClassDataTypes))
+		if (StructToCheck->HasMetaData(FHeaderParserNames::NAME_SparseClassDataTypes))
 		{
 			UE_LOG_ERROR_UHT(TEXT("%s contains sparse class data but is not a class."), *StructToCheck->GetName());
 		}
 		return;
 	}
 
-	if (!ClassToCheck->HasMetaData(NAME_SparseClassDataTypes))
+	if (!ClassToCheck->HasMetaData(FHeaderParserNames::NAME_SparseClassDataTypes))
 	{
 		return;
 	}
@@ -6170,9 +6171,9 @@ UClass* FHeaderParser::CompileClassDeclaration(FClasses& AllClasses)
 	MetaData.Append(ClassDeclarationData->MetaData);
 	if (ClassDeclarationData->ClassGroupNames.Num()) { MetaData.Add(NAME_ClassGroupNames, FString::Join(ClassDeclarationData->ClassGroupNames, TEXT(" "))); }
 	if (ClassDeclarationData->AutoCollapseCategories.Num()) { MetaData.Add(NAME_AutoCollapseCategories, FString::Join(ClassDeclarationData->AutoCollapseCategories, TEXT(" "))); }
-	if (ClassDeclarationData->HideCategories.Num()) { MetaData.Add(NAME_HideCategories, FString::Join(ClassDeclarationData->HideCategories, TEXT(" "))); }
-	if (ClassDeclarationData->ShowSubCatgories.Num()) { MetaData.Add(NAME_ShowCategories, FString::Join(ClassDeclarationData->ShowSubCatgories, TEXT(" "))); }
-	if (ClassDeclarationData->SparseClassDataTypes.Num()) { MetaData.Add(NAME_SparseClassDataTypes, FString::Join(ClassDeclarationData->SparseClassDataTypes, TEXT(" "))); }
+	if (ClassDeclarationData->HideCategories.Num()) { MetaData.Add(FHeaderParserNames::NAME_HideCategories, FString::Join(ClassDeclarationData->HideCategories, TEXT(" "))); }
+	if (ClassDeclarationData->ShowSubCatgories.Num()) { MetaData.Add(FHeaderParserNames::NAME_ShowCategories, FString::Join(ClassDeclarationData->ShowSubCatgories, TEXT(" "))); }
+	if (ClassDeclarationData->SparseClassDataTypes.Num()) { MetaData.Add(FHeaderParserNames::NAME_SparseClassDataTypes, FString::Join(ClassDeclarationData->SparseClassDataTypes, TEXT(" "))); }
 	if (ClassDeclarationData->HideFunctions.Num()) { MetaData.Add(NAME_HideFunctions, FString::Join(ClassDeclarationData->HideFunctions, TEXT(" "))); }
 	if (ClassDeclarationData->AutoExpandCategories.Num()) { MetaData.Add(NAME_AutoExpandCategories, FString::Join(ClassDeclarationData->AutoExpandCategories, TEXT(" "))); }
 
@@ -6376,7 +6377,7 @@ void FHeaderParser::CompileInterfaceDeclaration(FClasses& AllClasses)
 
 			case EInterfaceSpecifier::ConversionRoot:
 			{
-				MetaData.Add(NAME_IsConversionRoot, TEXT("true"));
+				MetaData.Add(FHeaderParserNames::NAME_IsConversionRoot, TEXT("true"));
 			}
 			break;
 		}
@@ -7751,8 +7752,7 @@ struct FExposeOnSpawnValidator
 
 		if (!ProperNativeType && (CPT_Struct == Property.Type) && Property.Struct)
 		{
-			static const FName BlueprintTypeName(TEXT("BlueprintType"));
-			ProperNativeType |= Property.Struct->GetBoolMetaData(BlueprintTypeName);
+			ProperNativeType |= Property.Struct->GetBoolMetaData(NAME_BlueprintType);
 		}
 
 		return ProperNativeType;
@@ -7874,7 +7874,7 @@ void FHeaderParser::CompileVariableDeclaration(FClasses& AllClasses, UStruct* St
 
 		if (NewProperty->HasAnyPropertyFlags(CPF_BlueprintVisible))
 		{
-			if (Struct->IsA<UScriptStruct>() && !Struct->GetBoolMetaDataHierarchical(TEXT("BlueprintType")))
+			if (Struct->IsA<UScriptStruct>() && !Struct->GetBoolMetaDataHierarchical(NAME_BlueprintType))
 			{
 				UE_LOG_ERROR_UHT(TEXT("Cannot expose property to blueprints in a struct that is not a BlueprintType. %s.%s"), *Struct->GetName(), *NewProperty->GetName());
 			}
@@ -9944,17 +9944,17 @@ void FHeaderParser::ResetClassData()
 		}
 
 		// Copy special categories from parent
-		if (SuperClass->HasMetaData(NAME_HideCategories))
+		if (SuperClass->HasMetaData(FHeaderParserNames::NAME_HideCategories))
 		{
-			CurrentClass->SetMetaData(NAME_HideCategories, *SuperClass->GetMetaData(NAME_HideCategories));
+			CurrentClass->SetMetaData(FHeaderParserNames::NAME_HideCategories, *SuperClass->GetMetaData(FHeaderParserNames::NAME_HideCategories));
 		}
-		if (SuperClass->HasMetaData(NAME_ShowCategories))
+		if (SuperClass->HasMetaData(FHeaderParserNames::NAME_ShowCategories))
 		{
-			CurrentClass->SetMetaData(NAME_ShowCategories, *SuperClass->GetMetaData(NAME_ShowCategories));
+			CurrentClass->SetMetaData(FHeaderParserNames::NAME_ShowCategories, *SuperClass->GetMetaData(FHeaderParserNames::NAME_ShowCategories));
 		}
-		if (SuperClass->HasMetaData(NAME_SparseClassDataTypes))
+		if (SuperClass->HasMetaData(FHeaderParserNames::NAME_SparseClassDataTypes))
 		{
-			CurrentClass->SetMetaData(NAME_SparseClassDataTypes, *SuperClass->GetMetaData(NAME_SparseClassDataTypes));
+			CurrentClass->SetMetaData(FHeaderParserNames::NAME_SparseClassDataTypes, *SuperClass->GetMetaData(FHeaderParserNames::NAME_SparseClassDataTypes));
 		}
 		if (SuperClass->HasMetaData(NAME_HideFunctions))
 		{
diff --git a/Engine/Source/Programs/UnrealHeaderTool/Private/ParserClass.cpp b/Engine/Source/Programs/UnrealHeaderTool/Private/ParserClass.cpp
--- a/Engine/Source/Programs/UnrealHeaderTool/Private/ParserClass.cpp
+++ b/Engine/Source/Programs/UnrealHeaderTool/Private/ParserClass.cpp
@@ -5,12 +5,6 @@
 #include "UObject/Package.h"
 #include "Templates/Casts.h"
 
-namespace
-{
-	static const FName NAME_HideCategories(TEXT("HideCategories"));
-	static const FName NAME_ShowCategories(TEXT("ShowCategories"));
-	static const FName NAME_SparseClassDataTypes(TEXT("SparseClassDataTypes"));
-}
 const FName FClass::NAME_ReplaceConverted(TEXT("ReplaceConverted"));
 
 FString FClass::GetNameWithPrefix(EEnforceInterfacePrefix::Type EnforceInterfacePrefix) const
@@ -77,27 +71,27 @@ TArray<FClass*> FClass::GetInterfaceTypes() const
 
 void FClass::GetHideCategories(TArray<FString>& OutHideCategories) const
 {
-	if (HasMetaData(NAME_HideCategories))
+	if (HasMetaData(FHeaderParserNames::NAME_HideCategories))
 	{
-		const FString& HideCategories = GetMetaData(NAME_HideCategories);
+		const FString& HideCategories = GetMetaData(FHeaderParserNames::NAME_HideCategories);
 		HideCategories.ParseIntoArray(OutHideCategories, TEXT(" "), true);
 	}
 }
 
 void FClass::GetShowCategories(TArray<FString>& OutShowCategories) const
 {
-	if (HasMetaData(NAME_ShowCategories))
+	if (HasMetaData(FHeaderParserNames::NAME_ShowCategories))
 	{
-		const FString& ShowCategories = GetMetaData(NAME_ShowCategories);
+		const FString& ShowCategories = GetMetaData(FHeaderParserNames::NAME_ShowCategories);
 		ShowCategories.ParseIntoArray(OutShowCategories, TEXT(" "), true);
 	}
 }
 
 void FClass::GetSparseClassDataTypes(TArray<FString>& OutSparseClassDataTypes) const
 {
-	if (HasMetaData(NAME_SparseClassDataTypes))
+	if (HasMetaData(FHeaderParserNames::NAME_SparseClassDataTypes))
 	{
-		const FString& SparseClassDataTypes = GetMetaData(NAME_SparseClassDataTypes);
+		const FString& SparseClassDataTypes = GetMetaData(FHeaderParserNames::NAME_SparseClassDataTypes);
 		SparseClassDataTypes.ParseIntoArray(OutSparseClassDataTypes, TEXT(" "), true);
 	}
 }
diff --git a/Engine/Source/Programs/UnrealHeaderTool/Private/UnrealHeaderTool.h b/Engine/Source/Programs/UnrealHeaderTool/Private/UnrealHeaderTool.h
--- a/Engine/Source/Programs/UnrealHeaderTool/Private/UnrealHeaderTool.h
+++ b/Engine/Source/Programs/UnrealHeaderTool/Private/UnrealHeaderTool.h
@@ -1,3 +1,11 @@
 // Copyright Epic Games, Inc. All Rights Reserved.
 
 #pragma once
+
+struct FHeaderParserNames
+{
+	static const FName NAME_IsConversionRoot;
+	static const FName NAME_HideCategories;
+	static const FName NAME_ShowCategories;
+	static const FName NAME_SparseClassDataTypes;
+};
\ No newline at end of file
