UHT optimizations part 5 - Minor optimizations and clean up of code generation
#rnx

diff --git a/Engine/Source/Programs/UnrealHeaderTool/Private/CodeGenerator.cpp b/Engine/Source/Programs/UnrealHeaderTool/Private/CodeGenerator.cpp
--- a/Engine/Source/Programs/UnrealHeaderTool/Private/CodeGenerator.cpp
+++ b/Engine/Source/Programs/UnrealHeaderTool/Private/CodeGenerator.cpp
@@ -1,7 +1,7 @@
 // Copyright Epic Games, Inc. All Rights Reserved.
 
-#include "UnrealHeaderTool.h"
 #include "CoreMinimal.h"
+#include "UnrealHeaderTool.h"
 #include "Misc/AssertionMacros.h"
 #include "HAL/PlatformProcess.h"
 #include "Templates/UnrealTemplate.h"
@@ -85,7 +85,19 @@ bool HasIdentifierExactMatch(const TCHAR* StringBegin, const TCHAR* StringEnd, c
 
 namespace
 {
-	static FString AsTEXT(FString InStr)
+	static const FName NAME_SerializeToFArchive("SerializeToFArchive");
+	static const FName NAME_SerializeToFStructuredArchive("SerializeToFStructuredArchive");
+	static const FName NAME_ObjectInitializerConstructorDeclared("ObjectInitializerConstructorDeclared");
+	static const FName NAME_InitializeStaticSearchableValues("InitializeStaticSearchableValues");
+	static const FName NAME_OverrideNativeName("OverrideNativeName");
+	static const FName NAME_NoGetter("NoGetter");
+	static const FName NAME_GetByRef("GetByRef");
+
+	static const FString STRING_StructPackage(TEXT("StructPackage"));
+
+	static const int32 HEADER_GENERATED_DELEGATE_SIGNATURE_SUFFIX_LENGTH = FString(HEADER_GENERATED_DELEGATE_SIGNATURE_SUFFIX).Len(); 
+
+	static FString AsTEXT(const FString& InStr)
 	{
 		return FString::Printf(TEXT("TEXT(\"%s\")"), *InStr);
 	}
@@ -165,9 +177,9 @@ namespace
 	{
 		const TArray<FRepRecord>& ClassReps = Class->ClassReps;
 		
-		const FString ClassName = FString::Printf(TEXT("%s%s"), Class->GetPrefixCPP(), *Class->GetName());
-		const FString ParentClassName = ParentRepClass ? FString::Printf(TEXT("%s%s"), ParentRepClass->GetPrefixCPP(), *ParentRepClass->GetName()) : FString();
-		const FString NetfieldStart = ParentRepClass ? FString::Printf(TEXT("(uint16)%s::ENetFields_Private::NETFIELD_REP_END + (uint16)1"), *ParentClassName, *ParentClassName) : FString(TEXT("0"));
+		const FString NetfieldStart = ParentRepClass
+				? FString::Printf(TEXT("(uint16)%s%s::ENetFields_Private::NETFIELD_REP_END + (uint16)1"), ParentRepClass->GetPrefixCPP(), *ParentRepClass->GetName())
+				: FString(TEXT("0"));
 
 		FUHTStringBuilder NetFieldBuilder;
 		NetFieldBuilder.Logf(TEXT(""
@@ -182,7 +194,7 @@ namespace
 		for (int32 ClassRepIndex = ClassReps.Num() - Class->NetProperties.Num() - Class->NetFields.Num(); ClassRepIndex < ClassReps.Num(); ++ClassRepIndex)
 		{
 			const FRepRecord& ClassRep = ClassReps[ClassRepIndex];
-			const FString& PropertyName = ClassRep.Property->GetName();
+			const FString PropertyName = ClassRep.Property->GetName();
 
 			if (ClassRep.Property->ArrayDim == 1)
 			{
@@ -216,11 +228,9 @@ namespace
 		}
 
 		const FProperty* LastProperty = ClassReps.Last().Property;
-		const FString NetfieldEnd = LastProperty->ArrayDim > 1 ?
+		NetFieldBuilder.Log( LastProperty->ArrayDim > 1 ?
 			FString::Printf(TEXT("\t\tNETFIELD_REP_END=%s_STATIC_ARRAY_END,\r\n"), *LastProperty->GetName()) :
-			FString::Printf(TEXT("\t\tNETFIELD_REP_END=%s,\r\n"), *LastProperty->GetName());
-
-		NetFieldBuilder.Log(*NetfieldEnd);
+			FString::Printf(TEXT("\t\tNETFIELD_REP_END=%s,\r\n"), *LastProperty->GetName()));
 
 		NetFieldBuilder.Log(TEXT("\t};"));
 
@@ -240,6 +250,8 @@ namespace
 			API);
 	}
 
+	static const FString GetLifetimeReplicatedPropsStr(TEXT("GetLifetimeReplicatedProps"));
+
 	static void WriteReplicatedMacroData(
 		const ClassDefinitionRange& ClassRange,
 		const TCHAR* ClassCPPName,
@@ -249,7 +261,7 @@ namespace
 		FOutputDevice& Writer,
 		const FUnrealSourceFile& SourceFile)
 	{
-		const bool bHasGetLifetimeReplicatedProps = HasIdentifierExactMatch(ClassRange.Start, ClassRange.End, TEXT("GetLifetimeReplicatedProps"));
+		const bool bHasGetLifetimeReplicatedProps = HasIdentifierExactMatch(ClassRange.Start, ClassRange.End, GetLifetimeReplicatedPropsStr);
 
 		if (!bHasGetLifetimeReplicatedProps)
 		{
@@ -386,23 +398,18 @@ void ConvertToBuildIncludePath(const UPackage* Package, FString& LocalPath)
 }
 
 /**
- *	Helper function for finding the location of a package
- *	This is required as source now lives in several possible directories
+ *	Helper function to retrieve the package manifest
  *
  *	@param	InPackage		The name of the package of interest
- *	@param	OutLocation		The location of the given package, if found
- *  @param	OutHeaderLocation	The directory where generated headers should be placed
  *
- *	@return	bool			true if found, false if not
+ *	@return	The manifest if found
  */
-bool FindPackageLocation(const TCHAR* InPackage, FString& OutLocation, FString& OutHeaderLocation)
+FManifestModule* GetPackageManifest(const FString& CheckPackage)
 {
 	// Mapping of processed packages to their locations
 	// An empty location string means it was processed but not found
 	static TMap<FString, FManifestModule*> CheckedPackageList;
 
-	FString CheckPackage(InPackage);
-
 	FManifestModule* ModuleInfoPtr = CheckedPackageList.FindRef(CheckPackage);
 
 	if (!ModuleInfoPtr)
@@ -415,22 +422,14 @@ bool FindPackageLocation(const TCHAR* InPackage, FString& OutLocation, FString&
 		}
 	}
 
-	if (!ModuleInfoPtr)
-	{
-		return false;
-	}
-
-	OutLocation       = ModuleInfoPtr->BaseDirectory;
-	OutHeaderLocation = ModuleInfoPtr->GeneratedIncludeDirectory;
-	return true;
+	return ModuleInfoPtr;
 }
 
-
-FString Macroize(const TCHAR* MacroName, const TCHAR* StringToMacroize)
+FString Macroize(const TCHAR* MacroName, FString&& StringToMacroize)
 {
 	FScopedDurationTimer Tracker(GMacroizeTime);
 
-	FString Result = StringToMacroize;
+	FString Result(MoveTemp(StringToMacroize));
 	if (Result.Len())
 	{
 		Result.ReplaceInline(TEXT("\r\n"), TEXT("\n"), ESearchCase::CaseSensitive);
@@ -450,7 +449,12 @@ FString Macroize(const TCHAR* MacroName, const TCHAR* StringToMacroize)
 		}
 		Result.ReplaceInline(TEXT("\n"), TEXT("\r\n"), ESearchCase::CaseSensitive);
 	}
-	return FString::Printf(TEXT("#define %s%s\r\n"), MacroName, Result.Len() ? TEXT(" \\") : TEXT("")) + Result;
+	return FString::Printf(TEXT("#define %s%s\r\n%s"), MacroName, Result.Len() ? TEXT(" \\") : TEXT(""), *Result);
+}
+
+FString Macroize(const TCHAR* MacroName, const TCHAR* StringToMacroize)
+{
+	return Macroize(MacroName, FString(StringToMacroize));
 }
 
 /** Generates a Hash tag string for the specified field */
@@ -468,7 +472,7 @@ static FString GetGeneratedCodeHashTag(void* Field)
 struct FParmsAndReturnProperties
 {
 	FParmsAndReturnProperties()
-		: Return(NULL)
+		: Return(nullptr)
 	{
 	}
 
@@ -678,34 +682,26 @@ static FString OutputMetaDataCodeForObject(FOutputDevice& OutDeclaration, FOutpu
 		MetaData = GenerateMetadataMapForField(Object.ToField());
 	}
 
-	FUHTStringBuilder MetaDataOutput;
+	FString Result;
 	if (MetaData.Num())
 	{
-		typedef TKeyValuePair<FName, FString> KVPType;
+		typedef TKeyValuePair<FName, FString*> KVPType;
 		TArray<KVPType> KVPs;
+		KVPs.Reserve(MetaData.Num());
 		for (TPair<FName, FString>& KVP : MetaData)
 		{
-			KVPs.Add(KVPType(KVP.Key, KVP.Value));
+			KVPs.Add(KVPType(KVP.Key, &KVP.Value));
 		}
 
 		// We sort the metadata here so that we can get consistent output across multiple runs
 		// even when metadata is added in a different order
 		Algo::SortBy(KVPs, &KVPType::Key, FNameLexicalLess());
 
-		for (const KVPType& KVP : KVPs)
-		{
-			MetaDataOutput.Logf(TEXT("%s\t{ %s, %s },\r\n"), Spaces, *CreateUTF8LiteralString(KVP.Key.ToString()), *CreateUTF8LiteralString(KVP.Value));
-		}
-	}
-
-	FString Result;
-	if (MetaDataOutput.Len())
-	{
 		FString MetaDataBlockNameWithoutScope = MetaDataBlockName;
 		int32 ScopeIndex = MetaDataBlockNameWithoutScope.Find(TEXT("::"), ESearchCase::CaseSensitive);
 		if (ScopeIndex != INDEX_NONE)
 		{
-			MetaDataBlockNameWithoutScope = MetaDataBlockNameWithoutScope.RightChop(ScopeIndex + 2);
+			MetaDataBlockNameWithoutScope.RightChopInline(ScopeIndex + 2, false);
 		}
 
 		OutDeclaration.Log (TEXT("#if WITH_METADATA\r\n"));
@@ -714,7 +710,12 @@ static FString OutputMetaDataCodeForObject(FOutputDevice& OutDeclaration, FOutpu
 
 		Out.Log (TEXT("#if WITH_METADATA\r\n"));
 		Out.Logf(TEXT("%sconst UE4CodeGen_Private::FMetaDataPairParam %s[] = {\r\n"), Spaces, MetaDataBlockName);
-		Out.Log (*MetaDataOutput);
+
+		for (const KVPType& KVP : KVPs)
+		{
+			Out.Logf(TEXT("%s\t{ %s, %s },\r\n"), Spaces, *CreateUTF8LiteralString(KVP.Key.ToString()), *CreateUTF8LiteralString(*KVP.Value));
+		}
+
 		Out.Logf(TEXT("%s};\r\n"), Spaces);
 		Out.Log (TEXT("#endif\r\n"));
 
@@ -997,8 +998,7 @@ FString FNativeClassHeaderGenerator::GetSingletonName(UField* Item, bool bRequir
 	// We don't need to export UFunction externs, though we may need the externs for UDelegateFunctions
 	if (UniqueCrossModuleReferences && (!Item->IsA<UFunction>() || Item->IsA<UDelegateFunction>()))
 	{
-		const FString& Extern = Cache.GetExternDecl();
-		UniqueCrossModuleReferences->Add(*Extern);
+		UniqueCrossModuleReferences->Add(Cache.GetExternDecl());
 	}
 
 	return Result;
@@ -1038,7 +1038,7 @@ void FNativeClassHeaderGenerator::PropertyNew(FOutputDevice& DeclOut, FOutputDev
 	if (ScopeIndex != INDEX_NONE)
 	{
 		Scope = NameWithoutScope.Left(ScopeIndex) + TEXT("_");
-		NameWithoutScope = NameWithoutScope.RightChop(ScopeIndex + 2);
+		NameWithoutScope.RightChopInline(ScopeIndex + 2, false);
 	}
 
 	if (FByteProperty* TypedProp = CastField<FByteProperty>(Prop))
@@ -1805,19 +1805,22 @@ inline FString GetEventStructParamsName(UObject* Outer, const TCHAR* FunctionNam
 
 void FNativeClassHeaderGenerator::OutputProperty(FOutputDevice& DeclOut, FOutputDevice& Out, const TCHAR* Scope, TArray<FPropertyNamePointerPair>& PropertyNamesAndPointers, FProperty* Prop, const TCHAR* DeclSpaces, const TCHAR* Spaces)
 {
-	FString PropName = Prop->GetName();
-
-	FString PropVariableName = FString::Printf(TEXT("%sNewProp_%s"), Scope, *PropName);
-
 	// Helper to handle the creation of the underlying properties if they're enum properties
-	auto HandleUnderlyingEnumProperty = [this, &PropertyNamesAndPointers, &DeclOut, &Out, DeclSpaces, Spaces](FProperty* LocalProp, const FString& OuterName)
+	auto HandleUnderlyingEnumProperty = [this, &PropertyNamesAndPointers, &DeclOut, &Out, DeclSpaces, Spaces](FProperty* LocalProp, FString OuterName)
 	{
 		if (FEnumProperty* EnumProp = CastField<FEnumProperty>(LocalProp))
 		{
 			FString PropVarName = OuterName + TEXT("_Underlying");
+
+			PropertyNamesAndPointers.Emplace(MoveTemp(OuterName), LocalProp);
+
 			PropertyNew(DeclOut, Out, EnumProp->UnderlyingProp, TEXT("0"), *PropVarName, DeclSpaces, Spaces);
 			PropertyNamesAndPointers.Emplace(MoveTemp(PropVarName), EnumProp->UnderlyingProp);
 		}
+		else
+		{
+			PropertyNamesAndPointers.Emplace(MoveTemp(OuterName), LocalProp);
+		}
 	};
 
 	{
@@ -1831,7 +1834,7 @@ void FNativeClassHeaderGenerator::OutputProperty(FOutputDevice& DeclOut, FOutput
 			FString FunctionName = Function->GetName();
 			if( Function->HasAnyFunctionFlags( FUNC_Delegate ) )
 			{
-				FunctionName = FunctionName.LeftChop( FString( HEADER_GENERATED_DELEGATE_SIGNATURE_SUFFIX ).Len() );
+				FunctionName.LeftChopInline(HEADER_GENERATED_DELEGATE_SIGNATURE_SUFFIX_LENGTH, false );
 			}
 
 			SourceStruct = GetEventStructParamsName(Function->GetOuter(), *FunctionName);
@@ -1841,17 +1844,18 @@ void FNativeClassHeaderGenerator::OutputProperty(FOutputDevice& DeclOut, FOutput
 			SourceStruct = NameLookupCPP.GetNameCPP(CastChecked<UStruct>(Prop->GetOwner<UObject>()));
 		}
 
-		FString PropNameDep = PropName;
+		FString PropName = Prop->GetName();
+		FString PropVariableName = FString::Printf(TEXT("%sNewProp_%s"), Scope, *PropName);
+
 		if (Prop->HasAllPropertyFlags(CPF_Deprecated))
 		{
-			 PropNameDep += TEXT("_DEPRECATED");
+			 PropName += TEXT("_DEPRECATED");
 		}
 
-		FString PropMacroOuterClass = FString::Printf(TEXT("STRUCT_OFFSET(%s, %s)"), *SourceStruct, *PropNameDep);
+		FString PropMacroOuterClass = FString::Printf(TEXT("STRUCT_OFFSET(%s, %s)"), *SourceStruct, *PropName);
 
 		PropertyNew(DeclOut, Out, Prop, *PropMacroOuterClass, *PropVariableName, DeclSpaces, Spaces, *SourceStruct);
-		PropertyNamesAndPointers.Emplace(PropVariableName, Prop);
-		HandleUnderlyingEnumProperty(Prop, PropVariableName);
+		HandleUnderlyingEnumProperty(Prop, MoveTemp(PropVariableName));
 	}
 
 	if (FArrayProperty* ArrayProperty = CastField<FArrayProperty>(Prop))
@@ -1859,8 +1863,7 @@ void FNativeClassHeaderGenerator::OutputProperty(FOutputDevice& DeclOut, FOutput
 		FString InnerVariableName = FString::Printf(TEXT("%sNewProp_%s_Inner"), Scope, *ArrayProperty->Inner->GetName());
 
 		PropertyNew(DeclOut, Out, ArrayProperty->Inner, TEXT("0"), *InnerVariableName, DeclSpaces, Spaces);
-		PropertyNamesAndPointers.Emplace(InnerVariableName, ArrayProperty->Inner);
-		HandleUnderlyingEnumProperty(ArrayProperty->Inner, InnerVariableName);
+		HandleUnderlyingEnumProperty(ArrayProperty->Inner, MoveTemp(InnerVariableName));
 	}
 
 	else if (FMapProperty* MapProperty = CastField<FMapProperty>(Prop))
@@ -1872,12 +1875,10 @@ void FNativeClassHeaderGenerator::OutputProperty(FOutputDevice& DeclOut, FOutput
 		FString ValueVariableName = FString::Printf(TEXT("%sNewProp_%s_ValueProp"), Scope, *Value->GetName());
 
 		PropertyNew(DeclOut, Out, Key, TEXT("0"), *KeyVariableName, DeclSpaces, Spaces);
-		PropertyNamesAndPointers.Emplace(KeyVariableName, Key);
-		HandleUnderlyingEnumProperty(Key, KeyVariableName);
+		HandleUnderlyingEnumProperty(Key, MoveTemp(KeyVariableName));
 
 		PropertyNew(DeclOut, Out, Value, TEXT("1"), *ValueVariableName, DeclSpaces, Spaces);
-		PropertyNamesAndPointers.Emplace(ValueVariableName, Value);
-		HandleUnderlyingEnumProperty(Value, ValueVariableName);
+		HandleUnderlyingEnumProperty(Value, MoveTemp(ValueVariableName));
 	}
 
 	else if (FSetProperty* SetProperty = CastField<FSetProperty>(Prop))
@@ -1887,8 +1888,7 @@ void FNativeClassHeaderGenerator::OutputProperty(FOutputDevice& DeclOut, FOutput
 		FString ElementVariableName = FString::Printf(TEXT("%sNewProp_%s_ElementProp"), Scope, *Inner->GetName());
 
 		PropertyNew(DeclOut, Out, Inner, TEXT("0"), *ElementVariableName, DeclSpaces, Spaces);
-		PropertyNamesAndPointers.Emplace(ElementVariableName, Inner);
-		HandleUnderlyingEnumProperty(Inner, ElementVariableName);
+		HandleUnderlyingEnumProperty(Inner, MoveTemp(ElementVariableName));
 	}
 }
 
@@ -1978,24 +1978,43 @@ static TArray<UScriptStruct*> FindNoExportStructs(UStruct* Start)
 	return Result;
 }
 
-FString FNativeClassHeaderGenerator::GetPackageSingletonName(const UPackage* InPackage)
+struct FPackageSingletonStrings
+{
+	FPackageSingletonStrings(FString&& InPackageSingletonName)
+		: PackageSingletonName(InPackageSingletonName)
+		, PackageUniqueCrossModuleReference(FString::Printf(TEXT("\tUPackage* %s;\r\n"), *PackageSingletonName))
+	{
+	}
+
+	FString PackageSingletonName;
+	FString PackageUniqueCrossModuleReference;
+};
+
+static TMap<const UPackage*, FPackageSingletonStrings> PackageSingletonNames;
+const FString& FNativeClassHeaderGenerator::GetPackageSingletonName(const UPackage* InPackage)
+{
+	FPackageSingletonStrings* PackageSingletonStrings = PackageSingletonNames.Find(InPackage);
+	if (UNLIKELY(PackageSingletonStrings == nullptr))
 {
-	const TCHAR* ClassString = NameLookupCPP.GetNameCPP(UPackage::StaticClass());
+		static const TCHAR* ClassString = NameLookupCPP.GetNameCPP(UPackage::StaticClass());
 
-	FString Result = FString::Printf(TEXT("Z_Construct_%s_%s()"), ClassString, *InPackage->GetName().Replace(TEXT("/"), TEXT("_")));
+		FString PackageName = InPackage->GetName();
+		PackageName.ReplaceInline(TEXT("/"), TEXT("_"), ESearchCase::CaseSensitive);
+
+		PackageSingletonStrings = &PackageSingletonNames.Add(InPackage, FPackageSingletonStrings(FString::Printf(TEXT("Z_Construct_%s_%s()"), ClassString, *PackageName)));
+	}
 
 	if (UniqueCrossModuleReferences)
 	{
-		UniqueCrossModuleReferences->Add(FString::Printf(TEXT("\tUPackage* %s;\r\n"), *Result));
+		UniqueCrossModuleReferences->Add(PackageSingletonStrings->PackageUniqueCrossModuleReference);
 	}
 
-	return Result;
+	return PackageSingletonStrings->PackageSingletonName;
 }
 
 void FNativeClassHeaderGenerator::ExportGeneratedPackageInitCode(FOutputDevice& Out, const TCHAR* InDeclarations, const UPackage* InPackage, uint32 Hash)
 {
-	FString ApiString = GetAPIString();
-	FString SingletonName = GetPackageSingletonName(InPackage);
+	const FString& SingletonName = GetPackageSingletonName(InPackage);
 
 	const TArray<UField*>* SingletonsToOutput = GPackageSingletons.Find(InPackage);
 	if (SingletonsToOutput)
@@ -2022,7 +2041,7 @@ void FNativeClassHeaderGenerator::ExportGeneratedPackageInitCode(FOutputDevice&
 		Out.Logf(TEXT("\t\t\tstatic UObject* (*const SingletonFuncArray[])() = {\r\n"));
 		for (UField* ScriptType : *SingletonsToOutput)
 		{
-			const FString& Name = FTypeSingletonCache::Get(ScriptType, true).GetName().LeftChop(2);
+			const FString Name = FTypeSingletonCache::Get(ScriptType, true).GetName().LeftChop(2);
 
 			Out.Logf(TEXT("\t\t\t\t(UObject* (*)())%s,\r\n"), *Name);
 		}
@@ -2062,8 +2081,7 @@ void FNativeClassHeaderGenerator::ExportNativeGeneratedInitCode(FOutputDevice& O
 	const bool   bIsDynamic   = FClass::IsDynamic(Class);
 	const TCHAR* ClassNameCPP = NameLookupCPP.GetNameCPP(Class);
 
-	FUHTStringBuilder BodyText;
-	FString ApiString = GetAPIString();
+	const FString& ApiString = GetAPIString();
 
 	TSet<FName> AlreadyIncludedNames;
 	TArray<UFunction*> FunctionsToExport;
@@ -2209,7 +2227,7 @@ void FNativeClassHeaderGenerator::ExportNativeGeneratedInitCode(FOutputDevice& O
 			FClassMetaData* ClassMetaData = GScriptHelper.FindClassData(Class);
 			if (ClassMetaData && ClassMetaData->bObjectInitializerConstructorDeclared)
 			{
-				MetaDataMap->Add(FName(TEXT("ObjectInitializerConstructorDeclared")), TEXT(""));
+				MetaDataMap->Add(NAME_ObjectInitializerConstructorDeclared, FString());
 			}
 		}
 
@@ -2303,7 +2321,7 @@ void FNativeClassHeaderGenerator::ExportNativeGeneratedInitCode(FOutputDevice& O
 		}
 		else
 		{
-			const FString DynamicClassPackageName = FClass::GetTypePackageName(Class);
+			const FString& DynamicClassPackageName = FClass::GetTypePackageName(Class);
 			GeneratedClassRegisterFunctionText.Logf(TEXT("\t\tUPackage* OuterPackage = FindOrConstructDynamicTypePackage(TEXT(\"%s\"));\r\n"), *DynamicClassPackageName);
 			GeneratedClassRegisterFunctionText.Logf(TEXT("\t\tUClass* OuterClass = Cast<UClass>(StaticFindObjectFast(UClass::StaticClass(), OuterPackage, TEXT(\"%s\")));\r\n"), *FNativeClassHeaderGenerator::GetOverriddenName(Class));
 			GeneratedClassRegisterFunctionText.Logf(TEXT("\t\tif (!OuterClass || !(OuterClass->ClassFlags & CLASS_Constructed))\r\n"));
@@ -2349,17 +2367,18 @@ void FNativeClassHeaderGenerator::ExportNativeGeneratedInitCode(FOutputDevice& O
 
 	FString OverriddenClassName = *FNativeClassHeaderGenerator::GetOverriddenName(Class);
 
-	const FString EmptyString = FString();
-	const FString& InitSearchableValuesFunctionName = bIsDynamic ? Class->GetMetaData(TEXT("InitializeStaticSearchableValues")) : EmptyString;
+	const FString EmptyString;
+	const FString& InitSearchableValuesFunctionName = bIsDynamic ? Class->GetMetaData(NAME_InitializeStaticSearchableValues) : EmptyString;
 	const FString InitSearchableValuesFunctionParam = InitSearchableValuesFunctionName.IsEmpty() ? FString(TEXT("nullptr")) :
 		FString::Printf(TEXT("&%s::%s"), ClassNameCPP, *InitSearchableValuesFunctionName);
 
 	// Append base class' hash at the end of the generated code, this will force update derived classes
 	// when base class changes during hot-reload.
 	uint32 BaseClassHash = 0;
-	if (Class->GetSuperClass() && !Class->GetSuperClass()->HasAnyClassFlags(CLASS_Intrinsic))
+	FClass* SuperClass = Class->GetSuperClass();
+	if (SuperClass && !SuperClass->HasAnyClassFlags(CLASS_Intrinsic))
 	{
-		BaseClassHash = GGeneratedCodeHashes.FindChecked(Class->GetSuperClass());
+		BaseClassHash = GGeneratedCodeHashes.FindChecked(SuperClass);
 	}
 	GeneratedClassRegisterFunctionText.Logf(TEXT("\r\n// %u\r\n"), BaseClassHash);
 
@@ -2399,16 +2418,29 @@ void FNativeClassHeaderGenerator::ExportNativeGeneratedInitCode(FOutputDevice& O
 	Out.Logf(TEXT("\t\treturn %s::StaticClass();\r\n"), ClassNameCPP);
 	Out.Logf(TEXT("\t}\r\n"));
 
-	Out.Logf(TEXT("\tstatic FCompiledInDefer Z_CompiledInDefer_UClass_%s(%s, &%s::StaticClass, TEXT(\"%s\"), TEXT(\"%s\"), %s, %s, %s, %s);\r\n"),
-		ClassNameCPP,
-		*SingletonName,
-		ClassNameCPP,
-		bIsDynamic ? *FClass::GetTypePackageName(Class) : *Class->GetOutermost()->GetName(),
-		bIsDynamic ? *OverriddenClassName : ClassNameCPP,
-		bIsDynamic ? TEXT("true") : TEXT("false"),
-		bIsDynamic ? *AsTEXT(FClass::GetTypePackageName(Class)) : TEXT("nullptr"),
-		bIsDynamic ? *AsTEXT(FNativeClassHeaderGenerator::GetOverriddenPathName(Class)) : TEXT("nullptr"),
-		*InitSearchableValuesFunctionParam);
+	if (bIsDynamic)
+	{
+		const FString& ClassPackageName = FClass::GetTypePackageName(Class);
+		Out.Logf(TEXT("\tstatic FCompiledInDefer Z_CompiledInDefer_UClass_%s(%s, &%s::StaticClass, TEXT(\"%s\"), TEXT(\"%s\"), true, %s, %s, %s);\r\n"),
+			ClassNameCPP,
+			*SingletonName,
+			ClassNameCPP,
+			*ClassPackageName,
+			*OverriddenClassName,
+			*AsTEXT(ClassPackageName),
+			*AsTEXT(FNativeClassHeaderGenerator::GetOverriddenPathName(Class)),
+			*InitSearchableValuesFunctionParam);
+	}
+	else
+	{
+		Out.Logf(TEXT("\tstatic FCompiledInDefer Z_CompiledInDefer_UClass_%s(%s, &%s::StaticClass, TEXT(\"%s\"), TEXT(\"%s\"), false, nullptr, nullptr, %s);\r\n"),
+			ClassNameCPP,
+			*SingletonName,
+			ClassNameCPP,
+			*Class->GetOutermost()->GetName(),
+			ClassNameCPP,
+			*InitSearchableValuesFunctionParam);
+	}
 
 
 	if (ClassHasReplicatedProperties(Class))
@@ -2419,7 +2451,7 @@ void FNativeClassHeaderGenerator::ExportNativeGeneratedInitCode(FOutputDevice& O
 			"\r\n"
 			"\tvoid %s::ValidateGeneratedRepEnums(const TArray<struct FRepRecord>& ClassReps) const\r\n"
 			"\t{\r\n"
-		), ClassNameCPP, ClassNameCPP);
+		), ClassNameCPP);
 
 		FUHTStringBuilder NameBuilder;
 
@@ -2429,7 +2461,7 @@ void FNativeClassHeaderGenerator::ExportNativeGeneratedInitCode(FOutputDevice& O
 		for (int32 i = Class->ClassReps.Num() - Class->NetFields.Num(); i < Class->ClassReps.Num(); ++i)
 		{
 			const FProperty* const Property = Class->ClassReps[i].Property;
-			const FString& PropertyName = Property->GetName();
+			const FString PropertyName = Property->GetName();
 
 			NameBuilder.Logf(TEXT("\t\tstatic const FName Name_%s(TEXT(\"%s\"));\r\n"), *PropertyName, *PropertyName);
 
@@ -2515,7 +2547,7 @@ void FNativeClassHeaderGenerator::ExportFunction(FOutputDevice& Out, const FUnre
 		FString FunctionName = TempFunction->GetName();
 		if (TempFunction->HasAnyFunctionFlags(FUNC_Delegate))
 		{
-			FunctionName = FunctionName.LeftChop(FString(HEADER_GENERATED_DELEGATE_SIGNATURE_SUFFIX).Len());
+			FunctionName.LeftChopInline(HEADER_GENERATED_DELEGATE_SIGNATURE_SUFFIX_LENGTH, false);
 		}
 
 		StructureSize = FString::Printf(TEXT("sizeof(%s)"), *GetEventStructParamsName(TempFunction->GetOuter(), *FunctionName));
@@ -2661,13 +2693,13 @@ void FNativeClassHeaderGenerator::ExportNatives(FOutputDevice& Out, FClass* Clas
 
 void FNativeClassHeaderGenerator::ExportInterfaceCallFunctions(FOutputDevice& OutCpp, FUHTStringBuilder& Out, const TArray<UFunction*>& CallbackFunctions, const TCHAR* ClassName)
 {
-	FString APIString = GetAPIString();
+	const FString& APIString = GetAPIString();
 
 	for (UFunction* Function : CallbackFunctions)
 	{
 		FString FunctionName = Function->GetName();
 
-		auto* CompilerInfo = FFunctionData::FindForFunction(Function);
+		FFunctionData* CompilerInfo = FFunctionData::FindForFunction(Function);
 
 		const FFuncInfo& FunctionData = CompilerInfo->GetFunctionData();
 		const TCHAR* ConstQualifier = FunctionData.FunctionReference->HasAllFunctionFlags(FUNC_Const) ? TEXT("const ") : TEXT("");
@@ -2685,7 +2717,7 @@ void FNativeClassHeaderGenerator::ExportInterfaceCallFunctions(FOutputDevice& Ou
 		OutCpp.Logf(TEXT("\t\tcheck(O != NULL);") LINE_TERMINATOR);
 		OutCpp.Logf(TEXT("\t\tcheck(O->GetClass()->ImplementsInterface(U%s::StaticClass()));") LINE_TERMINATOR, ClassName);
 
-		auto Parameters = GetFunctionParmsAndReturn(FunctionData.FunctionReference);
+		FParmsAndReturnProperties Parameters = GetFunctionParmsAndReturn(FunctionData.FunctionReference);
 
 		// See if we need to create Parms struct
 		const bool bHasParms = Parameters.HasParms();
@@ -2700,23 +2732,21 @@ void FNativeClassHeaderGenerator::ExportInterfaceCallFunctions(FOutputDevice& Ou
 		OutCpp.Log(TEXT("\t\t{") LINE_TERMINATOR);
 
 		// code to populate Parms struct
-		for (auto It = Parameters.Parms.CreateConstIterator(); It; ++It)
+		for (FProperty* Param : Parameters.Parms)
 		{
-			FProperty* Param = *It;
-
-			OutCpp.Logf(TEXT("\t\t\tParms.%s=%s;") LINE_TERMINATOR, *Param->GetName(), *Param->GetName());
+			const FString ParamName = Param->GetName();
+			OutCpp.Logf(TEXT("\t\t\tParms.%s=%s;") LINE_TERMINATOR, *ParamName, *ParamName);
 		}
 
 		const FString ObjectRef = FunctionData.FunctionReference->HasAllFunctionFlags(FUNC_Const) ? FString::Printf(TEXT("const_cast<UObject*>(O)")) : TEXT("O");
 		OutCpp.Logf(TEXT("\t\t\t%s->ProcessEvent(Func, %s);") LINE_TERMINATOR, *ObjectRef, bHasParms ? TEXT("&Parms") : TEXT("NULL"));
 
-		for (auto It = Parameters.Parms.CreateConstIterator(); It; ++It)
+		for (FProperty* Param : Parameters.Parms)
 		{
-			FProperty* Param = *It;
-
 			if( Param->HasAllPropertyFlags(CPF_OutParm) && !Param->HasAnyPropertyFlags(CPF_ConstParm|CPF_ReturnParm))
 			{
-				OutCpp.Logf(TEXT("\t\t\t%s=Parms.%s;") LINE_TERMINATOR, *Param->GetName(), *Param->GetName());
+				const FString ParamName = Param->GetName();
+				OutCpp.Logf(TEXT("\t\t\t%s=Parms.%s;") LINE_TERMINATOR, *ParamName, *ParamName);
 			}
 		}
 
@@ -2737,16 +2767,14 @@ void FNativeClassHeaderGenerator::ExportInterfaceCallFunctions(FOutputDevice& Ou
 
 			OutCpp.Logf(TEXT("I->%s_Implementation("), *FunctionName);
 
-			bool First = true;
-			for (auto It = Parameters.Parms.CreateConstIterator(); It; ++It)
+			bool bFirst = true;
+			for (FProperty* Param : Parameters.Parms)
 			{
-				FProperty* Param = *It;
-
-				if (!First)
+				if (!bFirst)
 				{
 					OutCpp.Logf(TEXT(","));
 				}
-				First = false;
+				bFirst = false;
 
 				OutCpp.Logf(TEXT("%s"), *Param->GetName());
 			}
@@ -2811,9 +2839,9 @@ FString GetPreservedAccessSpecifierString(FClass* Class)
 	return PreservedAccessSpecifier + LINE_TERMINATOR;
 }
 
-void WriteMacro(FOutputDevice& Output, const FString& MacroName, const FString& MacroContent)
+void WriteMacro(FOutputDevice& Output, const FString& MacroName, FString MacroContent)
 {
-	Output.Log(*Macroize(*MacroName, *MacroContent));
+	Output.Log(*Macroize(*MacroName, MoveTemp(MacroContent)));
 }
 
 static FString PrivatePropertiesOffsetGetters(const UStruct* Struct, const FString& StructCppName)
@@ -2947,7 +2975,8 @@ void FNativeClassHeaderGenerator::ExportClassFromSourceFileInner(
 	{
 		FString EnclosingDefines;
 		FUHTStringBuilder Boilerplate, BoilerPlateCPP;
-		FString MacroNameHeader, MacroNameCPP;
+		TCHAR* MacroNameHeader;
+		TCHAR* MacroNameCPP;
 		GeneratedSerializeFunctionHeaderMacroName = SourceFile.GetGeneratedMacroName(ClassData, TEXT("_ARCHIVESERIALIZER"));
 
 		EnclosingDefines = ArchiveTypeDefinePair->EnclosingDefine;
@@ -2969,9 +2998,9 @@ void FNativeClassHeaderGenerator::ExportClassFromSourceFileInner(
 			BoilerPlateCPP.Logf(TEXT("#if %s\r\n"), *EnclosingDefines);
 		}
 
-		Boilerplate.Logf(TEXT("\t%s(%s, %s_API)\r\n"), *MacroNameHeader, ClassCPPName, *APIArg);
+		Boilerplate.Logf(TEXT("\t%s(%s, %s_API)\r\n"), MacroNameHeader, ClassCPPName, *APIArg);
 		OutGeneratedHeaderText.Log(Macroize(*GeneratedSerializeFunctionHeaderMacroName, *Boilerplate));
-		BoilerPlateCPP.Logf(TEXT("\t%s(%s)\r\n"), *MacroNameCPP, ClassCPPName);
+		BoilerPlateCPP.Logf(TEXT("\t%s(%s)\r\n"), MacroNameCPP, ClassCPPName);
 
 		if (EnclosingDefines.Len())
 		{
@@ -3032,32 +3061,28 @@ void FNativeClassHeaderGenerator::ExportClassFromSourceFileInner(
 			const TCHAR* DeprecationPopString = TEXT("PRAGMA_ENABLE_DEPRECATION_WARNINGS") LINE_TERMINATOR;
 			const TCHAR* Offset = TEXT("\t");
 
-			OutGeneratedHeaderText.Logf(
-				TEXT("%s"),
-				*Macroize(
+			OutGeneratedHeaderText.Log(
+				Macroize(
 					*SourceFile.GetGeneratedBodyMacroName(ClassGeneratedBodyLine, true),
-					*(
-						FString() +
-						Offset + DeprecationWarning +
-						Offset + DeprecationPushString +
-						Offset + InterfaceMacroName + TEXT("()") LINE_TERMINATOR +
-						StandardUObjectConstructorsMacroCall +
-						Offset + DeprecationPopString
+					FString::Printf(TEXT("\t%s\t%s\t%s()") LINE_TERMINATOR TEXT("%s\t%s")
+						, *DeprecationWarning
+						, DeprecationPushString
+						, *InterfaceMacroName
+						, *StandardUObjectConstructorsMacroCall
+						, DeprecationPopString
 					)
 				)
 			);
 
-			OutGeneratedHeaderText.Logf(
-				TEXT("%s"),
-				*Macroize(
+			OutGeneratedHeaderText.Log(
+				Macroize(
 					*SourceFile.GetGeneratedBodyMacroName(ClassGeneratedBodyLine),
-					*(
-						FString() +
-						Offset + DeprecationPushString +
-						Offset + InterfaceMacroName + TEXT("()") LINE_TERMINATOR +
-						EnhancedUObjectConstructorsMacroCall +
-						GetPreservedAccessSpecifierString(Class) +
-						Offset + DeprecationPopString
+					FString::Printf(TEXT("\t%s\t%s()") LINE_TERMINATOR TEXT("%s%s\t%s")
+						, DeprecationPushString
+						, *InterfaceMacroName
+						, *EnhancedUObjectConstructorsMacroCall
+						, *GetPreservedAccessSpecifierString(Class)
+						, DeprecationPopString
 					)
 				)
 			);
@@ -3092,7 +3117,7 @@ void FNativeClassHeaderGenerator::ExportClassFromSourceFileInner(
 			}
 
 			if (ClassHasReplicatedProperties(Class))
-				{
+			{
 				WriteReplicatedMacroData(ClassRange, ClassCPPName, *APIArg, Class, SuperClass, InterfaceBoilerplate, SourceFile);
 			}
 
@@ -3119,7 +3144,7 @@ void FNativeClassHeaderGenerator::ExportClassFromSourceFileInner(
 			}
 
 			if (ClassHasReplicatedProperties(Class))
-				{
+			{
 				WriteReplicatedMacroData(ClassRange, ClassCPPName, *APIArg, Class, SuperClass, Boilerplate, SourceFile);
 			}
 
@@ -3149,32 +3174,40 @@ void FNativeClassHeaderGenerator::ExportClassFromSourceFileInner(
 	}
 
 	{
-		bool bIsIInterface = Class->HasAnyClassFlags(CLASS_Interface);
-
-		auto MacroName = FString::Printf(TEXT("GENERATED_%s_BODY()"), bIsIInterface ? TEXT("IINTERFACE") : TEXT("UCLASS"));
+		const TCHAR* DeprecationPushString = TEXT("PRAGMA_DISABLE_DEPRECATION_WARNINGS") LINE_TERMINATOR;
+		const TCHAR* DeprecationPopString = TEXT("PRAGMA_ENABLE_DEPRECATION_WARNINGS") LINE_TERMINATOR;
+		const TCHAR* Public = TEXT("public:") LINE_TERMINATOR;
 
-		auto DeprecationWarning = bIsIInterface ? FString(TEXT("")) : GetGeneratedMacroDeprecationWarning(*MacroName);
+		const bool bIsIInterface = Class->HasAnyClassFlags(CLASS_Interface);
 
-		auto DeprecationPushString = TEXT("PRAGMA_DISABLE_DEPRECATION_WARNINGS") LINE_TERMINATOR;
-		auto DeprecationPopString = TEXT("PRAGMA_ENABLE_DEPRECATION_WARNINGS") LINE_TERMINATOR;
+		const TCHAR* MacroName;
+		FString DeprecationWarning;
+		FString LegacyGeneratedBody;
+		FString GeneratedBody;
+		int32 GeneratedBodyLine;
 
-		auto Public = TEXT("public:" LINE_TERMINATOR);
 
-		auto GeneratedBodyLine = bIsIInterface ? ClassData->GetInterfaceGeneratedBodyLine() : ClassData->GetGeneratedBodyLine();
-		auto LegacyGeneratedBody = FString(bIsIInterface ? TEXT("") : PPOMacroName)
-			+ ClassMacroCalls
-			+ (bIsIInterface ? TEXT("") : StandardUObjectConstructorsMacroCall);
-		auto GeneratedBody = FString(bIsIInterface ? TEXT("") : PPOMacroName)
-			+ ClassNoPureDeclsMacroCalls
-			+ (bIsIInterface ? TEXT("") : EnhancedUObjectConstructorsMacroCall);
-
-		auto WrappedLegacyGeneratedBody = DeprecationWarning + DeprecationPushString + Public + LegacyGeneratedBody + Public + DeprecationPopString;
-		auto WrappedGeneratedBody = FString(DeprecationPushString) + Public + GeneratedBody + GetPreservedAccessSpecifierString(Class) + DeprecationPopString;
+		if (bIsIInterface)
+		{
+			MacroName = TEXT("GENERATED_IINTERFACE_BODY()");
+			GeneratedBodyLine = ClassData->GetInterfaceGeneratedBodyLine();
+			LegacyGeneratedBody = ClassMacroCalls;
+			GeneratedBody = ClassNoPureDeclsMacroCalls;
+		}
+		else
+		{
+			MacroName = TEXT("GENERATED_UCLASS_BODY()");
+			DeprecationWarning = GetGeneratedMacroDeprecationWarning(MacroName);
+			GeneratedBodyLine = ClassData->GetGeneratedBodyLine();
+			LegacyGeneratedBody = FString::Printf(TEXT("%s%s%s"), *PPOMacroName, *ClassMacroCalls, *StandardUObjectConstructorsMacroCall);
+			GeneratedBody = FString::Printf(TEXT("%s%s%s"), *PPOMacroName, *ClassNoPureDeclsMacroCalls, *EnhancedUObjectConstructorsMacroCall);
+		}
 
-		auto BodyMacros = Macroize(*SourceFile.GetGeneratedBodyMacroName(GeneratedBodyLine, true), *WrappedLegacyGeneratedBody) +
-			Macroize(*SourceFile.GetGeneratedBodyMacroName(GeneratedBodyLine, false), *WrappedGeneratedBody);
+		FString WrappedLegacyGeneratedBody = FString::Printf(TEXT("%s%s%s%s%s%s"), *DeprecationWarning, DeprecationPushString, Public, *LegacyGeneratedBody, Public, DeprecationPopString);
+		FString WrappedGeneratedBody = FString::Printf(TEXT("%s%s%s%s%s"), DeprecationPushString, Public, *GeneratedBody, *GetPreservedAccessSpecifierString(Class), DeprecationPopString);
 
-		OutGeneratedHeaderText.Log(*BodyMacros);
+		OutGeneratedHeaderText.Log(*Macroize(*SourceFile.GetGeneratedBodyMacroName(GeneratedBodyLine, true), MoveTemp(WrappedLegacyGeneratedBody)));
+		OutGeneratedHeaderText.Log(*Macroize(*SourceFile.GetGeneratedBodyMacroName(GeneratedBodyLine, false), MoveTemp(WrappedGeneratedBody)));
 	}
 
 	// Forward declare the StaticClass specialisation in the header
@@ -3482,15 +3515,15 @@ void FNativeClassHeaderGenerator::ExportGeneratedStructBodyMacros(FOutputDevice&
 {
 	const bool bIsDynamic = FClass::IsDynamic(Struct);
 	const FString ActualStructName = FNativeClassHeaderGenerator::GetOverriddenName(Struct);
-	const FString FriendApiString  = GetAPIString();
+	const FString& FriendApiString  = GetAPIString();
 
 	UStruct* BaseStruct = Struct->GetSuperStruct();
 
 	const TCHAR* StructNameCPP = NameLookupCPP.GetNameCPP(Struct);
 
 	const FString SingletonName = GetSingletonName(Struct);
+	const FString ChoppedSingletonName = SingletonName.LeftChop(2);
 
-	FString RigVMMacroPrefix = FString::Printf(TEXT("UE_%s"), StructNameCPP + 1);
 	TArray<FString> RigVMVirtualFuncProlog, RigVMStubProlog;
 
 	// for RigVM methods we need to generated a macro used for implementing the static method
@@ -3519,21 +3552,28 @@ void FNativeClassHeaderGenerator::ExportGeneratedStructBodyMacros(FOutputDevice&
 				RigVMVirtualFuncProlog.Add(FString::Printf(TEXT("%s %s(%s);"), *Parameter.CastType, *Parameter.CastName, *Parameter.Name));
 			}
 
-			FString VariableType = Parameter.TypeVariableRef(true);
-			FString ExtractedType = Parameter.TypeOriginal();
-			if (Parameter.TypeOriginal(true).StartsWith(TEXT("TArrayView")))
+			FString VariableType;
+			FString ExtractedType;
+			const FString& ParamTypeOriginal = Parameter.TypeOriginal(true);
+			if (ParamTypeOriginal.StartsWith(TEXT("TArrayView"), ESearchCase::CaseSensitive))
 			{
 				ExtractedType = Parameter.ExtendedType().LeftChop(1).RightChop(1);
-				VariableType = Parameter.TypeOriginal(true);
+				VariableType = ParamTypeOriginal;
+			}
+			else
+			{
+				VariableType = Parameter.TypeVariableRef(true);
+				ExtractedType = Parameter.TypeOriginal();
 			}
 
+			const FString& ParamNameOriginal = Parameter.NameOriginal(false);
 			RigVMStubProlog.Add(FString::Printf(TEXT("%s %s = RigVMMemoryContainer[%s_Operand.GetContainerIndex()]->%s<%s>(%s_Operand, true);"),
 				*VariableType,
-				*Parameter.NameOriginal(false),
-				*Parameter.NameOriginal(false),
+				*ParamNameOriginal,
+				*ParamNameOriginal,
 				*Parameter.Getter,
 				*ExtractedType,
-				*Parameter.NameOriginal(false)));
+				*ParamNameOriginal));
 		}
 
 		FString StructMembers = StructRigVMInfo->Members.Declarations(false, TEXT(", \\\r\n\t\t"), true, false);
@@ -3553,14 +3593,10 @@ void FNativeClassHeaderGenerator::ExportGeneratedStructBodyMacros(FOutputDevice&
 	{
 		check(Struct->StructMacroDeclaredLineNumber != INDEX_NONE);
 
-		FString RequiredAPI;
-		if (!(Struct->StructFlags & STRUCT_RequiredAPI))
-		{
-			RequiredAPI = FriendApiString;
-		}
+		const bool bRequiredAPI = !(Struct->StructFlags & STRUCT_RequiredAPI);
 
-		const FString FriendLine = FString::Printf(TEXT("\tfriend struct %s_Statics;\r\n"), *SingletonName.LeftChop(2));
-		const FString StaticClassLine = FString::Printf(TEXT("\t%sstatic class UScriptStruct* StaticStruct();\r\n"), *RequiredAPI);
+		const FString FriendLine = FString::Printf(TEXT("\tfriend struct %s_Statics;\r\n"), *ChoppedSingletonName);
+		const FString StaticClassLine = FString::Printf(TEXT("\t%sstatic class UScriptStruct* StaticStruct();\r\n"), (bRequiredAPI ? *FriendApiString : TEXT("")));
 		const FString PrivatePropertiesOffset = PrivatePropertiesOffsetGetters(Struct, StructNameCPP);
 		
 		// if we have RigVM methods on this struct we need to 
@@ -3579,27 +3615,25 @@ void FNativeClassHeaderGenerator::ExportGeneratedStructBodyMacros(FOutputDevice&
 
 		const FString SuperTypedef = BaseStruct ? FString::Printf(TEXT("\ttypedef %s Super;\r\n"), NameLookupCPP.GetNameCPP(BaseStruct)) : FString();
 
-		const FString CombinedLine = FriendLine + StaticClassLine + RigVMMethodsDeclarations + PrivatePropertiesOffset + SuperTypedef;
+		FString CombinedLine = FString::Printf(TEXT("%s%s%s%s%s"), *FriendLine, *StaticClassLine, *RigVMMethodsDeclarations, *PrivatePropertiesOffset, *SuperTypedef);
 		const FString MacroName = SourceFile.GetGeneratedBodyMacroName(Struct->StructMacroDeclaredLineNumber);
 
-		const FString Macroized = Macroize(*MacroName, *CombinedLine);
-		OutGeneratedHeaderText.Log(*Macroized);
+		const FString Macroized = Macroize(*MacroName, MoveTemp(CombinedLine));
+		OutGeneratedHeaderText.Log(Macroized);
 
-		FString GetHashName = FString::Printf(TEXT("Get_%s_Hash"), *SingletonName.LeftChop(2));
+		FString GetHashName = FString::Printf(TEXT("Get_%s_Hash"), *ChoppedSingletonName);
 
 		Out.Logf(TEXT("class UScriptStruct* %s::StaticStruct()\r\n"), StructNameCPP);
 		Out.Logf(TEXT("{\r\n"));
 
 		// UStructs can have UClass or UPackage outer (if declared in non-UClass headers).
-		FString OuterName;
+		const FString& OuterName (bIsDynamic ? STRING_StructPackage : GetPackageSingletonName(CastChecked<UPackage>(Struct->GetOuter())));
 		if (!bIsDynamic)
 		{
-			OuterName = GetPackageSingletonName(CastChecked<UPackage>(Struct->GetOuter()));
 			Out.Logf(TEXT("\tstatic class UScriptStruct* Singleton = NULL;\r\n"));
 		}
 		else
 		{
-			OuterName = TEXT("StructPackage");
 			Out.Logf(TEXT("\tclass UPackage* %s = FindOrConstructDynamicTypePackage(TEXT(\"%s\"));\r\n"), *OuterName, *FClass::GetTypePackageName(Struct));
 			Out.Logf(TEXT("\tclass UScriptStruct* Singleton = Cast<UScriptStruct>(StaticFindObjectFast(UScriptStruct::StaticClass(), %s, TEXT(\"%s\")));\r\n"), *OuterName, *ActualStructName);
 		}
@@ -3609,7 +3643,7 @@ void FNativeClassHeaderGenerator::ExportGeneratedStructBodyMacros(FOutputDevice&
 		Out.Logf(TEXT("\t\textern %suint32 %s();\r\n"), *FriendApiString, *GetHashName);
 
 		Out.Logf(TEXT("\t\tSingleton = GetStaticStruct(%s, %s, TEXT(\"%s\"), sizeof(%s), %s());\r\n"),
-			*SingletonName.LeftChop(2), *OuterName, *ActualStructName, StructNameCPP, *GetHashName);
+			*ChoppedSingletonName, *OuterName, *ActualStructName, StructNameCPP, *GetHashName);
 
 		// if this struct has RigVM methods - we need to register the method to our central
 		// registry on construction of the static struct
@@ -3635,14 +3669,25 @@ void FNativeClassHeaderGenerator::ExportGeneratedStructBodyMacros(FOutputDevice&
 		Out.Logf(TEXT("\treturn %s::StaticStruct();\r\n"), StructNameCPP);
 		Out.Logf(TEXT("}\r\n"));
 
-		Out.Logf(TEXT("static FCompiledInDeferStruct Z_CompiledInDeferStruct_UScriptStruct_%s(%s::StaticStruct, TEXT(\"%s\"), TEXT(\"%s\"), %s, %s, %s);\r\n"),
-			StructNameCPP,
-			StructNameCPP,
-			bIsDynamic ? *FClass::GetTypePackageName(Struct) : *Struct->GetOutermost()->GetName(),
-			*ActualStructName,
-			bIsDynamic ? TEXT("true") : TEXT("false"),
-			bIsDynamic ? *AsTEXT(FClass::GetTypePackageName(Struct)) : TEXT("nullptr"),
-			bIsDynamic ? *AsTEXT(FNativeClassHeaderGenerator::GetOverriddenPathName(Struct)) : TEXT("nullptr"));
+		if (bIsDynamic)
+		{
+			const FString& StructPackageName = FClass::GetTypePackageName(Struct);
+			Out.Logf(TEXT("static FCompiledInDeferStruct Z_CompiledInDeferStruct_UScriptStruct_%s(%s::StaticStruct, TEXT(\"%s\"), TEXT(\"%s\"), true, %s, %s);\r\n"),
+				StructNameCPP,
+				StructNameCPP,
+				*StructPackageName,
+				*ActualStructName,
+				*AsTEXT(StructPackageName),
+				*AsTEXT(FNativeClassHeaderGenerator::GetOverriddenPathName(Struct)));
+		}
+		else
+		{
+			Out.Logf(TEXT("static FCompiledInDeferStruct Z_CompiledInDeferStruct_UScriptStruct_%s(%s::StaticStruct, TEXT(\"%s\"), TEXT(\"%s\"), false, nullptr, nullptr);\r\n"),
+				StructNameCPP,
+				StructNameCPP,
+				*Struct->GetOutermost()->GetName(),
+				*ActualStructName);
+		}
 
 		// Generate StaticRegisterNatives equivalent for structs without classes.
 		if (!Struct->GetOuter()->IsA(UStruct::StaticClass()))
@@ -3660,7 +3705,7 @@ void FNativeClassHeaderGenerator::ExportGeneratedStructBodyMacros(FOutputDevice&
 		}
 	}
 
-	FString StaticsStructName = SingletonName.LeftChop(2) + TEXT("_Statics");
+	FString StaticsStructName = ChoppedSingletonName + TEXT("_Statics");
 
 	FUHTStringBuilder GeneratedStructRegisterFunctionText;
 	FUHTStringBuilder StaticDefinitions;
@@ -3871,15 +3916,7 @@ void FNativeClassHeaderGenerator::ExportGeneratedEnumInitCode(FOutputDevice& Out
 	FMacroBlockEmitter EditorOnlyData(Out, TEXT("WITH_EDITORONLY_DATA"));
 	EditorOnlyData(bIsEditorOnlyDataType);
 
-	FString PackageSingletonName;
-	if (!bIsDynamic)
-	{
-		PackageSingletonName = GetPackageSingletonName(CastChecked<UPackage>(Enum->GetOuter()));
-	}
-	else
-	{
-		PackageSingletonName = FClass::GetTypePackageName(static_cast<UField*>(Enum));
-	}
+	const FString& PackageSingletonName = (bIsDynamic ? FClass::GetTypePackageName(static_cast<UField*>(Enum)) : GetPackageSingletonName(CastChecked<UPackage>(Enum->GetOuter())));
 
 	Out.Logf(TEXT("\tstatic UEnum* %s_StaticEnum()\r\n"), *Enum->GetName());
 	Out.Logf(TEXT("\t{\r\n"));
@@ -3913,16 +3950,29 @@ void FNativeClassHeaderGenerator::ExportGeneratedEnumInitCode(FOutputDevice& Out
 	Out.Logf(TEXT("\t\treturn %s_StaticEnum();\r\n"), *Enum->GetName());
 	Out.Logf(TEXT("\t}\r\n"));
 
-	Out.Logf(
-		TEXT("\tstatic FCompiledInDeferEnum Z_CompiledInDeferEnum_UEnum_%s(%s_StaticEnum, TEXT(\"%s\"), TEXT(\"%s\"), %s, %s, %s);\r\n"),
-		*EnumNameCpp,
-		*EnumNameCpp,
-		bIsDynamic ? *FClass::GetTypePackageName(static_cast<UField*>(Enum)) : *Enum->GetOutermost()->GetName(),
-		*OverriddenEnumNameCpp,
-		bIsDynamic ? TEXT("true") : TEXT("false"),
-		bIsDynamic ? *AsTEXT(FClass::GetTypePackageName(static_cast<UField*>(Enum))) : TEXT("nullptr"),
-		bIsDynamic ? *AsTEXT(FNativeClassHeaderGenerator::GetOverriddenPathName(static_cast<UField*>(Enum))) : TEXT("nullptr")
-	);
+	if (bIsDynamic)
+	{
+		const FString& EnumPackageName = FClass::GetTypePackageName(static_cast<UField*>(Enum));
+		Out.Logf(
+			TEXT("\tstatic FCompiledInDeferEnum Z_CompiledInDeferEnum_UEnum_%s(%s_StaticEnum, TEXT(\"%s\"), TEXT(\"%s\"), true, %s, %s);\r\n"),
+			*EnumNameCpp,
+			*EnumNameCpp,
+			*EnumPackageName,
+			*OverriddenEnumNameCpp,
+			*AsTEXT(EnumPackageName),
+			*AsTEXT(FNativeClassHeaderGenerator::GetOverriddenPathName(static_cast<UField*>(Enum)))
+		);
+	}
+	else
+	{
+		Out.Logf(
+			TEXT("\tstatic FCompiledInDeferEnum Z_CompiledInDeferEnum_UEnum_%s(%s_StaticEnum, TEXT(\"%s\"), TEXT(\"%s\"), false, nullptr, nullptr);\r\n"),
+			*EnumNameCpp,
+			*EnumNameCpp,
+			*Enum->GetOutermost()->GetName(),
+			*OverriddenEnumNameCpp
+		);
+	}
 
 	const FString EnumSingletonName = GetSingletonName(Enum);
 	const FString HashFuncName       = FString::Printf(TEXT("Get_%s_Hash"), *SingletonName);
@@ -3955,7 +4005,7 @@ void FNativeClassHeaderGenerator::ExportGeneratedEnumInitCode(FOutputDevice& Out
 
 	const TCHAR* UEnumObjectFlags = bIsDynamic ? TEXT("RF_Public|RF_Transient") : TEXT("RF_Public|RF_Transient|RF_MarkAsNative");
 
-	FString EnumFormStr;
+	TCHAR* EnumFormStr = TEXT("");
 	switch (Enum->GetCppForm())
 	{
 		case UEnum::ECppForm::Regular:    EnumFormStr = TEXT("UEnum::ECppForm::Regular");    break;
@@ -3986,7 +4036,7 @@ void FNativeClassHeaderGenerator::ExportGeneratedEnumInitCode(FOutputDevice& Out
 	GeneratedEnumRegisterFunctionText.Logf(TEXT("\t\t\t\tUE_ARRAY_COUNT(Enumerators),\r\n"));
 	GeneratedEnumRegisterFunctionText.Logf(TEXT("\t\t\t\t%s,\r\n"), UEnumObjectFlags);
 	GeneratedEnumRegisterFunctionText.Logf(TEXT("\t\t\t\tUE4CodeGen_Private::EDynamicType::%s,\r\n"), bIsDynamic ? TEXT("Dynamic") : TEXT("NotDynamic"));
-	GeneratedEnumRegisterFunctionText.Logf(TEXT("\t\t\t\t(uint8)%s,\r\n"), *EnumFormStr);
+	GeneratedEnumRegisterFunctionText.Logf(TEXT("\t\t\t\t(uint8)%s,\r\n"), EnumFormStr);
 	GeneratedEnumRegisterFunctionText.Logf(TEXT("\t\t\t\t%s\r\n"), *MetaDataParams);
 	GeneratedEnumRegisterFunctionText.Logf(TEXT("\t\t\t};\r\n"));
 	GeneratedEnumRegisterFunctionText.Logf(TEXT("\t\t\tUE4CodeGen_Private::ConstructUEnum(ReturnEnum, EnumParams);\r\n"));
@@ -4037,10 +4087,8 @@ void WriteEventFunctionPrologue(FOutputDevice& Output, int32 Indent, const FParm
 	Output.Logf(TEXT("%s%s Parms;\r\n"), FCString::Tab(Indent + 1), *EventStructName );
 
 	// Declare a parameter struct for this event/delegate and assign the struct members using the values passed into the event/delegate call.
-	for (auto It = Parameters.Parms.CreateConstIterator(); It; ++It)
+	for (FProperty* Prop : Parameters.Parms)
 	{
-		FProperty* Prop = *It;
-
 		const FString PropertyName = Prop->GetName();
 		if (Prop->ArrayDim > 1)
 		{
@@ -4062,10 +4110,8 @@ void WriteEventFunctionPrologue(FOutputDevice& Output, int32 Indent, const FParm
 void WriteEventFunctionEpilogue(FOutputDevice& Output, int32 Indent, const FParmsAndReturnProperties& Parameters)
 {
 	// Out parm copying.
-	for (auto It = Parameters.Parms.CreateConstIterator(); It; ++It)
+	for (FProperty* Prop : Parameters.Parms)
 	{
-		FProperty* Prop = *It;
-
 		if ((Prop->PropertyFlags & (CPF_OutParm | CPF_ConstParm)) == CPF_OutParm)
 		{
 			const FString PropertyName = Prop->GetName();
@@ -4099,7 +4145,7 @@ void FNativeClassHeaderGenerator::ExportDelegateDeclaration(FOutputDevice& Out,
 	const bool bIsMulticastDelegate = Function->HasAnyFunctionFlags( FUNC_MulticastDelegate );
 
 	// Unmangle the function name
-	const FString DelegateName = Function->GetName().LeftChop( FString( HEADER_GENERATED_DELEGATE_SIGNATURE_SUFFIX ).Len() );
+	const FString DelegateName = Function->GetName().LeftChop(HEADER_GENERATED_DELEGATE_SIGNATURE_SUFFIX_LENGTH);
 
 	const FFunctionData* CompilerInfo = FFunctionData::FindForFunction(Function);
 
@@ -4142,7 +4188,7 @@ void FNativeClassHeaderGenerator::ExportDelegateDefinition(FOutputDevice& Out, c
 	const bool bIsMulticastDelegate = Function->HasAnyFunctionFlags( FUNC_MulticastDelegate );
 
 	// Unmangle the function name
-	const FString DelegateName = Function->GetName().LeftChop( FString( HEADER_GENERATED_DELEGATE_SIGNATURE_SUFFIX ).Len() );
+	const FString DelegateName = Function->GetName().LeftChop(HEADER_GENERATED_DELEGATE_SIGNATURE_SUFFIX_LENGTH);
 
 	const FFunctionData* CompilerInfo = FFunctionData::FindForFunction(Function);
 
@@ -4192,7 +4238,7 @@ void FNativeClassHeaderGenerator::ExportEventParm(FUHTStringBuilder& Out, TSet<F
 	FString FunctionName = Function->GetName();
 	if (Function->HasAnyFunctionFlags(FUNC_Delegate))
 	{
-		FunctionName = FunctionName.LeftChop(FString(HEADER_GENERATED_DELEGATE_SIGNATURE_SUFFIX).Len());
+		FunctionName.LeftChopInline(HEADER_GENERATED_DELEGATE_SIGNATURE_SUFFIX_LENGTH, false);
 	}
 
 	FString EventParmStructName = GetEventStructParamsName(Function->GetOuter(), *FunctionName);
@@ -4372,34 +4418,24 @@ FString FNativeClassHeaderGenerator::GetNullParameterValue( FProperty* Prop, boo
 
 FString FNativeClassHeaderGenerator::GetFunctionReturnString(UFunction* Function)
 {
+	FString Result;
+
 	if (FProperty* Return = Function->GetReturnProperty())
 	{
 		FString ExtendedReturnType;
 		ForwardDeclarations.Add(Return->GetCPPTypeForwardDeclaration());
 		FString ReturnType = Return->GetCPPType(&ExtendedReturnType, CPPF_ArgumentOrReturnValue);
 		FUHTStringBuilder ReplacementText;
-		ReplacementText += ReturnType;
+		ReplacementText += MoveTemp(ReturnType);
 		ApplyAlternatePropertyExportText(Return, ReplacementText, EExportingState::Normal);
-		return ReplacementText + ExtendedReturnType;
+		Result = MoveTemp(ReplacementText) + MoveTemp(ExtendedReturnType);
 	}
-
-	return TEXT("void");
-}
-
-/**
-* Gets string with function const modifier type.
-*
-* @param Function Function to get const modifier of.
-* @return Empty FString if function is non-const, FString("const") if function is const.
-*/
-FString GetFunctionConstModifierString(UFunction* Function)
-{
-	if (Function->HasAllFunctionFlags(FUNC_Const))
+	else
 	{
-		return TEXT("const");
+		Result = TEXT("void");
 	}
 
-	return FString();
+	return Result;
 }
 
 /**
@@ -4452,29 +4488,29 @@ void FNativeClassHeaderGenerator::CheckRPCFunctions(const FFuncInfo& FunctionDat
 	bool bHasImplementation = ImplementationPosition != INDEX_NONE;
 	bool bHasValidate = ValidatePosition != INDEX_NONE;
 
-	auto Function = FunctionData.FunctionReference;
-	auto FunctionReturnType = GetFunctionReturnString(Function);
-	auto ConstModifier = GetFunctionConstModifierString(Function) + TEXT(" ");
+	UFunction* Function = FunctionData.FunctionReference;
+	FString FunctionReturnType = GetFunctionReturnString(Function);
+	const TCHAR* ConstModifier = (Function->HasAllFunctionFlags(FUNC_Const) ? TEXT("const ") : TEXT(" "));
 
-	auto bIsNative = Function->HasAllFunctionFlags(FUNC_Native);
-	auto bIsNet = Function->HasAllFunctionFlags(FUNC_Net);
-	auto bIsNetValidate = Function->HasAllFunctionFlags(FUNC_NetValidate);
-	auto bIsNetResponse = Function->HasAllFunctionFlags(FUNC_NetResponse);
-	auto bIsBlueprintEvent = Function->HasAllFunctionFlags(FUNC_BlueprintEvent);
+	const bool bIsNative = Function->HasAllFunctionFlags(FUNC_Native);
+	const bool bIsNet = Function->HasAllFunctionFlags(FUNC_Net);
+	const bool bIsNetValidate = Function->HasAllFunctionFlags(FUNC_NetValidate);
+	const bool bIsNetResponse = Function->HasAllFunctionFlags(FUNC_NetResponse);
+	const bool bIsBlueprintEvent = Function->HasAllFunctionFlags(FUNC_BlueprintEvent);
 
 	bool bNeedsImplementation = (bIsNet && !bIsNetResponse) || bIsBlueprintEvent || bIsNative;
 	bool bNeedsValidate = (bIsNative || bIsNet) && !bIsNetResponse && bIsNetValidate;
 
 	check(bNeedsImplementation || bNeedsValidate);
 
-	auto ParameterString = GetFunctionParameterString(Function);
-	const auto& Filename = SourceFile.GetFilename();
-	const auto& FileContent = SourceFile.GetContent();
+	FString ParameterString = GetFunctionParameterString(Function);
+	const FString& Filename = SourceFile.GetFilename();
+	const FString& FileContent = SourceFile.GetContent();
 
 	//
 	// Get string with function specifiers, listing why we need _Implementation or _Validate functions.
 	//
-	TArray<FString> FunctionSpecifiers;
+	TArray<const TCHAR*> FunctionSpecifiers;
 	FunctionSpecifiers.Reserve(4);
 	if (bIsNative)			{ FunctionSpecifiers.Add(TEXT("Native"));			}
 	if (bIsNet)				{ FunctionSpecifiers.Add(TEXT("Net"));				}
@@ -4487,10 +4523,10 @@ void FNativeClassHeaderGenerator::CheckRPCFunctions(const FFuncInfo& FunctionDat
 	// Coin static_assert message
 	//
 	FUHTStringBuilder AssertMessage;
-	AssertMessage.Logf(TEXT("Function %s was marked as %s"), *(Function->GetName()), *FunctionSpecifiers[0]);
+	AssertMessage.Logf(TEXT("Function %s was marked as %s"), *(Function->GetName()), FunctionSpecifiers[0]);
 	for (int32 i = 1; i < FunctionSpecifiers.Num(); ++i)
 	{
-		AssertMessage.Logf(TEXT(", %s"), *FunctionSpecifiers[i]);
+		AssertMessage.Logf(TEXT(", %s"), FunctionSpecifiers[i]);
 	}
 
 	AssertMessage.Logf(TEXT("."));
@@ -4702,7 +4738,7 @@ void FNativeClassHeaderGenerator::ExportNativeFunctionHeader(
 				Out.Log(TEXT("=0"));
 			}
 		}
-}
+	}
 }
 
 /**
@@ -4805,9 +4841,8 @@ void FNativeClassHeaderGenerator::ExportFunctionThunk(FUHTStringBuilder& RPCWrap
 			if (DelegateProp != NULL)
 			{
 				// For delegates, add an explicit conversion to the specific type of delegate before passing it along
-				const FString FunctionName = DelegateProp->SignatureFunction->GetName().LeftChop(FString(HEADER_GENERATED_DELEGATE_SIGNATURE_SUFFIX).Len());
-				const FString CPPDelegateName = FString(TEXT("F")) + FunctionName;
-				ParamName = FString::Printf(TEXT("%s(%s)"), *CPPDelegateName, *ParamName);
+				const FString FunctionName = DelegateProp->SignatureFunction->GetName().LeftChop(HEADER_GENERATED_DELEGATE_SIGNATURE_SUFFIX_LENGTH);
+				ParamName = FString::Printf(TEXT("F%s(%s)"), *FunctionName, *ParamName);
 			}
 		}
 
@@ -4816,9 +4851,8 @@ void FNativeClassHeaderGenerator::ExportFunctionThunk(FUHTStringBuilder& RPCWrap
 			if (MulticastDelegateProp != NULL)
 			{
 				// For delegates, add an explicit conversion to the specific type of delegate before passing it along
-				const FString FunctionName = MulticastDelegateProp->SignatureFunction->GetName().LeftChop(FString(HEADER_GENERATED_DELEGATE_SIGNATURE_SUFFIX).Len());
-				const FString CPPDelegateName = FString(TEXT("F")) + FunctionName;
-				ParamName = FString::Printf(TEXT("%s(%s)"), *CPPDelegateName, *ParamName);
+				const FString FunctionName = MulticastDelegateProp->SignatureFunction->GetName().LeftChop(HEADER_GENERATED_DELEGATE_SIGNATURE_SUFFIX_LENGTH);
+				ParamName = FString::Printf(TEXT("F%s(%s)"), *FunctionName, *ParamName);
 			}
 		}
 
@@ -4948,8 +4982,8 @@ FString FNativeClassHeaderGenerator::GetFunctionParameterString(UFunction* Funct
 			ParameterList += TEXT(", ");
 		}
 
-		auto Dim = GArrayDimensions.Find(Property);
-		Property->ExportCppDeclaration(PropertyText, EExportedDeclaration::Parameter, Dim ? **Dim : NULL, 0, true);
+		FString* Dim = GArrayDimensions.Find(Property);
+		Property->ExportCppDeclaration(PropertyText, EExportedDeclaration::Parameter, Dim ? **Dim : nullptr, 0, true);
 		ApplyAlternatePropertyExportText(Property, PropertyText, EExportingState::Normal);
 
 		ParameterList += PropertyText;
@@ -5018,9 +5052,9 @@ void FNativeClassHeaderGenerator::ExportNativeFunctions(FOutputDevice& OutGenera
 					CleanVarName = VarName.RightChop(1);
 				}
 
-				if (!Child->HasMetaData("NoGetter"))
+				if (!Child->HasMetaData(NAME_NoGetter))
 				{
-					if (Child->HasMetaData("GetByRef"))
+					if (Child->HasMetaData(NAME_GetByRef))
 					{
 						RuntimeStringBuilders.AutogeneratedStaticDataFuncs.Logf(TEXT("const %s& Get%s()\r\n"), *VarType, *CleanVarName);
 					}
@@ -5032,7 +5066,7 @@ void FNativeClassHeaderGenerator::ExportNativeFunctions(FOutputDevice& OutGenera
 					RuntimeStringBuilders.AutogeneratedStaticDataFuncs.Logf(TEXT("\treturn Get%s()->%s;\r\n"), *SparseClassDataString, *VarName);
 					RuntimeStringBuilders.AutogeneratedStaticDataFuncs.Logf(TEXT("}\r\n"));
 
-					if (Child->HasMetaData("GetByRef"))
+					if (Child->HasMetaData(NAME_GetByRef))
 					{
 						RuntimeStringBuilders.AutogeneratedStaticDataFuncs.Logf(TEXT("const %s& Get%s() const\r\n"), *VarType, *CleanVarName);
 					}
@@ -5480,6 +5514,7 @@ FNativeClassHeaderGenerator::FNativeClassHeaderGenerator(
 	bool InAllowSaveExportedHeaders
 )
 	: API                        (FPackageName::GetShortName(InPackage).ToUpper())
+	, APIStringPrivate           (FString::Printf(TEXT("%s_API "), *API))
 	, Package                    (InPackage)
 	, UniqueCrossModuleReferences(nullptr)
 	, bAllowSaveExportedHeaders  (InAllowSaveExportedHeaders)
@@ -5487,9 +5522,8 @@ FNativeClassHeaderGenerator::FNativeClassHeaderGenerator(
 {
 	const FString PackageName = FPackageName::GetShortName(Package);
 
-	FString PkgDir;
-	FString GeneratedIncludeDirectory;
-	if (!FindPackageLocation(*PackageName, PkgDir, GeneratedIncludeDirectory))
+	FManifestModule* PackageManifest = GetPackageManifest(PackageName);
+	if (!PackageManifest)
 	{
 		UE_LOG(LogCompile, Error, TEXT("Failed to find path for package %s"), *PackageName);
 	}
@@ -5501,9 +5535,11 @@ FNativeClassHeaderGenerator::FNativeClassHeaderGenerator(
 	});
 	if (bPackageHasAnyExportClasses)
 	{
+		TArray<UClass*> DefinedClasses;
 		for (FUnrealSourceFile* SourceFile : SourceFiles)
 		{
-			TArray<UClass*> DefinedClasses = SourceFile->GetDefinedClasses();
+			DefinedClasses.Reset();
+			SourceFile->AppendDefinedClasses(DefinedClasses);
 			for (UClass* Class : DefinedClasses)
 			{
 				if (!Class->HasAnyClassFlags(CLASS_Native))
@@ -5521,7 +5557,7 @@ FNativeClassHeaderGenerator::FNativeClassHeaderGenerator(
 	}
 
 	// Export an include line for each header
-	TArray<FUnrealSourceFile*> PublicHeaderGroupIncludes;
+	TSet<FUnrealSourceFile*> PublicHeaderGroupIncludes;
 	FUHTStringBuilder GeneratedFunctionDeclarations;
 
 	TArray<FUnrealSourceFile*> Exported;
@@ -5545,14 +5581,21 @@ FNativeClassHeaderGenerator::FNativeClassHeaderGenerator(
 		TSet<FString>                CrossModuleReferences;
 	};
 
+	TArray<UEnum*> Enums;
+	TArray<UScriptStruct*> Structs;
+	TArray<UDelegateFunction*> DelegateFunctions;
+	TArray<UClass*>	DefinedClasses;
+
 	TMap<FUnrealSourceFile*, FGeneratedCPP> GeneratedCPPs;
+	GeneratedCPPs.Reserve(Exported.Num());
+
 	for (FUnrealSourceFile* SourceFile : Exported)
 	{
 		FString ModuleRelativeFilename = SourceFile->GetFilename();
 		ConvertToBuildIncludePath(Package, ModuleRelativeFilename);
 
 		FString StrippedName       = FPaths::GetBaseFilename(ModuleRelativeFilename);
-		FString BaseSourceFilename = GeneratedIncludeDirectory / StrippedName;
+		FString BaseSourceFilename = PackageManifest->GeneratedIncludeDirectory / StrippedName;
 
 		FUHTStringBuilder GeneratedHeaderText;
 		FGeneratedCPP& GeneratedCPP = GeneratedCPPs.Emplace(SourceFile, BaseSourceFilename + TEXT(".gen.cpp"));
@@ -5568,9 +5611,9 @@ FNativeClassHeaderGenerator::FNativeClassHeaderGenerator(
 
 		NameLookupCPP.SetCurrentSourceFile(SourceFile);
 
-		TArray<UEnum*>             Enums;
-		TArray<UScriptStruct*>     Structs;
-		TArray<UDelegateFunction*> DelegateFunctions;
+		Enums.Reset();
+		Structs.Reset();
+		DelegateFunctions.Reset();
 		SourceFile->GetScope()->SplitTypesIntoArrays(Enums, Structs, DelegateFunctions);
 
 		RecordPackageSingletons(*SourceFile->GetPackage(), Structs, DelegateFunctions);
@@ -5620,7 +5663,8 @@ FNativeClassHeaderGenerator::FNativeClassHeaderGenerator(
 			ExportDelegateDefinition(GeneratedHeaderText, *SourceFile, Func);
 		}
 
-		TArray<UClass*> DefinedClasses = SourceFile->GetDefinedClasses();
+		DefinedClasses.Reset();
+		SourceFile->AppendDefinedClasses(DefinedClasses);
 		for (UClass* Class : DefinedClasses)
 		{
 			if (!(Class->ClassFlags & CLASS_Intrinsic))
@@ -5642,14 +5686,14 @@ FNativeClassHeaderGenerator::FNativeClassHeaderGenerator(
 		FString HeaderPath = BaseSourceFilename + TEXT(".generated.h");
 		bool bHasChanged = WriteHeader(*HeaderPath, GeneratedHeaderText, ForwardDeclarations);
 
-		SourceFile->SetGeneratedFilename(HeaderPath);
+		SourceFile->SetGeneratedFilename(MoveTemp(HeaderPath));
 		SourceFile->SetHasChanged(bHasChanged);
 
 		ForwardDeclarations.Reset();
 
 		if (GPublicSourceFileSet.Contains(SourceFile))
 		{
-			PublicHeaderGroupIncludes.AddUnique(SourceFile);
+			PublicHeaderGroupIncludes.Add(SourceFile);
 		}
 	}
 
@@ -5660,7 +5704,8 @@ FNativeClassHeaderGenerator::FNativeClassHeaderGenerator(
 		bool bAddedArchiveUObjectFromStructuredArchiveHeader = false;
 
 		TArray<FString>& RelativeIncludes = GeneratedCPPs[SourceFile].RelativeIncludes;
-		TArray<UClass*> DefinedClasses = SourceFile->GetDefinedClasses();
+		DefinedClasses.Reset();
+		SourceFile->AppendDefinedClasses(DefinedClasses);
 		for (UClass* Class : DefinedClasses)
 		{
 			if (Class->ClassWithin && Class->ClassWithin != UObject::StaticClass())
@@ -5704,7 +5749,7 @@ FNativeClassHeaderGenerator::FNativeClassHeaderGenerator(
 		{
 			if (SourceFile->GetPackage() == InPackage)
 			{
-				PublicHeaderGroupIncludes.AddUnique(SourceFile);
+				PublicHeaderGroupIncludes.Add(SourceFile);
 			}
 		}
 
@@ -5716,7 +5761,7 @@ FNativeClassHeaderGenerator::FNativeClassHeaderGenerator(
 		ClassesHText.Log(LINE_TERMINATOR);
 
 		// Save the classes header if it has changed.
-		FString ClassesHeaderPath = GeneratedIncludeDirectory / (PackageName + TEXT("Classes.h"));
+		const FString ClassesHeaderPath = PackageManifest->GeneratedIncludeDirectory / (PackageName + TEXT("Classes.h"));
 		SaveHeaderIfChanged(*ClassesHeaderPath, *ClassesHText);
 	}
 
@@ -5742,7 +5787,7 @@ FNativeClassHeaderGenerator::FNativeClassHeaderGenerator(
 			}
 		}
 
-		FGeneratedCPP& GeneratedCPP = GeneratedCPPs.Emplace(nullptr, GeneratedIncludeDirectory / FString::Printf(TEXT("%s.init.gen.cpp"), *PackageName));
+		FGeneratedCPP& GeneratedCPP = GeneratedCPPs.Emplace(nullptr, PackageManifest->GeneratedIncludeDirectory / FString::Printf(TEXT("%s.init.gen.cpp"), *PackageName));
 		ExportGeneratedPackageInitCode(GeneratedCPP.GeneratedText, *GeneratedFunctionDeclarations, Package, CombinedHash);
 	}
 
@@ -5760,17 +5805,22 @@ FNativeClassHeaderGenerator::FNativeClassHeaderGenerator(
 			GeneratedIncludes += FString::Printf(TEXT("#include \"%s\"\r\n"), *RelativeInclude);
 		}
 
+		FString CleanFilename = FPaths::GetCleanFilename(GeneratedCPP.Value.GeneratedCppFullFilename);
+
+		GeneratedCPPNames.Add(CleanFilename);
+
+		CleanFilename.ReplaceInline(TEXT(".gen.cpp"), TEXT(""), ESearchCase::CaseSensitive);
+		CleanFilename.ReplaceInline(TEXT("."), TEXT("_"), ESearchCase::CaseSensitive);
+
 		ExportGeneratedCPP(
 			FileText,
 			GeneratedCPP.Value.CrossModuleReferences,
-			*FPaths::GetCleanFilename(GeneratedCPP.Value.GeneratedCppFullFilename).Replace(TEXT(".gen.cpp"), TEXT("")).Replace(TEXT("."), TEXT("_")),
+			*CleanFilename,
 			*GeneratedCPP.Value.GeneratedText,
 			*GeneratedIncludes
 		);
 
 		SaveHeaderIfChanged(*GeneratedCPP.Value.GeneratedCppFullFilename, *FileText);
-
-		GeneratedCPPNames.Add(FPaths::GetCleanFilename(*GeneratedCPP.Value.GeneratedCppFullFilename));
 	}
 
 	if (bAllowSaveExportedHeaders)
@@ -5859,11 +5909,12 @@ bool FNativeClassHeaderGenerator::SaveHeaderIfChanged(const TCHAR* HeaderPath, c
 	{
 		bTestedCmdLine = true;
 
-		const FString ReferenceGeneratedCodePath = FPaths::ProjectSavedDir() / TEXT("ReferenceGeneratedCode/");
-		const FString VerifyGeneratedCodePath = FPaths::ProjectSavedDir() / TEXT("VerifyGeneratedCode/");
+		const FString& ProjectSavedDir = FPaths::ProjectSavedDir();
 
 		if (FParse::Param(FCommandLine::Get(), TEXT("WRITEREF")))
 		{
+			const FString ReferenceGeneratedCodePath = ProjectSavedDir / TEXT("ReferenceGeneratedCode/");
+
 			bWriteContents = true;
 			UE_LOG(LogCompile, Log, TEXT("********************************* Writing reference generated code to %s."), *ReferenceGeneratedCodePath);
 			UE_LOG(LogCompile, Log, TEXT("********************************* Deleting all files in ReferenceGeneratedCode."));
@@ -5872,6 +5923,9 @@ bool FNativeClassHeaderGenerator::SaveHeaderIfChanged(const TCHAR* HeaderPath, c
 		}
 		else if (FParse::Param( FCommandLine::Get(), TEXT("VERIFYREF")))
 		{
+			const FString ReferenceGeneratedCodePath = ProjectSavedDir / TEXT("ReferenceGeneratedCode/");
+			const FString VerifyGeneratedCodePath = ProjectSavedDir / TEXT("VerifyGeneratedCode/");
+
 			bVerifyContents = true;
 			UE_LOG(LogCompile, Log, TEXT("********************************* Writing generated code to %s and comparing to %s"), *VerifyGeneratedCodePath, *ReferenceGeneratedCodePath);
 			UE_LOG(LogCompile, Log, TEXT("********************************* Deleting all files in VerifyGeneratedCode."));
@@ -5882,8 +5936,9 @@ bool FNativeClassHeaderGenerator::SaveHeaderIfChanged(const TCHAR* HeaderPath, c
 
 	if (bWriteContents || bVerifyContents)
 	{
-		FString Ref    = FPaths::ProjectSavedDir() / TEXT("ReferenceGeneratedCode") / FPaths::GetCleanFilename(HeaderPath);
-		FString Verify = FPaths::ProjectSavedDir() / TEXT("VerifyGeneratedCode") / FPaths::GetCleanFilename(HeaderPath);
+		const FString& ProjectSavedDir = FPaths::ProjectSavedDir();
+		const FString CleanFilename = FPaths::GetCleanFilename(HeaderPath);
+		const FString Ref = ProjectSavedDir / TEXT("ReferenceGeneratedCode") / CleanFilename;
 
 		if (bWriteContents)
 		{
@@ -5900,6 +5955,8 @@ bool FNativeClassHeaderGenerator::SaveHeaderIfChanged(const TCHAR* HeaderPath, c
 		}
 		else
 		{
+			const FString Verify = ProjectSavedDir / TEXT("VerifyGeneratedCode") / CleanFilename;
+
 			int32 i;
 			for (i = 0 ;i < 10; i++)
 			{
@@ -5916,24 +5973,25 @@ bool FNativeClassHeaderGenerator::SaveHeaderIfChanged(const TCHAR* HeaderPath, c
 				SCOPE_SECONDS_COUNTER_UHT(LoadHeaderContentFromFile);
 				if (!FFileHelper::LoadFileToString(RefHeader, *Ref))
 				{
-					Message = FString::Printf(TEXT("********************************* %s appears to be a new generated file."), *FPaths::GetCleanFilename(HeaderPath));
+					Message = FString::Printf(TEXT("********************************* %s appears to be a new generated file."), *CleanFilename);
 				}
 				else
 				{
 					if (FCString::Strcmp(NewHeaderContents, *RefHeader) != 0)
 					{
-						Message = FString::Printf(TEXT("********************************* %s has changed."), *FPaths::GetCleanFilename(HeaderPath));
+						Message = FString::Printf(TEXT("********************************* %s has changed."), *CleanFilename);
 					}
 				}
 			}
 			if (Message.Len())
 			{
 				UE_LOG(LogCompile, Log, TEXT("%s"), *Message);
-				ChangeMessages.AddUnique(Message);
+				ChangeMessages.AddUnique(MoveTemp(Message));
 			}
 		}
 	}
 
+	FString HeaderPathStr(HeaderPath);
 
 	FString OriginalHeaderLocal;
 	{
@@ -5946,7 +6004,7 @@ bool FNativeClassHeaderGenerator::SaveHeaderIfChanged(const TCHAR* HeaderPath, c
 	{
 		if (bFailIfGeneratedCodeChanges)
 		{
-			FString ConflictPath = FString(HeaderPath) + TEXT(".conflict");
+			FString ConflictPath = HeaderPathStr + TEXT(".conflict");
 			FFileHelper::SaveStringToFile(NewHeaderContents, *ConflictPath);
 
 			GCompilationResult = ECompilationResult::FailedDueToHeaderChange;
@@ -5954,7 +6012,7 @@ bool FNativeClassHeaderGenerator::SaveHeaderIfChanged(const TCHAR* HeaderPath, c
 		}
 
 		// save the updated version to a tmp file so that the user can see what will be changing
-		const FString TmpHeaderFilename = GenerateTempHeaderName( HeaderPath, false );
+		FString TmpHeaderFilename = GenerateTempHeaderName(HeaderPathStr, false );
 
 		// delete any existing temp file
 		IFileManager::Get().Delete( *TmpHeaderFilename, false, true );
@@ -5963,11 +6021,13 @@ bool FNativeClassHeaderGenerator::SaveHeaderIfChanged(const TCHAR* HeaderPath, c
 			UE_LOG_WARNING_UHT(TEXT("Failed to save header export preview: '%s'"), *TmpHeaderFilename);
 		}
 
-		TempHeaderPaths.Add(TmpHeaderFilename);
+		TempHeaderPaths.Add(MoveTemp(TmpHeaderFilename));
 	}
 
 	// Remember this header filename to be able to check for any old (unused) headers later.
-	PackageHeaderPaths.Add(FString(HeaderPath).Replace(TEXT("\\"), TEXT("/"), ESearchCase::CaseSensitive));
+	HeaderPathStr.ReplaceInline(TEXT("\\"), TEXT("/"), ESearchCase::CaseSensitive);
+
+	PackageHeaderPaths.Add(MoveTemp(HeaderPathStr));
 
 	return bHasChanged;
 }
@@ -5980,7 +6040,7 @@ bool FNativeClassHeaderGenerator::SaveHeaderIfChanged(const TCHAR* HeaderPath, c
 *
 * @return	The generated string
 */
-FString FNativeClassHeaderGenerator::GenerateTempHeaderName( FString CurrentFilename, bool bReverseOperation )
+FString FNativeClassHeaderGenerator::GenerateTempHeaderName( const FString& CurrentFilename, bool bReverseOperation )
 {
 	return bReverseOperation
 		? CurrentFilename.Replace(TEXT(".tmp"), TEXT(""))
@@ -5992,7 +6052,7 @@ FString FNativeClassHeaderGenerator::GenerateTempHeaderName( FString CurrentFile
 *
 * @param	PackageName	Name of the package being saved
 */
-void FNativeClassHeaderGenerator::ExportUpdatedHeaders(FString PackageName)
+void FNativeClassHeaderGenerator::ExportUpdatedHeaders(const FString& PackageName)
 {
 	for (const FString& TmpFilename : TempHeaderPaths)
 	{
@@ -6051,7 +6111,7 @@ void GetScriptPlugins(TArray<IScriptGeneratorPluginInterface*>& ScriptPlugins)
 	// Check if we can use these plugins and initialize them
 	for (int32 PluginIndex = ScriptPlugins.Num() - 1; PluginIndex >= 0; --PluginIndex)
 	{
-		auto ScriptGenerator = ScriptPlugins[PluginIndex];
+		IScriptGeneratorPluginInterface* ScriptGenerator = ScriptPlugins[PluginIndex];
 		bool bSupportedPlugin = ScriptGenerator->SupportsTarget(GManifest.TargetName);
 		if (bSupportedPlugin)
 		{
@@ -6433,7 +6493,7 @@ ECompilationResult::Type UnrealHeaderTool_Main(const FString& ModuleInfoFilename
 					TSharedRef<FUnrealTypeDefinitionInfo> DefinitionInfo;
 				} ContextSupplier(ScriptClass);
 
-				auto OldContext = GWarn->GetContext();
+				FContextSupplier* OldContext = GWarn->GetContext();
 
 				TGuardValue<ELogTimes::Type> DisableLogTimes(GPrintLogTimes, ELogTimes::None);
 
@@ -6627,7 +6687,7 @@ UClass* ProcessParsedClass(bool bClassIsAnInterface, TArray<FHeaderProvider>& De
 
 	if (bVerboseOutput)
 	{
-		for (const auto& Dependency : DependentOn)
+		for (const FHeaderProvider& Dependency : DependentOn)
 		{
 			UE_LOG(LogCompile, Log, TEXT("\tAdding %s as a dependency"), *Dependency.ToString());
 		}
@@ -6653,7 +6713,7 @@ TSharedRef<FUnrealSourceFile> PerformInitialParseOnHeader(UPackage* InParent, co
 
 	FUnrealSourceFile* UnrealSourceFilePtr = new FUnrealSourceFile(InParent, FileName, MoveTemp(ClassHeaderTextStrippedOfCppText));
 	TSharedRef<FUnrealSourceFile> UnrealSourceFile = MakeShareable(UnrealSourceFilePtr);
-	for (auto& ParsedClassInfo : ParsedClassArray)
+	for (FSimplifiedParsingClassInfo& ParsedClassInfo : ParsedClassArray)
 	{
 		UClass* ResultClass = ProcessParsedClass(ParsedClassInfo.IsInterface(), DependsOn, ParsedClassInfo.GetClassName(), ParsedClassInfo.GetBaseClassName(), InParent, Flags);
 		GStructToSourceLine.Add(ResultClass, MakeTuple(UnrealSourceFile, ParsedClassInfo.GetClassDefLine()));
@@ -6664,7 +6724,7 @@ TSharedRef<FUnrealSourceFile> PerformInitialParseOnHeader(UPackage* InParent, co
 		UnrealSourceFile->AddDefinedClass(ResultClass, MoveTemp(ParsedClassInfo));
 	}
 
-	for (auto& DependsOnElement : DependsOn)
+	for (FHeaderProvider& DependsOnElement : DependsOn)
 	{
 		UnrealSourceFile->GetIncludes().AddUnique(DependsOnElement);
 	}
diff --git a/Engine/Source/Programs/UnrealHeaderTool/Private/NativeClassExporter.h b/Engine/Source/Programs/UnrealHeaderTool/Private/NativeClassExporter.h
--- a/Engine/Source/Programs/UnrealHeaderTool/Private/NativeClassExporter.h
+++ b/Engine/Source/Programs/UnrealHeaderTool/Private/NativeClassExporter.h
@@ -81,13 +81,14 @@ struct FNativeClassHeaderGenerator
 {
 private:
 	FString API;
+	FString APIStringPrivate;
 
 	/**
 	 * Gets API string for this header.
 	 */
-	FString GetAPIString()
+	const FString& GetAPIString() const
 	{
-		return FString::Printf(TEXT("%s_API "), *API);
+		return APIStringPrivate;
 	}
 
 	const UPackage* Package;
@@ -124,7 +125,7 @@ private:
 	static void ExportProperties(FOutputDevice& Out, UStruct* Struct, int32 TextIndent);
 
 	/** Return the name of the singleton function that returns the UObject for Item */
-	FString GetPackageSingletonName(const UPackage* Item);
+	const FString& GetPackageSingletonName(const UPackage* Item);
 
 	/** Return the name of the singleton function that returns the UObject for Item */
 	FString GetSingletonName(UField* Item, bool bRequiresValidObject=true);
@@ -276,7 +277,7 @@ private:
 	* 
 	* @param	PackageName	Name of the package being saved
 	*/
-	void ExportUpdatedHeaders( FString PackageName  );
+	void ExportUpdatedHeaders( const FString& PackageName  );
 
 	/**
 	 * Exports the generated cpp file for all functions/events/delegates in package.
@@ -447,7 +448,7 @@ private:
 	*
 	* @return	The generated string
 	*/
-	static FString GenerateTempHeaderName( FString CurrentFilename, bool bReverseOperation = false );
+	static FString GenerateTempHeaderName( const FString& CurrentFilename, bool bReverseOperation = false );
 
 	/**
 	 * Saves a generated header if it has changed. 
