Fix Mac UHT compile
#rnx

diff --git a/Engine/Source/Programs/UnrealHeaderTool/Private/CodeGenerator.cpp b/Engine/Source/Programs/UnrealHeaderTool/Private/CodeGenerator.cpp
--- a/Engine/Source/Programs/UnrealHeaderTool/Private/CodeGenerator.cpp
+++ b/Engine/Source/Programs/UnrealHeaderTool/Private/CodeGenerator.cpp
@@ -2975,8 +2975,8 @@ void FNativeClassHeaderGenerator::ExportClassFromSourceFileInner(
 	{
 		FString EnclosingDefines;
 		FUHTStringBuilder Boilerplate, BoilerPlateCPP;
-		TCHAR* MacroNameHeader;
-		TCHAR* MacroNameCPP;
+		const TCHAR* MacroNameHeader;
+		const TCHAR* MacroNameCPP;
 		GeneratedSerializeFunctionHeaderMacroName = SourceFile.GetGeneratedMacroName(ClassData, TEXT("_ARCHIVESERIALIZER"));
 
 		EnclosingDefines = ArchiveTypeDefinePair->EnclosingDefine;
@@ -4005,7 +4005,7 @@ void FNativeClassHeaderGenerator::ExportGeneratedEnumInitCode(FOutputDevice& Out
 
 	const TCHAR* UEnumObjectFlags = bIsDynamic ? TEXT("RF_Public|RF_Transient") : TEXT("RF_Public|RF_Transient|RF_MarkAsNative");
 
-	TCHAR* EnumFormStr = TEXT("");
+	const TCHAR* EnumFormStr = TEXT("");
 	switch (Enum->GetCppForm())
 	{
 		case UEnum::ECppForm::Regular:    EnumFormStr = TEXT("UEnum::ECppForm::Regular");    break;
