UHT optimizations part 6 - Generate files in parallel. Improves code generation phase by ~50%.
* Restructured code so that FNativeClassHeaderGenerator can now use ParallelFor for 3 phases: creating the .generated.h files, creating the .gen.cpp, and exporting the generated headers
* FNameLookupCPP, GPackageSingletons, GGeneratedCodeHashes, GetTypePackageName, FNativeClassHeaderGenerator::GetPackageSingletonName, FTypeSingletonCache, and FTypeSingleton::GetExternDecl now use RWScopeLocks to safely control access to globals/statics
#rnx

diff --git a/Engine/Source/Programs/UnrealHeaderTool/Private/ClassMaps.cpp b/Engine/Source/Programs/UnrealHeaderTool/Private/ClassMaps.cpp
--- a/Engine/Source/Programs/UnrealHeaderTool/Private/ClassMaps.cpp
+++ b/Engine/Source/Programs/UnrealHeaderTool/Private/ClassMaps.cpp
@@ -7,12 +7,14 @@
 TMap<FString, TSharedRef<FUnrealSourceFile> > GUnrealSourceFilesMap;
 TMap<UField*, TSharedRef<FUnrealTypeDefinitionInfo> > GTypeDefinitionInfoMap;
 TMap<const UPackage*, TArray<UField*>> GPackageSingletons;
+FCriticalSection GPackageSingletonsCriticalSection;
 TMap<UClass*, FString> GClassStrippedHeaderTextMap;
 TMap<UClass*, FString> GClassHeaderNameWithNoPathMap;
 TSet<FUnrealSourceFile*> GPublicSourceFileSet;
 TMap<FProperty*, FString> GArrayDimensions;
 TMap<UPackage*,  const FManifestModule*> GPackageToManifestModuleMap;
 TMap<void*, uint32> GGeneratedCodeHashes;
+FRWLock GGeneratedCodeHashesLock;
 TMap<UEnum*,  EUnderlyingEnumType> GEnumUnderlyingTypes;
 TMap<FName, TSharedRef<FClassDeclarationMetaData> > GClassDeclarations;
 TSet<FProperty*> GUnsizedProperties;
diff --git a/Engine/Source/Programs/UnrealHeaderTool/Private/ClassMaps.h b/Engine/Source/Programs/UnrealHeaderTool/Private/ClassMaps.h
--- a/Engine/Source/Programs/UnrealHeaderTool/Private/ClassMaps.h
+++ b/Engine/Source/Programs/UnrealHeaderTool/Private/ClassMaps.h
@@ -35,12 +35,14 @@ struct FArchiveTypeDefinePair
 extern TMap<FString, TSharedRef<FUnrealSourceFile> > GUnrealSourceFilesMap;
 extern TMap<UField*, TSharedRef<FUnrealTypeDefinitionInfo> > GTypeDefinitionInfoMap;
 extern TMap<const UPackage*, TArray<UField*>> GPackageSingletons;
+extern FCriticalSection GPackageSingletonsCriticalSection;
 extern TMap<UClass*, FString> GClassStrippedHeaderTextMap;
 extern TMap<UClass*, FString> GClassHeaderNameWithNoPathMap;
 extern TSet<FUnrealSourceFile*> GPublicSourceFileSet;
 extern TMap<FProperty*, FString> GArrayDimensions;
 extern TMap<UPackage*,  const FManifestModule*> GPackageToManifestModuleMap;
 extern TMap<void*, uint32> GGeneratedCodeHashes;
+extern FRWLock GGeneratedCodeHashesLock;
 extern TMap<UEnum*, EUnderlyingEnumType> GEnumUnderlyingTypes;
 extern TMap<FName, TSharedRef<FClassDeclarationMetaData> > GClassDeclarations;
 extern TSet<FProperty*> GUnsizedProperties;
diff --git a/Engine/Source/Programs/UnrealHeaderTool/Private/CodeGenerator.cpp b/Engine/Source/Programs/UnrealHeaderTool/Private/CodeGenerator.cpp
--- a/Engine/Source/Programs/UnrealHeaderTool/Private/CodeGenerator.cpp
+++ b/Engine/Source/Programs/UnrealHeaderTool/Private/CodeGenerator.cpp
@@ -53,6 +53,7 @@
 #include "Algo/Copy.h"
 #include "Algo/Sort.h"
 #include "Algo/Reverse.h"
+#include "Async/ParallelFor.h"
 #include "Misc/ScopeExit.h"
 #include "UnrealTypeDefinitionInfo.h"
 
@@ -457,14 +458,31 @@ FString Macroize(const TCHAR* MacroName, const TCHAR* StringToMacroize)
 	return Macroize(MacroName, FString(StringToMacroize));
 }
 
+static void AddGeneratedCodeHash(void* Field, uint32 Hash)
+{
+	FRWScopeLock Lock(GGeneratedCodeHashesLock, SLT_Write);
+	GGeneratedCodeHashes.Add(Field, Hash);
+}
+
 /** Generates a Hash tag string for the specified field */
 static FString GetGeneratedCodeHashTag(void* Field)
 {
 	FString Tag;
-	const uint32* FieldHash = GGeneratedCodeHashes.Find(Field);
-	if (FieldHash)
+	bool bFoundHash = false;
+	uint32 Hash = 0;
+
+	{
+		FRWScopeLock Lock(GGeneratedCodeHashesLock, SLT_ReadOnly);
+		if (const uint32* FieldHash = GGeneratedCodeHashes.Find(Field))
+		{
+			bFoundHash = true;
+			Hash = *FieldHash;
+		}
+	}
+
+	if (bFoundHash)
 	{
-		Tag = FString::Printf(TEXT(" // %u"), *FieldHash);
+		Tag = FString::Printf(TEXT(" // %u"), Hash);
 	}
 	return Tag;
 }
@@ -812,9 +830,16 @@ public:
 	 */
 	const FString& GetExternDecl() const
 	{
+		FRWScopeLock Lock(ExternDeclLock, SLT_ReadOnly);
 		if (ExternDecl.IsEmpty())
 		{
-			ExternDecl = GenerateExternDecl(Type, GetName());
+			Lock.ReleaseReadOnlyLockAndAcquireWriteLock_USE_WITH_CAUTION();
+
+			// Verify the decl is still empty in case another thread had also been waiting on writing this data and got the write lock first
+			if (ExternDecl.IsEmpty())
+			{
+				ExternDecl = GenerateExternDecl(Type, GetName());
+			}
 		}
 
 		return ExternDecl;
@@ -863,6 +888,9 @@ private:
 	/** Cached field that stores this singleton extern declaration. */
 	mutable FString ExternDecl;
 
+	/** Mutex to ensure 2 threads don't try to generate the extern decl at the same time. */
+	mutable FRWLock ExternDeclLock;
+
 	/** Type of the singleton */
 	UField* Type;
 };
@@ -881,17 +909,25 @@ public:
 	 */
 	static const FTypeSingleton& Get(UField* Type, bool bRequiresValidObject = true)
 	{
-		static TMap<FTypeSingletonCacheKey, FTypeSingleton> CacheData;
-
 		FTypeSingletonCacheKey Key(Type, bRequiresValidObject);
-		if (FTypeSingleton* SingletonPtr = CacheData.Find(Key))
+
+		FRWScopeLock Lock(Mutex, SLT_ReadOnly);
+		TUniquePtr<FTypeSingleton>* SingletonPtr = CacheData.Find(Key);
+		if (SingletonPtr == nullptr)
 		{
-			return *SingletonPtr;
+			TUniquePtr<FTypeSingleton> TypeSingleton(MakeUnique<FTypeSingleton>(GenerateSingletonName(Type, bRequiresValidObject), Type));
+
+			Lock.ReleaseReadOnlyLockAndAcquireWriteLock_USE_WITH_CAUTION();
+
+			// Check the map again in case another thread had also been waiting on writing this data and got the write lock first
+			SingletonPtr = CacheData.Find(Key);
+			if (SingletonPtr == nullptr)
+			{
+				SingletonPtr = &CacheData.Add(Key, MoveTemp(TypeSingleton));
+			}
 		}
 
-		return CacheData.Add(Key,
-			FTypeSingleton(GenerateSingletonName(Type, bRequiresValidObject), Type)
-		);
+		return **SingletonPtr;
 	}
 
 private:
@@ -987,24 +1023,28 @@ private:
 		const TCHAR* ClassString = NameLookupCPP.GetNameCPP(Item->GetClass());
 		return FString::Printf(TEXT("Z_Construct_%s_%s%s()"), ClassString, *Result, Suffix);
 	}
+
+	static TMap<FTypeSingletonCacheKey, TUniquePtr<FTypeSingleton>> CacheData;
+	static FRWLock Mutex;
 };
 
-FString FNativeClassHeaderGenerator::GetSingletonName(UField* Item, bool bRequiresValidObject)
+TMap<FTypeSingletonCache::FTypeSingletonCacheKey, TUniquePtr<FTypeSingleton>> FTypeSingletonCache::CacheData;
+FRWLock FTypeSingletonCache::Mutex;
+
+const FString& FNativeClassHeaderGenerator::GetSingletonName(UField* Item, TSet<FString>* UniqueCrossModuleReferences, bool bRequiresValidObject)
 {
 	const FTypeSingleton& Cache = FTypeSingletonCache::Get(Item, bRequiresValidObject);
 
-	FString Result = Cache.GetName();
-
 	// We don't need to export UFunction externs, though we may need the externs for UDelegateFunctions
 	if (UniqueCrossModuleReferences && (!Item->IsA<UFunction>() || Item->IsA<UDelegateFunction>()))
 	{
 		UniqueCrossModuleReferences->Add(Cache.GetExternDecl());
 	}
 
-	return Result;
+	return Cache.GetName();
 }
 
-FString FNativeClassHeaderGenerator::GetSingletonNameFuncAddr(UField* Item, bool bRequiresValidObject)
+FString FNativeClassHeaderGenerator::GetSingletonNameFuncAddr(UField* Item, TSet<FString>* UniqueCrossModuleReferences, bool bRequiresValidObject)
 {
 	FString Result;
 	if (!Item)
@@ -1013,12 +1053,12 @@ FString FNativeClassHeaderGenerator::GetSingletonNameFuncAddr(UField* Item, bool
 	}
 	else
 	{
-		Result = this->GetSingletonName(Item, bRequiresValidObject).LeftChop(2);
+		Result = GetSingletonName(Item, UniqueCrossModuleReferences, bRequiresValidObject).LeftChop(2);
 	}
 	return Result;
 }
 
-void FNativeClassHeaderGenerator::PropertyNew(FOutputDevice& DeclOut, FOutputDevice& Out, FProperty* Prop, const TCHAR* OffsetStr, const TCHAR* Name, const TCHAR* DeclSpaces, const TCHAR* Spaces, const TCHAR* SourceStruct)
+void FNativeClassHeaderGenerator::PropertyNew(FOutputDevice& DeclOut, FOutputDevice& Out, FReferenceGatherers& OutReferenceGatherers, FProperty* Prop, const TCHAR* OffsetStr, const TCHAR* Name, const TCHAR* DeclSpaces, const TCHAR* Spaces, const TCHAR* SourceStruct) const
 {
 	FString        PropName             = CreateUTF8LiteralString(FNativeClassHeaderGenerator::GetOverriddenName(Prop));
 	FString        PropNameDep          = Prop->HasAllPropertyFlags(CPF_Deprecated) ? Prop->GetName() + TEXT("_DEPRECATED") : Prop->GetName();
@@ -1055,7 +1095,7 @@ void FNativeClassHeaderGenerator::PropertyNew(FOutputDevice& DeclOut, FOutputDev
 			FPropertyObjectFlags,
 			*ArrayDim,
 			OffsetStr,
-			*GetSingletonNameFuncAddr(TypedProp->Enum),
+			*GetSingletonNameFuncAddr(TypedProp->Enum, OutReferenceGatherers.UniqueCrossModuleReferences),
 			*MetaDataParams,
 			*PropTag
 		);
@@ -1317,7 +1357,7 @@ void FNativeClassHeaderGenerator::PropertyNew(FOutputDevice& DeclOut, FOutputDev
 			FPropertyObjectFlags,
 			*ArrayDim,
 			OffsetStr,
-			*GetSingletonNameFuncAddr(TypedProp->MetaClass, false),
+			*GetSingletonNameFuncAddr(TypedProp->MetaClass, OutReferenceGatherers.UniqueCrossModuleReferences, false),
 			*MetaDataParams,
 			*PropTag
 		);
@@ -1339,7 +1379,7 @@ void FNativeClassHeaderGenerator::PropertyNew(FOutputDevice& DeclOut, FOutputDev
 			FPropertyObjectFlags,
 			*ArrayDim,
 			OffsetStr,
-			*GetSingletonNameFuncAddr(TypedProp->PropertyClass, false),
+			*GetSingletonNameFuncAddr(TypedProp->PropertyClass, OutReferenceGatherers.UniqueCrossModuleReferences, false),
 			*MetaDataParams,
 			*PropTag
 		);
@@ -1361,7 +1401,7 @@ void FNativeClassHeaderGenerator::PropertyNew(FOutputDevice& DeclOut, FOutputDev
 			FPropertyObjectFlags,
 			*ArrayDim,
 			OffsetStr,
-			*GetSingletonNameFuncAddr(TypedProp->PropertyClass, false),
+			*GetSingletonNameFuncAddr(TypedProp->PropertyClass, OutReferenceGatherers.UniqueCrossModuleReferences, false),
 			*MetaDataParams,
 			*PropTag
 		);
@@ -1383,7 +1423,7 @@ void FNativeClassHeaderGenerator::PropertyNew(FOutputDevice& DeclOut, FOutputDev
 			FPropertyObjectFlags,
 			*ArrayDim,
 			OffsetStr,
-			*GetSingletonNameFuncAddr(TypedProp->PropertyClass, false),
+			*GetSingletonNameFuncAddr(TypedProp->PropertyClass, OutReferenceGatherers.UniqueCrossModuleReferences, false),
 			*MetaDataParams,
 			*PropTag
 		);
@@ -1405,8 +1445,8 @@ void FNativeClassHeaderGenerator::PropertyNew(FOutputDevice& DeclOut, FOutputDev
 			FPropertyObjectFlags,
 			*ArrayDim,
 			OffsetStr,
-			*GetSingletonNameFuncAddr(TypedProp->MetaClass, false),
-			*GetSingletonNameFuncAddr(TypedProp->PropertyClass, false),
+			*GetSingletonNameFuncAddr(TypedProp->MetaClass, OutReferenceGatherers.UniqueCrossModuleReferences, false),
+			*GetSingletonNameFuncAddr(TypedProp->PropertyClass, OutReferenceGatherers.UniqueCrossModuleReferences, false),
 			*MetaDataParams,
 			*PropTag
 		);
@@ -1428,7 +1468,7 @@ void FNativeClassHeaderGenerator::PropertyNew(FOutputDevice& DeclOut, FOutputDev
 			FPropertyObjectFlags,
 			*ArrayDim,
 			OffsetStr,
-			*GetSingletonNameFuncAddr(TypedProp->PropertyClass, false),
+			*GetSingletonNameFuncAddr(TypedProp->PropertyClass, OutReferenceGatherers.UniqueCrossModuleReferences, false),
 			*MetaDataParams,
 			*PropTag
 		);
@@ -1450,7 +1490,7 @@ void FNativeClassHeaderGenerator::PropertyNew(FOutputDevice& DeclOut, FOutputDev
 			FPropertyObjectFlags,
 			*ArrayDim,
 			OffsetStr,
-			*GetSingletonNameFuncAddr(TypedProp->InterfaceClass, false),
+			*GetSingletonNameFuncAddr(TypedProp->InterfaceClass, OutReferenceGatherers.UniqueCrossModuleReferences, false),
 			*MetaDataParams,
 			*PropTag
 		);
@@ -1577,7 +1617,7 @@ void FNativeClassHeaderGenerator::PropertyNew(FOutputDevice& DeclOut, FOutputDev
 			FPropertyObjectFlags,
 			*ArrayDim,
 			OffsetStr,
-			*GetSingletonNameFuncAddr(TypedProp->Struct),
+			*GetSingletonNameFuncAddr(TypedProp->Struct, OutReferenceGatherers.UniqueCrossModuleReferences),
 			*MetaDataParams,
 			*PropTag
 		);
@@ -1599,7 +1639,7 @@ void FNativeClassHeaderGenerator::PropertyNew(FOutputDevice& DeclOut, FOutputDev
 			FPropertyObjectFlags,
 			*ArrayDim,
 			OffsetStr,
-			*GetSingletonNameFuncAddr(TypedProp->SignatureFunction),
+			*GetSingletonNameFuncAddr(TypedProp->SignatureFunction, OutReferenceGatherers.UniqueCrossModuleReferences),
 			*MetaDataParams,
 			*PropTag
 		);
@@ -1622,7 +1662,7 @@ void FNativeClassHeaderGenerator::PropertyNew(FOutputDevice& DeclOut, FOutputDev
 			FPropertyObjectFlags,
 			*ArrayDim,
 			OffsetStr,
-			*GetSingletonNameFuncAddr(TypedProp->SignatureFunction),
+			*GetSingletonNameFuncAddr(TypedProp->SignatureFunction, OutReferenceGatherers.UniqueCrossModuleReferences),
 			*MetaDataParams,
 			*PropTag
 		);
@@ -1665,7 +1705,7 @@ void FNativeClassHeaderGenerator::PropertyNew(FOutputDevice& DeclOut, FOutputDev
 			FPropertyObjectFlags,
 			*ArrayDim,
 			OffsetStr,
-			*GetSingletonNameFuncAddr(TypedProp->Enum),
+			*GetSingletonNameFuncAddr(TypedProp->Enum, OutReferenceGatherers.UniqueCrossModuleReferences),
 			*MetaDataParams,
 			*PropTag
 		);
@@ -1714,7 +1754,7 @@ bool IsEditorOnlyDataProperty(FProperty* Prop)
 	return false;
 }
 
-TTuple<FString, FString> FNativeClassHeaderGenerator::OutputProperties(FOutputDevice& DeclOut, FOutputDevice& Out, const TCHAR* Scope, const TArray<FProperty*>& Properties, const TCHAR* DeclSpaces, const TCHAR* Spaces)
+TTuple<FString, FString> FNativeClassHeaderGenerator::OutputProperties(FOutputDevice& DeclOut, FOutputDevice& Out, FReferenceGatherers& OutReferenceGatherers, const TCHAR* Scope, const TArray<FProperty*>& Properties, const TCHAR* DeclSpaces, const TCHAR* Spaces) const
 {
 	if (Properties.Num() == 0)
 	{
@@ -1740,7 +1780,7 @@ TTuple<FString, FString> FNativeClassHeaderGenerator::OutputProperties(FOutputDe
 
 			WithEditorOnlyMacroEmitter(bRequiresHasEditorOnlyMacro);
 			WithEditorOnlyMacroEmitterDecl(bRequiresHasEditorOnlyMacro);
-			OutputProperty(DeclOut, Out, Scope, PropertyNamesAndPointers, Prop, DeclSpaces, Spaces);
+			OutputProperty(DeclOut, Out, OutReferenceGatherers, Scope, PropertyNamesAndPointers, Prop, DeclSpaces, Spaces);
 		}
 
 		WithEditorOnlyMacroEmitter(bHasAllEditorOnlyDataProperties);
@@ -1803,10 +1843,10 @@ inline FString GetEventStructParamsName(UObject* Outer, const TCHAR* FunctionNam
 	return Result;
 }
 
-void FNativeClassHeaderGenerator::OutputProperty(FOutputDevice& DeclOut, FOutputDevice& Out, const TCHAR* Scope, TArray<FPropertyNamePointerPair>& PropertyNamesAndPointers, FProperty* Prop, const TCHAR* DeclSpaces, const TCHAR* Spaces)
+void FNativeClassHeaderGenerator::OutputProperty(FOutputDevice& DeclOut, FOutputDevice& Out, FReferenceGatherers& OutReferenceGatherers, const TCHAR* Scope, TArray<FPropertyNamePointerPair>& PropertyNamesAndPointers, FProperty* Prop, const TCHAR* DeclSpaces, const TCHAR* Spaces) const
 {
 	// Helper to handle the creation of the underlying properties if they're enum properties
-	auto HandleUnderlyingEnumProperty = [this, &PropertyNamesAndPointers, &DeclOut, &Out, DeclSpaces, Spaces](FProperty* LocalProp, FString OuterName)
+	auto HandleUnderlyingEnumProperty = [this, &PropertyNamesAndPointers, &DeclOut, &Out, &OutReferenceGatherers, DeclSpaces, Spaces](FProperty* LocalProp, FString OuterName)
 	{
 		if (FEnumProperty* EnumProp = CastField<FEnumProperty>(LocalProp))
 		{
@@ -1814,7 +1854,7 @@ void FNativeClassHeaderGenerator::OutputProperty(FOutputDevice& DeclOut, FOutput
 
 			PropertyNamesAndPointers.Emplace(MoveTemp(OuterName), LocalProp);
 
-			PropertyNew(DeclOut, Out, EnumProp->UnderlyingProp, TEXT("0"), *PropVarName, DeclSpaces, Spaces);
+			PropertyNew(DeclOut, Out, OutReferenceGatherers, EnumProp->UnderlyingProp, TEXT("0"), *PropVarName, DeclSpaces, Spaces);
 			PropertyNamesAndPointers.Emplace(MoveTemp(PropVarName), EnumProp->UnderlyingProp);
 		}
 		else
@@ -1854,7 +1894,7 @@ void FNativeClassHeaderGenerator::OutputProperty(FOutputDevice& DeclOut, FOutput
 
 		FString PropMacroOuterClass = FString::Printf(TEXT("STRUCT_OFFSET(%s, %s)"), *SourceStruct, *PropName);
 
-		PropertyNew(DeclOut, Out, Prop, *PropMacroOuterClass, *PropVariableName, DeclSpaces, Spaces, *SourceStruct);
+		PropertyNew(DeclOut, Out, OutReferenceGatherers, Prop, *PropMacroOuterClass, *PropVariableName, DeclSpaces, Spaces, *SourceStruct);
 		HandleUnderlyingEnumProperty(Prop, MoveTemp(PropVariableName));
 	}
 
@@ -1862,7 +1902,7 @@ void FNativeClassHeaderGenerator::OutputProperty(FOutputDevice& DeclOut, FOutput
 	{
 		FString InnerVariableName = FString::Printf(TEXT("%sNewProp_%s_Inner"), Scope, *ArrayProperty->Inner->GetName());
 
-		PropertyNew(DeclOut, Out, ArrayProperty->Inner, TEXT("0"), *InnerVariableName, DeclSpaces, Spaces);
+		PropertyNew(DeclOut, Out, OutReferenceGatherers, ArrayProperty->Inner, TEXT("0"), *InnerVariableName, DeclSpaces, Spaces);
 		HandleUnderlyingEnumProperty(ArrayProperty->Inner, MoveTemp(InnerVariableName));
 	}
 
@@ -1874,10 +1914,10 @@ void FNativeClassHeaderGenerator::OutputProperty(FOutputDevice& DeclOut, FOutput
 		FString KeyVariableName   = FString::Printf(TEXT("%sNewProp_%s_KeyProp"), Scope, *Key->GetName());
 		FString ValueVariableName = FString::Printf(TEXT("%sNewProp_%s_ValueProp"), Scope, *Value->GetName());
 
-		PropertyNew(DeclOut, Out, Key, TEXT("0"), *KeyVariableName, DeclSpaces, Spaces);
+		PropertyNew(DeclOut, Out, OutReferenceGatherers, Key, TEXT("0"), *KeyVariableName, DeclSpaces, Spaces);
 		HandleUnderlyingEnumProperty(Key, MoveTemp(KeyVariableName));
 
-		PropertyNew(DeclOut, Out, Value, TEXT("1"), *ValueVariableName, DeclSpaces, Spaces);
+		PropertyNew(DeclOut, Out, OutReferenceGatherers, Value, TEXT("1"), *ValueVariableName, DeclSpaces, Spaces);
 		HandleUnderlyingEnumProperty(Value, MoveTemp(ValueVariableName));
 	}
 
@@ -1887,7 +1927,7 @@ void FNativeClassHeaderGenerator::OutputProperty(FOutputDevice& DeclOut, FOutput
 
 		FString ElementVariableName = FString::Printf(TEXT("%sNewProp_%s_ElementProp"), Scope, *Inner->GetName());
 
-		PropertyNew(DeclOut, Out, Inner, TEXT("0"), *ElementVariableName, DeclSpaces, Spaces);
+		PropertyNew(DeclOut, Out, OutReferenceGatherers, Inner, TEXT("0"), *ElementVariableName, DeclSpaces, Spaces);
 		HandleUnderlyingEnumProperty(Inner, MoveTemp(ElementVariableName));
 	}
 }
@@ -1981,7 +2021,7 @@ static TArray<UScriptStruct*> FindNoExportStructs(UStruct* Start)
 struct FPackageSingletonStrings
 {
 	FPackageSingletonStrings(FString&& InPackageSingletonName)
-		: PackageSingletonName(InPackageSingletonName)
+		: PackageSingletonName(MoveTemp(InPackageSingletonName))
 		, PackageUniqueCrossModuleReference(FString::Printf(TEXT("\tUPackage* %s;\r\n"), *PackageSingletonName))
 	{
 	}
@@ -1990,35 +2030,56 @@ struct FPackageSingletonStrings
 	FString PackageUniqueCrossModuleReference;
 };
 
-static TMap<const UPackage*, FPackageSingletonStrings> PackageSingletonNames;
-const FString& FNativeClassHeaderGenerator::GetPackageSingletonName(const UPackage* InPackage)
-{
-	FPackageSingletonStrings* PackageSingletonStrings = PackageSingletonNames.Find(InPackage);
-	if (UNLIKELY(PackageSingletonStrings == nullptr))
+static TMap<const UPackage*, TUniquePtr<FPackageSingletonStrings>> PackageSingletonNames;
+static FRWLock PackageSingletonNamesLock;
+
+const FString& FNativeClassHeaderGenerator::GetPackageSingletonName(const UPackage* InPackage, TSet<FString>* UniqueCrossModuleReferences)
 {
+	FRWScopeLock Lock(PackageSingletonNamesLock, SLT_ReadOnly);
+
+	TUniquePtr<FPackageSingletonStrings>* PackageSingletonStrings = PackageSingletonNames.Find(InPackage);
+	if (PackageSingletonStrings == nullptr)
+	{
 		static const TCHAR* ClassString = NameLookupCPP.GetNameCPP(UPackage::StaticClass());
 
 		FString PackageName = InPackage->GetName();
 		PackageName.ReplaceInline(TEXT("/"), TEXT("_"), ESearchCase::CaseSensitive);
 
-		PackageSingletonStrings = &PackageSingletonNames.Add(InPackage, FPackageSingletonStrings(FString::Printf(TEXT("Z_Construct_%s_%s()"), ClassString, *PackageName)));
+		TUniquePtr<FPackageSingletonStrings> NewPackageSingletonStrings(MakeUnique<FPackageSingletonStrings>(FString::Printf(TEXT("Z_Construct_%s_%s()"), ClassString, *PackageName)));
+
+		Lock.ReleaseReadOnlyLockAndAcquireWriteLock_USE_WITH_CAUTION();
+		
+		// Check the map again in case another thread had also been waiting on writing this data and got the write lock first
+		PackageSingletonStrings = PackageSingletonNames.Find(InPackage);
+
+		if (PackageSingletonStrings == nullptr)
+		{
+			PackageSingletonStrings = &PackageSingletonNames.Add(InPackage, MoveTemp(NewPackageSingletonStrings));
+		}
 	}
 
 	if (UniqueCrossModuleReferences)
 	{
-		UniqueCrossModuleReferences->Add(PackageSingletonStrings->PackageUniqueCrossModuleReference);
+		UniqueCrossModuleReferences->Add((*PackageSingletonStrings)->PackageUniqueCrossModuleReference);
 	}
 
-	return PackageSingletonStrings->PackageSingletonName;
+	return (*PackageSingletonStrings)->PackageSingletonName;
 }
 
 void FNativeClassHeaderGenerator::ExportGeneratedPackageInitCode(FOutputDevice& Out, const TCHAR* InDeclarations, const UPackage* InPackage, uint32 Hash)
 {
-	const FString& SingletonName = GetPackageSingletonName(InPackage);
+	const FString& SingletonName = GetPackageSingletonName(InPackage, nullptr);
 
-	const TArray<UField*>* SingletonsToOutput = GPackageSingletons.Find(InPackage);
+	TArray<UField*>* SingletonsToOutput = GPackageSingletons.Find(InPackage);
 	if (SingletonsToOutput)
 	{
+		Algo::Sort(*SingletonsToOutput, [](UField* A, UField* B)
+		{
+			// Structs before delegates then UniqueId order
+			return (uint64(A->IsA<UDelegateFunction>()) << 32) + A->GetUniqueID() <
+			       (uint64(B->IsA<UDelegateFunction>()) << 32) + B->GetUniqueID();
+		});
+
 		for (UField* ScriptType : *SingletonsToOutput)
 		{
 			Out.Log(FTypeSingletonCache::Get(ScriptType, true).GetExternDecl());
@@ -2071,7 +2132,7 @@ void FNativeClassHeaderGenerator::ExportGeneratedPackageInitCode(FOutputDevice&
 	Out.Logf(TEXT("\t}\r\n"));
 }
 
-void FNativeClassHeaderGenerator::ExportNativeGeneratedInitCode(FOutputDevice& Out, FOutputDevice& OutDeclarations, const FUnrealSourceFile& SourceFile, FClass* Class, FUHTStringBuilder& OutFriendText)
+void FNativeClassHeaderGenerator::ExportNativeGeneratedInitCode(FOutputDevice& Out, FOutputDevice& OutDeclarations, FReferenceGatherers& OutReferenceGatherers, const FUnrealSourceFile& SourceFile, FClass* Class, FUHTStringBuilder& OutFriendText) const
 {
 	check(!OutFriendText.Len());
 
@@ -2116,7 +2177,7 @@ void FNativeClassHeaderGenerator::ExportNativeGeneratedInitCode(FOutputDevice& O
 	{
 		// simple ::StaticClass wrapper to avoid header, link and DLL hell
 		{
-			FString SingletonNameNoRegister = GetSingletonName(Class, false);
+			const FString& SingletonNameNoRegister = GetSingletonName(Class, OutReferenceGatherers.UniqueCrossModuleReferences, false);
 
 			OutDeclarations.Log(FTypeSingletonCache::Get(Class, false).GetExternDecl());
 
@@ -2125,7 +2186,7 @@ void FNativeClassHeaderGenerator::ExportNativeGeneratedInitCode(FOutputDevice& O
 			GeneratedClassRegisterFunctionText.Logf(TEXT("\t\treturn %s::StaticClass();\r\n"), ClassNameCPP);
 			GeneratedClassRegisterFunctionText.Logf(TEXT("\t}\r\n"));
 		}
-		FString SingletonName = GetSingletonName(Class);
+		const FString& SingletonName = GetSingletonName(Class, OutReferenceGatherers.UniqueCrossModuleReferences);
 
 		FString StaticsStructName = SingletonName.LeftChop(2) + TEXT("_Statics");
 
@@ -2142,11 +2203,11 @@ void FNativeClassHeaderGenerator::ExportNativeGeneratedInitCode(FOutputDevice& O
 		if (SuperClass && SuperClass != Class)
 		{
 			OutDeclarations.Log(FTypeSingletonCache::Get(SuperClass).GetExternDecl());
-			Singletons.Logf(TEXT("\t\t(UObject* (*)())%s,\r\n"), *GetSingletonName(SuperClass).LeftChop(2));
+			Singletons.Logf(TEXT("\t\t(UObject* (*)())%s,\r\n"), *GetSingletonName(SuperClass, OutReferenceGatherers.UniqueCrossModuleReferences).LeftChop(2));
 		}
 		if (!bIsDynamic)
 		{
-			FString PackageSingletonName = GetPackageSingletonName(Class->GetOutermost());
+			const FString& PackageSingletonName = GetPackageSingletonName(Class->GetOutermost(), OutReferenceGatherers.UniqueCrossModuleReferences);
 
 			OutDeclarations.Logf(TEXT("\t%s_API UPackage* %s;\r\n"), *ApiString, *PackageSingletonName);
 			Singletons.Logf(TEXT("\t\t(UObject* (*)())%s,\r\n"), *PackageSingletonName.LeftChop(2));
@@ -2188,13 +2249,13 @@ void FNativeClassHeaderGenerator::ExportNativeGeneratedInitCode(FOutputDevice& O
 
 				if (!Function->IsA<UDelegateFunction>())
 				{
-					ExportFunction(Out, SourceFile, Function, bIsNoExport);
+					ExportFunction(Out, OutReferenceGatherers, SourceFile, Function, bIsNoExport);
 				}
 
 				StaticDefinitions.Logf(
 					TEXT("%s\t\t{ &%s, %s },%s\r\n%s"),
 					BEGIN_WRAP_EDITOR_ONLY(bIsEditorOnlyFunction),
-					*GetSingletonNameFuncAddr(Function),
+					*GetSingletonNameFuncAddr(Function, OutReferenceGatherers.UniqueCrossModuleReferences),
 					*FNativeClassHeaderGenerator::GetUTF8OverriddenNameForLiteral(Function),
 					*GetGeneratedCodeHashTag(Function),
 					END_WRAP_EDITOR_ONLY(bIsEditorOnlyFunction)
@@ -2239,7 +2300,7 @@ void FNativeClassHeaderGenerator::ExportNativeGeneratedInitCode(FOutputDevice& O
 			Props.Add(Prop);
 		}
 
-		TTuple<FString, FString> PropertyRange = OutputProperties(GeneratedClassRegisterFunctionText, StaticDefinitions, *FString::Printf(TEXT("%s::"), *StaticsStructName), Props, TEXT("\t\t"), TEXT("\t"));
+		TTuple<FString, FString> PropertyRange = OutputProperties(GeneratedClassRegisterFunctionText, StaticDefinitions, OutReferenceGatherers, *FString::Printf(TEXT("%s::"), *StaticsStructName), Props, TEXT("\t\t"), TEXT("\t"));
 
 		const TCHAR* InterfaceArray;
 		const TCHAR* InterfaceCount;
@@ -2262,7 +2323,7 @@ void FNativeClassHeaderGenerator::ExportNativeGeneratedInitCode(FOutputDevice& O
 				}
 				StaticDefinitions.Logf(
 					TEXT("\t\t\t{ %s, %s, %s },\r\n"),
-					*GetSingletonName(Inter.Class, false).LeftChop(2),
+					*GetSingletonName(Inter.Class, OutReferenceGatherers.UniqueCrossModuleReferences, false).LeftChop(2),
 					*OffsetString,
 					Inter.bImplementedByK2 ? TEXT("true") : TEXT("false")
 				);
@@ -2362,7 +2423,7 @@ void FNativeClassHeaderGenerator::ExportNativeGeneratedInitCode(FOutputDevice& O
 		OutFriendText.Reset();
 	}
 
-	FString SingletonName = GetSingletonName(Class);
+	FString SingletonName = GetSingletonName(Class, OutReferenceGatherers.UniqueCrossModuleReferences);
 	SingletonName.ReplaceInline(TEXT("()"), TEXT(""), ESearchCase::CaseSensitive); // function address
 
 	FString OverriddenClassName = *FNativeClassHeaderGenerator::GetOverriddenName(Class);
@@ -2378,7 +2439,23 @@ void FNativeClassHeaderGenerator::ExportNativeGeneratedInitCode(FOutputDevice& O
 	FClass* SuperClass = Class->GetSuperClass();
 	if (SuperClass && !SuperClass->HasAnyClassFlags(CLASS_Intrinsic))
 	{
-		BaseClassHash = GGeneratedCodeHashes.FindChecked(SuperClass);
+		// Since we are dependent on our SuperClass having generated its hash, if it is not available
+		// we will need to wait on it becoming available. Since the SourceFile array provided to the
+		// ParallelFor is in dependency order and does not allow cyclic dependencies, we can be certain
+		// that another thread has started processing the file containing our SuperClass before this
+		// file would have been assigned out,  so we just have to wait
+		while (1)
+		{
+			{
+				FRWScopeLock Lock(GGeneratedCodeHashesLock, SLT_ReadOnly);
+				if (const uint32* Hash = GGeneratedCodeHashes.Find(SuperClass))
+				{
+					BaseClassHash = *Hash;
+					break;
+				}
+			}
+			FPlatformProcess::Sleep(0.01);
+		}
 	}
 	GeneratedClassRegisterFunctionText.Logf(TEXT("\r\n// %u\r\n"), BaseClassHash);
 
@@ -2402,7 +2479,7 @@ void FNativeClassHeaderGenerator::ExportNativeGeneratedInitCode(FOutputDevice& O
 
 	// Calculate generated class initialization code hash so that we know when it changes after hot-reload
 	uint32 ClassHash = GenerateTextHash(*GeneratedClassRegisterFunctionText);
-	GGeneratedCodeHashes.Add(Class, ClassHash);
+	AddGeneratedCodeHash(Class, ClassHash);
 	// Emit the IMPLEMENT_CLASS macro to go in the generated cpp file.
 	if (!bIsDynamic)
 	{
@@ -2486,7 +2563,7 @@ void FNativeClassHeaderGenerator::ExportNativeGeneratedInitCode(FOutputDevice& O
 	}
 }
 
-void FNativeClassHeaderGenerator::ExportFunction(FOutputDevice& Out, const FUnrealSourceFile& SourceFile, UFunction* Function, bool bIsNoExport)
+void FNativeClassHeaderGenerator::ExportFunction(FOutputDevice& Out, FReferenceGatherers& OutReferenceGatherers, const FUnrealSourceFile& SourceFile, UFunction* Function, bool bIsNoExport) const
 {
 	UFunction* SuperFunction = Function->GetSuperFunction();
 
@@ -2494,7 +2571,7 @@ void FNativeClassHeaderGenerator::ExportFunction(FOutputDevice& Out, const FUnre
 
 	bool bIsDelegate = Function->HasAnyFunctionFlags(FUNC_Delegate);
 
-	FString SingletonName = GetSingletonName(Function);
+	const FString& SingletonName = GetSingletonName(Function, OutReferenceGatherers.UniqueCrossModuleReferences);
 	FString StaticsStructName = SingletonName.LeftChop(2) + TEXT("_Statics");
 
 	FUHTStringBuilder CurrentFunctionText;
@@ -2517,7 +2594,7 @@ void FNativeClassHeaderGenerator::ExportFunction(FOutputDevice& Out, const FUnre
 			ExportMirrorsForNoexportStruct(CurrentFunctionText, Struct, /*Indent=*/ 2);
 		}
 
-		ExportEventParm(CurrentFunctionText, ForwardDeclarations, Function, /*Indent=*/ 2, /*bOutputConstructor=*/ false, EExportingState::TypeEraseDelegates);
+		ExportEventParm(CurrentFunctionText, OutReferenceGatherers.ForwardDeclarations, Function, /*Indent=*/ 2, /*bOutputConstructor=*/ false, EExportingState::TypeEraseDelegates);
 	}
 
 	UField* FieldOuter = Cast<UField>(Function->GetOuter());
@@ -2526,7 +2603,7 @@ void FNativeClassHeaderGenerator::ExportFunction(FOutputDevice& Out, const FUnre
 	FString OuterFunc;
 	if (UObject* Outer = Function->GetOuter())
 	{
-		OuterFunc = Outer->IsA<UPackage>() ? GetPackageSingletonName((UPackage*)Outer).LeftChop(2) : GetSingletonNameFuncAddr(Function->GetOwnerClass());
+		OuterFunc = Outer->IsA<UPackage>() ? GetPackageSingletonName((UPackage*)Outer, OutReferenceGatherers.UniqueCrossModuleReferences).LeftChop(2) : GetSingletonNameFuncAddr(Function->GetOwnerClass(), OutReferenceGatherers.UniqueCrossModuleReferences);
 	}
 	else
 	{
@@ -2560,7 +2637,7 @@ void FNativeClassHeaderGenerator::ExportFunction(FOutputDevice& Out, const FUnre
 	USparseDelegateFunction* SparseDelegateFunction = Cast<USparseDelegateFunction>(Function);
 	const TCHAR* UFunctionObjectFlags = FClass::IsOwnedByDynamicType(Function) ? TEXT("RF_Public|RF_Transient") : TEXT("RF_Public|RF_Transient|RF_MarkAsNative");
 
-	TTuple<FString, FString> PropertyRange = OutputProperties(CurrentFunctionText, StaticDefinitions, *FString::Printf(TEXT("%s::"), *StaticsStructName), Props, TEXT("\t\t"), TEXT("\t"));
+	TTuple<FString, FString> PropertyRange = OutputProperties(CurrentFunctionText, StaticDefinitions, OutReferenceGatherers, *FString::Printf(TEXT("%s::"), *StaticsStructName), Props, TEXT("\t\t"), TEXT("\t"));
 
 	const FFunctionData* CompilerInfo = FFunctionData::FindForFunction(Function);
 	const FFuncInfo&     FunctionData = CompilerInfo->GetFunctionData();
@@ -2574,7 +2651,7 @@ void FNativeClassHeaderGenerator::ExportFunction(FOutputDevice& Out, const FUnre
 		TEXT("\tconst UE4CodeGen_Private::FFunctionParams %s::FuncParams = { (UObject*(*)())%s, %s, %s, %s, %s, %s, %s, %s, %s, (EFunctionFlags)0x%08X, %d, %d, %s };\r\n"),
 		*StaticsStructName,
 		*OuterFunc,
-		*GetSingletonNameFuncAddr(SuperFunction),
+		*GetSingletonNameFuncAddr(SuperFunction, OutReferenceGatherers.UniqueCrossModuleReferences),
 		*CreateUTF8LiteralString(FNativeClassHeaderGenerator::GetOverriddenName(Function)),
 		(SparseDelegateFunction ? *CreateUTF8LiteralString(SparseDelegateFunction->OwningClassName.ToString()) : TEXT("nullptr")),
 		(SparseDelegateFunction ? *CreateUTF8LiteralString(SparseDelegateFunction->DelegateName.ToString()) : TEXT("nullptr")),
@@ -2619,7 +2696,7 @@ void FNativeClassHeaderGenerator::ExportFunction(FOutputDevice& Out, const FUnre
 	}
 
 	uint32 FunctionHash = GenerateTextHash(*CurrentFunctionText);
-	GGeneratedCodeHashes.Add(Function, FunctionHash);
+	AddGeneratedCodeHash(Function, FunctionHash);
 	Out.Log(CurrentFunctionText);
 }
 
@@ -2691,7 +2768,7 @@ void FNativeClassHeaderGenerator::ExportNatives(FOutputDevice& Out, FClass* Clas
 	Out.Logf(TEXT("\t}\r\n"));
 }
 
-void FNativeClassHeaderGenerator::ExportInterfaceCallFunctions(FOutputDevice& OutCpp, FUHTStringBuilder& Out, const TArray<UFunction*>& CallbackFunctions, const TCHAR* ClassName)
+void FNativeClassHeaderGenerator::ExportInterfaceCallFunctions(FOutputDevice& OutCpp, FUHTStringBuilder& Out, FReferenceGatherers& OutReferenceGatherers, const TArray<UFunction*>& CallbackFunctions, const TCHAR* ClassName) const
 {
 	const FString& APIString = GetAPIString();
 
@@ -2705,13 +2782,13 @@ void FNativeClassHeaderGenerator::ExportInterfaceCallFunctions(FOutputDevice& Ou
 		const TCHAR* ConstQualifier = FunctionData.FunctionReference->HasAllFunctionFlags(FUNC_Const) ? TEXT("const ") : TEXT("");
 		FString ExtraParam = FString::Printf(TEXT("%sUObject* O"), ConstQualifier);
 
-		ExportNativeFunctionHeader(Out, ForwardDeclarations, FunctionData, EExportFunctionType::Interface, EExportFunctionHeaderStyle::Declaration, *ExtraParam, *APIString);
+		ExportNativeFunctionHeader(Out, OutReferenceGatherers.ForwardDeclarations, FunctionData, EExportFunctionType::Interface, EExportFunctionHeaderStyle::Declaration, *ExtraParam, *APIString);
 		Out.Logf( TEXT(";") LINE_TERMINATOR );
 
 		FString FunctionNameName = FString::Printf(TEXT("NAME_%s_%s"), NameLookupCPP.GetNameCPP(CastChecked<UStruct>(Function->GetOuter())), *FunctionName);
 		OutCpp.Logf(TEXT("\tstatic FName %s = FName(TEXT(\"%s\"));") LINE_TERMINATOR, *FunctionNameName, *GetOverriddenFName(Function).ToString());
 
-		ExportNativeFunctionHeader(OutCpp, ForwardDeclarations, FunctionData, EExportFunctionType::Interface, EExportFunctionHeaderStyle::Definition, *ExtraParam, *APIString);
+		ExportNativeFunctionHeader(OutCpp, OutReferenceGatherers.ForwardDeclarations, FunctionData, EExportFunctionType::Interface, EExportFunctionHeaderStyle::Definition, *ExtraParam, *APIString);
 		OutCpp.Logf( LINE_TERMINATOR TEXT("\t{") LINE_TERMINATOR );
 
 		OutCpp.Logf(TEXT("\t\tcheck(O != NULL);") LINE_TERMINATOR);
@@ -2876,9 +2953,10 @@ void FNativeClassHeaderGenerator::ExportClassFromSourceFileInner(
 	FOutputDevice&           OutGeneratedHeaderText,
 	FOutputDevice&           OutCpp,
 	FOutputDevice&           OutDeclarations,
+	FReferenceGatherers&     OutReferenceGatherers,
 	FClass*                  Class,
 	const FUnrealSourceFile& SourceFile
-)
+) const
 {
 	FUHTStringBuilder StandardUObjectConstructorsMacroCall;
 	FUHTStringBuilder EnhancedUObjectConstructorsMacroCall;
@@ -2889,7 +2967,7 @@ void FNativeClassHeaderGenerator::ExportClassFromSourceFileInner(
 	// C++ -> VM stubs (native function execs)
 	FUHTStringBuilder ClassMacroCalls;
 	FUHTStringBuilder ClassNoPureDeclsMacroCalls;
-	ExportNativeFunctions(OutGeneratedHeaderText, OutCpp, ClassMacroCalls, ClassNoPureDeclsMacroCalls, SourceFile, Class, ClassData);
+	ExportNativeFunctions(OutGeneratedHeaderText, OutCpp, ClassMacroCalls, ClassNoPureDeclsMacroCalls, OutReferenceGatherers, SourceFile, Class, ClassData);
 
 	// Get Callback functions
 	TArray<UFunction*> CallbackFunctions;
@@ -2911,7 +2989,7 @@ void FNativeClassHeaderGenerator::ExportClassFromSourceFileInner(
 		// export parameters structs for all events and delegates
 		for (UFunction* Function : CallbackFunctions)
 		{
-			ExportEventParm(UClassMacroContent, ForwardDeclarations, Function, /*Indent=*/ 1, /*bOutputConstructor=*/ true, EExportingState::Normal);
+			ExportEventParm(UClassMacroContent, OutReferenceGatherers.ForwardDeclarations, Function, /*Indent=*/ 1, /*bOutputConstructor=*/ true, EExportingState::Normal);
 		}
 
 		FString MacroName = SourceFile.GetGeneratedMacroName(ClassData, TEXT("_EVENT_PARMS"));
@@ -2925,7 +3003,7 @@ void FNativeClassHeaderGenerator::ExportClassFromSourceFileInner(
 		ExportCallbackFunctions(
 			OutGeneratedHeaderText,
 			CallbackOut,
-			ForwardDeclarations,
+			OutReferenceGatherers.ForwardDeclarations,
 			CallbackFunctions,
 			*CallbackWrappersMacroName,
 			(Class->ClassFlags & CLASS_Interface) ? EExportCallbackType::Interface : EExportCallbackType::Class,
@@ -2943,7 +3021,7 @@ void FNativeClassHeaderGenerator::ExportClassFromSourceFileInner(
 	}
 
 	FUHTStringBuilder FriendText;
-	ExportNativeGeneratedInitCode(OutCpp, OutDeclarations, SourceFile, Class, FriendText);
+	ExportNativeGeneratedInitCode(OutCpp, OutDeclarations, OutReferenceGatherers, SourceFile, Class, FriendText);
 
 	FClass* SuperClass = Class->GetSuperClass();
 
@@ -3105,7 +3183,7 @@ void FNativeClassHeaderGenerator::ExportClassFromSourceFileInner(
 			InterfaceBoilerplate.Logf(TEXT("public:\r\n\ttypedef %s UClassType;\r\n"), ClassCPPName);
 			InterfaceBoilerplate.Logf(TEXT("\ttypedef %s ThisClass;\r\n"), *InterfaceCPPName);
 
-			ExportInterfaceCallFunctions(OutCpp, InterfaceBoilerplate, CallbackFunctions, *Class->GetName());
+			ExportInterfaceCallFunctions(OutCpp, InterfaceBoilerplate, OutReferenceGatherers, CallbackFunctions, *Class->GetName());
 
 			// we'll need a way to get to the UObject portion of a native interface, so that we can safely pass native interfaces
 			// to script VM functions
@@ -3295,6 +3373,16 @@ void ExportConstructorDefinition(FOutputDevice& Out, FClass* Class, FClassMetaDa
 			FClassMetaData* SuperClassData = GScriptHelper.FindClassData(SuperClass);
 			if (SuperClassData)
 			{
+				// Since we are dependent on our SuperClass having determined which constructors are defined, 
+				// if it is not yet determined we will need to wait on it becoming available. 
+				// Since the SourceFile array provided to the ParallelFor is in dependency order and does not allow cyclic dependencies, 
+				// we can be certain that another thread has started processing the file containing our SuperClass before this
+				// file would have been assigned out,  so we just have to wait
+				while (!SuperClassData->bConstructorDeclared)
+				{
+					FPlatformProcess::Sleep(0.01);
+				}
+
 				bSuperClassObjectInitializerConstructorDeclared = SuperClassData->bObjectInitializerConstructorDeclared;
 			}
 		}
@@ -3393,7 +3481,7 @@ void FNativeClassHeaderGenerator::ExportConstructorsMacros(FOutputDevice& OutGen
 	EnhancedUObjectConstructorsMacroCall.Logf(TEXT("\t%s\r\n"), *EnhMacroName);
 }
 
-bool FNativeClassHeaderGenerator::WriteHeader(const TCHAR* Path, const FString& InBodyText, const TSet<FString>& InFwdDecl)
+bool FNativeClassHeaderGenerator::WriteHeader(const TCHAR* Path, const FString& InBodyText, FReferenceGatherers& InOutReferenceGatherers) const
 {
 	FUHTStringBuilder GeneratedHeaderTextWithCopyright;
 	GeneratedHeaderTextWithCopyright.Logf(TEXT("%s"), HeaderCopyright);
@@ -3402,7 +3490,7 @@ bool FNativeClassHeaderGenerator::WriteHeader(const TCHAR* Path, const FString&
 	GeneratedHeaderTextWithCopyright.Log(LINE_TERMINATOR);
 	GeneratedHeaderTextWithCopyright.Log(TEXT("PRAGMA_DISABLE_DEPRECATION_WARNINGS") LINE_TERMINATOR);
 
-	for (const FString& FWDecl : InFwdDecl)
+	for (const FString& FWDecl : InOutReferenceGatherers.ForwardDeclarations)
 	{
 		if (FWDecl.Len() > 0)
 		{
@@ -3413,7 +3501,7 @@ bool FNativeClassHeaderGenerator::WriteHeader(const TCHAR* Path, const FString&
 	GeneratedHeaderTextWithCopyright.Log(*InBodyText);
 	GeneratedHeaderTextWithCopyright.Log(TEXT("PRAGMA_ENABLE_DEPRECATION_WARNINGS") LINE_TERMINATOR);
 
-	bool bHasChanged = SaveHeaderIfChanged(Path, *GeneratedHeaderTextWithCopyright);
+	bool bHasChanged = SaveHeaderIfChanged(InOutReferenceGatherers, Path, *GeneratedHeaderTextWithCopyright);
 	return bHasChanged;
 }
 
@@ -3459,7 +3547,7 @@ FString FNativeClassHeaderGenerator::GetClassFlagExportText( UClass* Class )
 *
 * @param	Enums	the enums to export
 */
-void FNativeClassHeaderGenerator::ExportEnum(FOutputDevice& Out, UEnum* Enum)
+void FNativeClassHeaderGenerator::ExportEnum(FOutputDevice& Out, UEnum* Enum) const
 {
 	// Export FOREACH macro
 	Out.Logf( TEXT("#define FOREACH_ENUM_%s(op) "), *Enum->GetName().ToUpper() );
@@ -3511,7 +3599,7 @@ void FNativeClassHeaderGenerator::ExportEnum(FOutputDevice& Out, UEnum* Enum)
 }
 
 // Exports the header text for the list of structs specified (GENERATED_BODY impls)
-void FNativeClassHeaderGenerator::ExportGeneratedStructBodyMacros(FOutputDevice& OutGeneratedHeaderText, FOutputDevice& Out, const FUnrealSourceFile& SourceFile, UScriptStruct* Struct)
+void FNativeClassHeaderGenerator::ExportGeneratedStructBodyMacros(FOutputDevice& OutGeneratedHeaderText, FOutputDevice& Out, FReferenceGatherers& OutReferenceGatherers, const FUnrealSourceFile& SourceFile, UScriptStruct* Struct) const
 {
 	const bool bIsDynamic = FClass::IsDynamic(Struct);
 	const FString ActualStructName = FNativeClassHeaderGenerator::GetOverriddenName(Struct);
@@ -3521,7 +3609,7 @@ void FNativeClassHeaderGenerator::ExportGeneratedStructBodyMacros(FOutputDevice&
 
 	const TCHAR* StructNameCPP = NameLookupCPP.GetNameCPP(Struct);
 
-	const FString SingletonName = GetSingletonName(Struct);
+	const FString& SingletonName = GetSingletonName(Struct, OutReferenceGatherers.UniqueCrossModuleReferences);
 	const FString ChoppedSingletonName = SingletonName.LeftChop(2);
 
 	TArray<FString> RigVMVirtualFuncProlog, RigVMStubProlog;
@@ -3627,7 +3715,7 @@ void FNativeClassHeaderGenerator::ExportGeneratedStructBodyMacros(FOutputDevice&
 		Out.Logf(TEXT("{\r\n"));
 
 		// UStructs can have UClass or UPackage outer (if declared in non-UClass headers).
-		const FString& OuterName (bIsDynamic ? STRING_StructPackage : GetPackageSingletonName(CastChecked<UPackage>(Struct->GetOuter())));
+		const FString& OuterName (bIsDynamic ? STRING_StructPackage : GetPackageSingletonName(CastChecked<UPackage>(Struct->GetOuter()), OutReferenceGatherers.UniqueCrossModuleReferences));
 		if (!bIsDynamic)
 		{
 			Out.Logf(TEXT("\tstatic class UScriptStruct* Singleton = NULL;\r\n"));
@@ -3720,15 +3808,10 @@ void FNativeClassHeaderGenerator::ExportGeneratedStructBodyMacros(FOutputDevice&
 		ExportMirrorsForNoexportStruct(GeneratedStructRegisterFunctionText, NoExportStruct, /*Indent=*/ 2);
 	}
 
-	FString BaseStructString;
 	if (BaseStruct)
 	{
 		CastChecked<UScriptStruct>(BaseStruct); // this better actually be a script struct
-		BaseStructString = GetSingletonName(BaseStruct);
-	}
-	else
-	{
-		BaseStructString = TEXT("nullptr");
+		GetSingletonName(BaseStruct, OutReferenceGatherers.UniqueCrossModuleReferences); // Call to potentially collect references
 	}
 
 	EStructFlags UncomputedFlags = (EStructFlags)(Struct->StructFlags & ~STRUCT_ComputedFlags);
@@ -3736,7 +3819,7 @@ void FNativeClassHeaderGenerator::ExportGeneratedStructBodyMacros(FOutputDevice&
 	FString OuterFunc;
 	if (!bIsDynamic)
 	{
-		OuterFunc = GetPackageSingletonName(CastChecked<UPackage>(Struct->GetOuter())).LeftChop(2);
+		OuterFunc = GetPackageSingletonName(CastChecked<UPackage>(Struct->GetOuter()), OutReferenceGatherers.UniqueCrossModuleReferences).LeftChop(2);
 	}
 	else
 	{
@@ -3772,13 +3855,13 @@ void FNativeClassHeaderGenerator::ExportGeneratedStructBodyMacros(FOutputDevice&
 		NewStructOps = TEXT("nullptr");
 	}
 
-	TTuple<FString, FString> PropertyRange = OutputProperties(GeneratedStructRegisterFunctionText, StaticDefinitions, *FString::Printf(TEXT("%s::"), *StaticsStructName), Props, TEXT("\t\t"), TEXT("\t"));
+	TTuple<FString, FString> PropertyRange = OutputProperties(GeneratedStructRegisterFunctionText, StaticDefinitions, OutReferenceGatherers, *FString::Printf(TEXT("%s::"), *StaticsStructName), Props, TEXT("\t\t"), TEXT("\t"));
 
 	GeneratedStructRegisterFunctionText.Log (TEXT("\t\tstatic const UE4CodeGen_Private::FStructParams ReturnStructParams;\r\n"));
 
 	StaticDefinitions.Logf(TEXT("\tconst UE4CodeGen_Private::FStructParams %s::ReturnStructParams = {\r\n"), *StaticsStructName);
 	StaticDefinitions.Logf(TEXT("\t\t(UObject* (*)())%s,\r\n"), *OuterFunc);
-	StaticDefinitions.Logf(TEXT("\t\t%s,\r\n"), *GetSingletonNameFuncAddr(BaseStruct));
+	StaticDefinitions.Logf(TEXT("\t\t%s,\r\n"), *GetSingletonNameFuncAddr(BaseStruct, OutReferenceGatherers.UniqueCrossModuleReferences));
 	StaticDefinitions.Logf(TEXT("\t\t%s,\r\n"), *NewStructOps);
 	StaticDefinitions.Logf(TEXT("\t\t%s,\r\n"), *CreateUTF8LiteralString(ActualStructName));
 	StaticDefinitions.Logf(TEXT("\t\tsizeof(%s),\r\n"), StructNameCPP);
@@ -3813,7 +3896,7 @@ void FNativeClassHeaderGenerator::ExportGeneratedStructBodyMacros(FOutputDevice&
 	{
 		GeneratedStructRegisterFunctionText.Log (TEXT("#if WITH_HOT_RELOAD\r\n"));
 		GeneratedStructRegisterFunctionText.Logf(TEXT("\t\textern uint32 %s();\r\n"), *HashFuncName);
-		GeneratedStructRegisterFunctionText.Logf(TEXT("\t\tUPackage* Outer = %s;\r\n"), *GetPackageSingletonName(CastChecked<UPackage>(Struct->GetOuter())));
+		GeneratedStructRegisterFunctionText.Logf(TEXT("\t\tUPackage* Outer = %s;\r\n"), *GetPackageSingletonName(CastChecked<UPackage>(Struct->GetOuter()), OutReferenceGatherers.UniqueCrossModuleReferences));
 		GeneratedStructRegisterFunctionText.Logf(TEXT("\t\tstatic UScriptStruct* ReturnStruct = FindExistingStructIfHotReloadOrDynamic(Outer, TEXT(\"%s\"), sizeof(%s), %s(), false);\r\n"), *ActualStructName, *NoExportStructNameCPP, *HashFuncName);
 		GeneratedStructRegisterFunctionText.Log (TEXT("#else\r\n"));
 		GeneratedStructRegisterFunctionText.Logf(TEXT("\t\tstatic UScriptStruct* ReturnStruct = nullptr;\r\n"));
@@ -3834,7 +3917,7 @@ void FNativeClassHeaderGenerator::ExportGeneratedStructBodyMacros(FOutputDevice&
 	GeneratedStructRegisterFunctionText.Log (TEXT("\t}\r\n"));
 
 	uint32 StructHash = GenerateTextHash(*GeneratedStructRegisterFunctionText);
-	GGeneratedCodeHashes.Add(Struct, StructHash);
+	AddGeneratedCodeHash(Struct, StructHash);
 
 	Out.Log(GeneratedStructRegisterFunctionText);
 	Out.Logf(TEXT("\tuint32 %s() { return %uU; }\r\n"), *HashFuncName, StructHash);
@@ -3904,10 +3987,10 @@ void FNativeClassHeaderGenerator::ExportGeneratedStructBodyMacros(FOutputDevice&
 	}
 }
 
-void FNativeClassHeaderGenerator::ExportGeneratedEnumInitCode(FOutputDevice& Out, const FUnrealSourceFile& SourceFile, UEnum* Enum)
+void FNativeClassHeaderGenerator::ExportGeneratedEnumInitCode(FOutputDevice& Out, FReferenceGatherers& OutReferenceGatherers, const FUnrealSourceFile& SourceFile, UEnum* Enum) const
 {
 	const bool    bIsDynamic            = FClass::IsDynamic(static_cast<UField*>(Enum));
-	const FString SingletonName         = GetSingletonNameFuncAddr(Enum);
+	const FString SingletonName         = GetSingletonNameFuncAddr(Enum, OutReferenceGatherers.UniqueCrossModuleReferences);
 	const FString EnumNameCpp           = Enum->GetName(); //UserDefinedEnum should already have a valid cpp name.
 	const FString OverriddenEnumNameCpp = FNativeClassHeaderGenerator::GetOverriddenName(Enum);
 
@@ -3916,7 +3999,7 @@ void FNativeClassHeaderGenerator::ExportGeneratedEnumInitCode(FOutputDevice& Out
 	FMacroBlockEmitter EditorOnlyData(Out, TEXT("WITH_EDITORONLY_DATA"));
 	EditorOnlyData(bIsEditorOnlyDataType);
 
-	const FString& PackageSingletonName = (bIsDynamic ? FClass::GetTypePackageName(static_cast<UField*>(Enum)) : GetPackageSingletonName(CastChecked<UPackage>(Enum->GetOuter())));
+	const FString& PackageSingletonName = (bIsDynamic ? FClass::GetTypePackageName(static_cast<UField*>(Enum)) : GetPackageSingletonName(CastChecked<UPackage>(Enum->GetOuter()), OutReferenceGatherers.UniqueCrossModuleReferences));
 
 	Out.Logf(TEXT("\tstatic UEnum* %s_StaticEnum()\r\n"), *Enum->GetName());
 	Out.Logf(TEXT("\t{\r\n"));
@@ -3974,7 +4057,7 @@ void FNativeClassHeaderGenerator::ExportGeneratedEnumInitCode(FOutputDevice& Out
 		);
 	}
 
-	const FString EnumSingletonName = GetSingletonName(Enum);
+	const FString& EnumSingletonName = GetSingletonName(Enum, OutReferenceGatherers.UniqueCrossModuleReferences);
 	const FString HashFuncName       = FString::Printf(TEXT("Get_%s_Hash"), *SingletonName);
 
 	FUHTStringBuilder GeneratedEnumRegisterFunctionText;
@@ -4045,7 +4128,7 @@ void FNativeClassHeaderGenerator::ExportGeneratedEnumInitCode(FOutputDevice& Out
 	GeneratedEnumRegisterFunctionText.Logf(TEXT("\t}\r\n"));
 
 	uint32 EnumHash = GenerateTextHash(*GeneratedEnumRegisterFunctionText);
-	GGeneratedCodeHashes.Add(Enum, EnumHash);
+	AddGeneratedCodeHash(Enum, EnumHash);
 	Out.Logf(TEXT("\tuint32 %s() { return %uU; }\r\n"), *HashFuncName, EnumHash);
 	Out.Log(GeneratedEnumRegisterFunctionText);
 }
@@ -4136,7 +4219,7 @@ void WriteEventFunctionEpilogue(FOutputDevice& Output, int32 Indent, const FParm
 	Output.Logf(TEXT("%s}\r\n"), FCString::Tab(Indent));
 }
 
-void FNativeClassHeaderGenerator::ExportDelegateDeclaration(FOutputDevice& Out, const FUnrealSourceFile& SourceFile, UFunction* Function)
+void FNativeClassHeaderGenerator::ExportDelegateDeclaration(FOutputDevice& Out, FReferenceGatherers& OutReferenceGatherers, const FUnrealSourceFile& SourceFile, UFunction* Function) const
 {
 	static const TCHAR DelegateStr[] = TEXT("delegate");
 
@@ -4167,23 +4250,23 @@ void FNativeClassHeaderGenerator::ExportDelegateDeclaration(FOutputDevice& Out,
 	DelegateOutput.Log(TEXT("static "));
 
 	// export the line that looks like: int32 Main(const FString& Parms)
-	ExportNativeFunctionHeader(DelegateOutput, ForwardDeclarations, FunctionData, EExportFunctionType::Event, EExportFunctionHeaderStyle::Declaration, *ExtraParam, *GetAPIString());
+	ExportNativeFunctionHeader(DelegateOutput, OutReferenceGatherers.ForwardDeclarations, FunctionData, EExportFunctionType::Event, EExportFunctionHeaderStyle::Declaration, *ExtraParam, *GetAPIString());
 
 	// Only exporting function prototype
 	DelegateOutput.Logf(TEXT(";\r\n"));
 
-	ExportFunction(Out, SourceFile, Function, false);
+	ExportFunction(Out, OutReferenceGatherers, SourceFile, Function, false);
 }
 
-void FNativeClassHeaderGenerator::ExportDelegateDefinition(FOutputDevice& Out, const FUnrealSourceFile& SourceFile, UFunction* Function)
+void FNativeClassHeaderGenerator::ExportDelegateDefinition(FOutputDevice& Out, FReferenceGatherers& OutReferenceGatherers, const FUnrealSourceFile& SourceFile, UFunction* Function) const
 {
-	static const TCHAR DelegateStr[] = TEXT("delegate");
+	const TCHAR DelegateStr[] = TEXT("delegate");
 
 	check(Function->HasAnyFunctionFlags(FUNC_Delegate));
 
 	// Export parameters structs for all delegates.  We'll need these to declare our delegate execution function.
 	FUHTStringBuilder DelegateOutput;
-	ExportEventParm(DelegateOutput, ForwardDeclarations, Function, /*Indent=*/ 0, /*bOutputConstructor=*/ true, EExportingState::Normal);
+	ExportEventParm(DelegateOutput, OutReferenceGatherers.ForwardDeclarations, Function, /*Indent=*/ 0, /*bOutputConstructor=*/ true, EExportingState::Normal);
 
 	const bool bIsMulticastDelegate = Function->HasAnyFunctionFlags( FUNC_MulticastDelegate );
 
@@ -4212,7 +4295,7 @@ void FNativeClassHeaderGenerator::ExportDelegateDefinition(FOutputDevice& Out, c
 	DelegateOutput.Log(TEXT("static "));
 
 	// export the line that looks like: int32 Main(const FString& Parms)
-	ExportNativeFunctionHeader(DelegateOutput, ForwardDeclarations, FunctionData, EExportFunctionType::Event, EExportFunctionHeaderStyle::Declaration, *ExtraParam, *GetAPIString());
+	ExportNativeFunctionHeader(DelegateOutput, OutReferenceGatherers.ForwardDeclarations, FunctionData, EExportFunctionType::Event, EExportFunctionHeaderStyle::Declaration, *ExtraParam, *GetAPIString());
 
 	FParmsAndReturnProperties Parameters = GetFunctionParmsAndReturn(FunctionData.FunctionReference);
 
@@ -4416,14 +4499,14 @@ FString FNativeClassHeaderGenerator::GetNullParameterValue( FProperty* Prop, boo
 }
 
 
-FString FNativeClassHeaderGenerator::GetFunctionReturnString(UFunction* Function)
+FString FNativeClassHeaderGenerator::GetFunctionReturnString(UFunction* Function, FReferenceGatherers& OutReferenceGatherers)
 {
 	FString Result;
 
 	if (FProperty* Return = Function->GetReturnProperty())
 	{
 		FString ExtendedReturnType;
-		ForwardDeclarations.Add(Return->GetCPPTypeForwardDeclaration());
+		OutReferenceGatherers.ForwardDeclarations.Add(Return->GetCPPTypeForwardDeclaration());
 		FString ReturnType = Return->GetCPPType(&ExtendedReturnType, CPPF_ArgumentOrReturnValue);
 		FUHTStringBuilder ReplacementText;
 		ReplacementText += MoveTemp(ReturnType);
@@ -4483,13 +4566,13 @@ FString CreateClickableErrorMessage(const FString& Filename, int32 Line, int32 C
 	return FString::Printf(TEXT("%s(%d,%d): error: "), *Filename, Line, Column);
 }
 
-void FNativeClassHeaderGenerator::CheckRPCFunctions(const FFuncInfo& FunctionData, const FString& ClassName, int32 ImplementationPosition, int32 ValidatePosition, const FUnrealSourceFile& SourceFile)
+void FNativeClassHeaderGenerator::CheckRPCFunctions(FReferenceGatherers& OutReferenceGatherers, const FFuncInfo& FunctionData, const FString& ClassName, int32 ImplementationPosition, int32 ValidatePosition, const FUnrealSourceFile& SourceFile) const
 {
 	bool bHasImplementation = ImplementationPosition != INDEX_NONE;
 	bool bHasValidate = ValidatePosition != INDEX_NONE;
 
 	UFunction* Function = FunctionData.FunctionReference;
-	FString FunctionReturnType = GetFunctionReturnString(Function);
+	FString FunctionReturnType = GetFunctionReturnString(Function, OutReferenceGatherers);
 	const TCHAR* ConstModifier = (Function->HasAllFunctionFlags(FUNC_Const) ? TEXT("const ") : TEXT(" "));
 
 	const bool bIsNative = Function->HasAllFunctionFlags(FUNC_Native);
@@ -4503,7 +4586,7 @@ void FNativeClassHeaderGenerator::CheckRPCFunctions(const FFuncInfo& FunctionDat
 
 	check(bNeedsImplementation || bNeedsValidate);
 
-	FString ParameterString = GetFunctionParameterString(Function);
+	FString ParameterString = GetFunctionParameterString(Function, OutReferenceGatherers);
 	const FString& Filename = SourceFile.GetFilename();
 	const FString& FileContent = SourceFile.GetContent();
 
@@ -4750,14 +4833,14 @@ void FNativeClassHeaderGenerator::ExportNativeFunctionHeader(
  * @param Return return parameter for the function
  * @param DeprecationWarningOutputDevice Device to output deprecation warnings for _Validate and _Implementation functions.
  */
-void FNativeClassHeaderGenerator::ExportFunctionThunk(FUHTStringBuilder& RPCWrappers, UFunction* Function, const FFuncInfo& FunctionData, const TArray<FProperty*>& Parameters, FProperty* Return)
+void FNativeClassHeaderGenerator::ExportFunctionThunk(FUHTStringBuilder& RPCWrappers, FReferenceGatherers& OutReferenceGatherers, UFunction* Function, const FFuncInfo& FunctionData, const TArray<FProperty*>& Parameters, FProperty* Return) const
 {
 	// export the GET macro for this parameter
 	FString ParameterList;
 	for (int32 ParameterIndex = 0; ParameterIndex < Parameters.Num(); ParameterIndex++)
 	{
 		FProperty* Param = Parameters[ParameterIndex];
-		ForwardDeclarations.Add(Param->GetCPPTypeForwardDeclaration());
+		OutReferenceGatherers.ForwardDeclarations.Add(Param->GetCPPTypeForwardDeclaration());
 
 		FString EvalBaseText = TEXT("P_GET_");	// e.g. P_GET_STR
 		FString EvalModifierText;				// e.g. _REF
@@ -4924,7 +5007,7 @@ void FNativeClassHeaderGenerator::ExportFunctionThunk(FUHTStringBuilder& RPCWrap
 
 	//Emit warning here if necessary
 	FUHTStringBuilder FunctionDeclaration;
-	ExportNativeFunctionHeader(FunctionDeclaration, ForwardDeclarations, FunctionData, EExportFunctionType::Function, EExportFunctionHeaderStyle::Declaration, nullptr, *GetAPIString());
+	ExportNativeFunctionHeader(FunctionDeclaration, OutReferenceGatherers.ForwardDeclarations, FunctionData, EExportFunctionType::Function, EExportFunctionHeaderStyle::Declaration, nullptr, *GetAPIString());
 
 	// Call the validate function if there is one
 	if (!(FunctionData.FunctionExportFlags & FUNCEXPORT_CppStatic) && (FunctionData.FunctionFlags & FUNC_NetValidate))
@@ -4940,7 +5023,7 @@ void FNativeClassHeaderGenerator::ExportFunctionThunk(FUHTStringBuilder& RPCWrap
 	RPCWrappers.Log(TEXT("\t\t"));
 	if (Return)
 	{
-		ForwardDeclarations.Add(Return->GetCPPTypeForwardDeclaration());
+		OutReferenceGatherers.ForwardDeclarations.Add(Return->GetCPPTypeForwardDeclaration());
 
 		FUHTStringBuilder ReplacementText;
 		FString ReturnExtendedType;
@@ -4963,14 +5046,14 @@ void FNativeClassHeaderGenerator::ExportFunctionThunk(FUHTStringBuilder& RPCWrap
 	RPCWrappers += TEXT("\t\tP_NATIVE_END;") LINE_TERMINATOR;
 }
 
-FString FNativeClassHeaderGenerator::GetFunctionParameterString(UFunction* Function)
+FString FNativeClassHeaderGenerator::GetFunctionParameterString(UFunction* Function, FReferenceGatherers& OutReferenceGatherers)
 {
 	FString ParameterList;
 	FUHTStringBuilder PropertyText;
 
 	for (FProperty* Property : TFieldRange<FProperty>(Function))
 	{
-		ForwardDeclarations.Add(Property->GetCPPTypeForwardDeclaration());
+		OutReferenceGatherers.ForwardDeclarations.Add(Property->GetCPPTypeForwardDeclaration());
 
 		if ((Property->PropertyFlags & (CPF_Parm | CPF_ReturnParm)) != CPF_Parm)
 		{
@@ -5003,7 +5086,7 @@ struct FNativeFunctionStringBuilder
 	FUHTStringBuilder AutogeneratedStaticDataFuncs;
 };
 
-void FNativeClassHeaderGenerator::ExportNativeFunctions(FOutputDevice& OutGeneratedHeaderText, FOutputDevice& OutGeneratedCPPText, FOutputDevice& OutMacroCalls, FOutputDevice& OutNoPureDeclsMacroCalls, const FUnrealSourceFile& SourceFile, UClass* Class, FClassMetaData* ClassData)
+void FNativeClassHeaderGenerator::ExportNativeFunctions(FOutputDevice& OutGeneratedHeaderText, FOutputDevice& OutGeneratedCPPText, FOutputDevice& OutMacroCalls, FOutputDevice& OutNoPureDeclsMacroCalls, FReferenceGatherers& OutReferenceGatherers, const FUnrealSourceFile& SourceFile, UClass* Class, FClassMetaData* ClassData) const
 {
 	FNativeFunctionStringBuilder RuntimeStringBuilders;
 	FNativeFunctionStringBuilder EditorStringBuilders;
@@ -5136,13 +5219,13 @@ void FNativeClassHeaderGenerator::ExportNativeFunctions(FOutputDevice& OutGenera
 
 			//Emit warning here if necessary
 			FUHTStringBuilder FunctionDeclaration;
-			ExportNativeFunctionHeader(FunctionDeclaration, ForwardDeclarations, FunctionData, EExportFunctionType::Function, EExportFunctionHeaderStyle::Declaration, nullptr, *GetAPIString());
+			ExportNativeFunctionHeader(FunctionDeclaration, OutReferenceGatherers.ForwardDeclarations, FunctionData, EExportFunctionType::Function, EExportFunctionHeaderStyle::Declaration, nullptr, *GetAPIString());
 			FunctionDeclaration.Log(TEXT(";\r\n"));
 
 			// Declare validation function if needed
 			if (FunctionData.FunctionFlags & FUNC_NetValidate)
 			{
-				FString ParameterList = GetFunctionParameterString(Function);
+				FString ParameterList = GetFunctionParameterString(Function, OutReferenceGatherers);
 
 				const TCHAR* Virtual = (!FunctionData.FunctionReference->HasAnyFunctionFlags(FUNC_Static) && !(FunctionData.FunctionExportFlags & FUNCEXPORT_Final)) ? TEXT("virtual") : TEXT("");
 				FStringOutputDevice ValidDecl;
@@ -5163,7 +5246,7 @@ void FNativeClassHeaderGenerator::ExportNativeFunctions(FOutputDevice& OutGenera
 			// Versions that skip function autodeclaration throw an error when a function is missing.
 			if (ClassRange.bHasGeneratedBody && (SourceFile.GetGeneratedCodeVersionForStruct(Class) > EGeneratedCodeVersion::V1))
 			{
-				CheckRPCFunctions(FunctionData, ClassCPPName, ImplementationPosition, ValidatePosition, SourceFile);
+				CheckRPCFunctions(OutReferenceGatherers, FunctionData, ClassCPPName, ImplementationPosition, ValidatePosition, SourceFile);
 			}
 		}
 
@@ -5182,7 +5265,7 @@ void FNativeClassHeaderGenerator::ExportNativeFunctions(FOutputDevice& OutGenera
 		FuncStringBuilders.RPCImplementations += LINE_TERMINATOR TEXT("\t{") LINE_TERMINATOR;
 
 		FParmsAndReturnProperties Parameters = GetFunctionParmsAndReturn(FunctionData.FunctionReference);
-		ExportFunctionThunk(FuncStringBuilders.RPCImplementations, Function, FunctionData, Parameters.Parms, Parameters.Return);
+		ExportFunctionThunk(FuncStringBuilders.RPCImplementations, OutReferenceGatherers, Function, FunctionData, Parameters.Parms, Parameters.Return);
 
 		FuncStringBuilders.RPCImplementations += TEXT("\t}") LINE_TERMINATOR;
 	}
@@ -5400,7 +5483,7 @@ void FNativeClassHeaderGenerator::ApplyAlternatePropertyExportText(FProperty* Pr
 		if (DelegateProperty || MulticastDelegateProperty)
 		{
 			FString Original = Prop->GetCPPType();
-			FString PlaceholderOfSameSizeAndAlignemnt;
+			const TCHAR* PlaceholderOfSameSizeAndAlignemnt;
 			if (DelegateProperty)
 			{
 				PlaceholderOfSameSizeAndAlignemnt = TEXT("FScriptDelegate");
@@ -5409,7 +5492,7 @@ void FNativeClassHeaderGenerator::ApplyAlternatePropertyExportText(FProperty* Pr
 			{
 				PlaceholderOfSameSizeAndAlignemnt = TEXT("FMulticastScriptDelegate");
 			}
-			PropertyText.ReplaceInline(*Original, *PlaceholderOfSameSizeAndAlignemnt, ESearchCase::CaseSensitive);
+			PropertyText.ReplaceInline(*Original, PlaceholderOfSameSizeAndAlignemnt, ESearchCase::CaseSensitive);
 		}
 	}
 }
@@ -5477,6 +5560,7 @@ static void RecordPackageSingletons(
 	const TArray<UDelegateFunction*>& Delegates)
 {
 	TArray<UField*> Singletons;
+	Singletons.Reserve(Structs.Num() + Delegates.Num());
 	for (UScriptStruct* Struct : Structs)
 	{
 		if (Struct->StructFlags & STRUCT_NoExport && !HasDynamicOuter(Struct))
@@ -5495,14 +5579,10 @@ static void RecordPackageSingletons(
 
 	if (Singletons.Num())
 	{
+		FScopeLock PackageSingletonLock(&GPackageSingletonsCriticalSection);
+
 		TArray<UField*>& PackageSingletons = GPackageSingletons.FindOrAdd(&Package);
-		PackageSingletons.Append(Singletons);
-		Algo::Sort(PackageSingletons, [](UField* A, UField* B)
-		{
-			// Structs before delegates then UniqueId order
-			return	(uint64(A->IsA<UDelegateFunction>()) << 32) + A->GetUniqueID() <
-					(uint64(B->IsA<UDelegateFunction>()) << 32) + B->GetUniqueID();
-		});
+		PackageSingletons.Append(MoveTemp(Singletons));
 	}
 }
 
@@ -5516,7 +5596,6 @@ FNativeClassHeaderGenerator::FNativeClassHeaderGenerator(
 	: API                        (FPackageName::GetShortName(InPackage).ToUpper())
 	, APIStringPrivate           (FString::Printf(TEXT("%s_API "), *API))
 	, Package                    (InPackage)
-	, UniqueCrossModuleReferences(nullptr)
 	, bAllowSaveExportedHeaders  (InAllowSaveExportedHeaders)
 	, bFailIfGeneratedCodeChanges(FParse::Param(FCommandLine::Get(), TEXT("FailIfGeneratedCodeChanges")))
 {
@@ -5556,10 +5635,6 @@ FNativeClassHeaderGenerator::FNativeClassHeaderGenerator(
 		}
 	}
 
-	// Export an include line for each header
-	TSet<FUnrealSourceFile*> PublicHeaderGroupIncludes;
-	FUHTStringBuilder GeneratedFunctionDeclarations;
-
 	TArray<FUnrealSourceFile*> Exported;
 	{
 		// Get source files and ignore them next time round
@@ -5579,41 +5654,56 @@ FNativeClassHeaderGenerator::FNativeClassHeaderGenerator(
 		TArray<FString>              RelativeIncludes;
 		FUHTStringBuilderLineCounter GeneratedText;
 		TSet<FString>                CrossModuleReferences;
+		TSet<FString>                PackageHeaderPaths;
+		TArray<FString>              TempHeaderPaths;
+		FUHTStringBuilder			 GeneratedFunctionDeclarations;
 	};
 
-	TArray<UEnum*> Enums;
-	TArray<UScriptStruct*> Structs;
-	TArray<UDelegateFunction*> DelegateFunctions;
-	TArray<UClass*>	DefinedClasses;
-
 	TMap<FUnrealSourceFile*, FGeneratedCPP> GeneratedCPPs;
 	GeneratedCPPs.Reserve(Exported.Num());
 
+	// Set up the generated cpp map
 	for (FUnrealSourceFile* SourceFile : Exported)
 	{
 		FString ModuleRelativeFilename = SourceFile->GetFilename();
 		ConvertToBuildIncludePath(Package, ModuleRelativeFilename);
 
+		FString StrippedName = FPaths::GetBaseFilename(ModuleRelativeFilename);
+		FString GeneratedSourceFilename = (PackageManifest->GeneratedIncludeDirectory / StrippedName) + TEXT(".gen.cpp");
+
+		FGeneratedCPP& GeneratedCPP = GeneratedCPPs.Emplace(SourceFile, MoveTemp(GeneratedSourceFilename));
+		GeneratedCPP.RelativeIncludes.Add(MoveTemp(ModuleRelativeFilename));
+	}
+
+	const FManifestModule* ConstPackageManifest = GetPackageManifest(PackageName);
+	const FNativeClassHeaderGenerator* ConstThis = this;
+
+	ParallelFor(Exported.Num(), [&Exported, Package=Package, ConstPackageManifest, &GeneratedCPPs, ConstThis](int32 Index)
+	{
+		FUnrealSourceFile* SourceFile = Exported[Index];
+
+		/** Forward declarations that we need for this sourcefile. */
+		TSet<FString> ForwardDeclarations;
+
+		FString ModuleRelativeFilename = SourceFile->GetFilename();
+		ConvertToBuildIncludePath(Package, ModuleRelativeFilename);
+
 		FString StrippedName       = FPaths::GetBaseFilename(ModuleRelativeFilename);
-		FString BaseSourceFilename = PackageManifest->GeneratedIncludeDirectory / StrippedName;
+		FString BaseSourceFilename = ConstPackageManifest->GeneratedIncludeDirectory / StrippedName;
 
 		FUHTStringBuilder GeneratedHeaderText;
-		FGeneratedCPP& GeneratedCPP = GeneratedCPPs.Emplace(SourceFile, BaseSourceFilename + TEXT(".gen.cpp"));
-		GeneratedCPP.RelativeIncludes.Add(MoveTemp(ModuleRelativeFilename));
+		FGeneratedCPP& GeneratedCPP = GeneratedCPPs.FindChecked(SourceFile);
+		FUHTStringBuilder& GeneratedFunctionDeclarations = GeneratedCPP.GeneratedFunctionDeclarations;
 
-		UniqueCrossModuleReferences = &GeneratedCPP.CrossModuleReferences;
-		ON_SCOPE_EXIT
-		{
-			UniqueCrossModuleReferences = nullptr;
-		};
+		FReferenceGatherers ReferenceGatherers(&GeneratedCPP.CrossModuleReferences, GeneratedCPP.PackageHeaderPaths, GeneratedCPP.TempHeaderPaths);
 
 		FUHTStringBuilderLineCounter& OutText = GeneratedCPP.GeneratedText;
 
 		NameLookupCPP.SetCurrentSourceFile(SourceFile);
 
-		Enums.Reset();
-		Structs.Reset();
-		DelegateFunctions.Reset();
+		TArray<UEnum*> Enums;
+		TArray<UScriptStruct*> Structs;
+		TArray<UDelegateFunction*> DelegateFunctions;
 		SourceFile->GetScope()->SplitTypesIntoArrays(Enums, Structs, DelegateFunctions);
 
 		RecordPackageSingletons(*SourceFile->GetPackage(), Structs, DelegateFunctions);
@@ -5635,7 +5725,7 @@ FNativeClassHeaderGenerator::FNativeClassHeaderGenerator(
 		for (UDelegateFunction* Func : DelegateFunctions)
 		{
 			GeneratedFunctionDeclarations.Log(FTypeSingletonCache::Get(Func).GetExternDecl());
-			ExportDelegateDeclaration(OutText, *SourceFile, Func);
+			ConstThis->ExportDelegateDeclaration(OutText, ReferenceGatherers, *SourceFile, Func);
 		}
 
 		// Export enums declared in non-UClass headers.
@@ -5645,7 +5735,7 @@ FNativeClassHeaderGenerator::FNativeClassHeaderGenerator(
 			if (Enum->GetOuter()->IsA(UPackage::StaticClass()))
 			{
 				GeneratedFunctionDeclarations.Log(FTypeSingletonCache::Get(Enum).GetExternDecl());
-				ExportGeneratedEnumInitCode(OutText, *SourceFile, Enum);
+				ConstThis->ExportGeneratedEnumInitCode(OutText, ReferenceGatherers, *SourceFile, Enum);
 			}
 		}
 
@@ -5654,25 +5744,23 @@ FNativeClassHeaderGenerator::FNativeClassHeaderGenerator(
 		for (UScriptStruct* Struct : Structs)
 		{
 			GeneratedFunctionDeclarations.Log(FTypeSingletonCache::Get(Struct).GetExternDecl());
-			ExportGeneratedStructBodyMacros(GeneratedHeaderText, OutText, *SourceFile, Struct);
+			ConstThis->ExportGeneratedStructBodyMacros(GeneratedHeaderText, OutText, ReferenceGatherers, *SourceFile, Struct);
 		}
 
 		// export delegate wrapper function implementations
 		for (UDelegateFunction* Func : DelegateFunctions)
 		{
-			ExportDelegateDefinition(GeneratedHeaderText, *SourceFile, Func);
+			ConstThis->ExportDelegateDefinition(GeneratedHeaderText, ReferenceGatherers, *SourceFile, Func);
 		}
 
-		DefinedClasses.Reset();
+		TArray<UClass*>	DefinedClasses;
 		SourceFile->AppendDefinedClasses(DefinedClasses);
 		for (UClass* Class : DefinedClasses)
 		{
 			if (!(Class->ClassFlags & CLASS_Intrinsic))
 			{
-				ExportClassFromSourceFileInner(GeneratedHeaderText, OutText, GeneratedFunctionDeclarations, (FClass*)Class, *SourceFile);
+				ConstThis->ExportClassFromSourceFileInner(GeneratedHeaderText, OutText, GeneratedFunctionDeclarations, ReferenceGatherers, (FClass*)Class, *SourceFile);
 			}
-
-			GClassToSourceFileMap.Add(Class, SourceFile);
 		}
 
 		GeneratedHeaderText.Log(TEXT("#undef CURRENT_FILE_ID\r\n"));
@@ -5680,21 +5768,37 @@ FNativeClassHeaderGenerator::FNativeClassHeaderGenerator(
 
 		for (UEnum* Enum : Enums)
 		{
-			ExportEnum(GeneratedHeaderText, Enum);
+			ConstThis->ExportEnum(GeneratedHeaderText, Enum);
 		}
 
 		FString HeaderPath = BaseSourceFilename + TEXT(".generated.h");
-		bool bHasChanged = WriteHeader(*HeaderPath, GeneratedHeaderText, ForwardDeclarations);
+		bool bHasChanged = ConstThis->WriteHeader(*HeaderPath, GeneratedHeaderText, ReferenceGatherers);
 
 		SourceFile->SetGeneratedFilename(MoveTemp(HeaderPath));
 		SourceFile->SetHasChanged(bHasChanged);
+	});
 
-		ForwardDeclarations.Reset();
+	// Export an include line for each header
+	TSet<FUnrealSourceFile*> PublicHeaderGroupIncludes;
+	TArray<UClass*>	DefinedClasses;
+	FUHTStringBuilder GeneratedFunctionDeclarations;
+
+	for (FUnrealSourceFile* SourceFile : Exported)
+	{
+		DefinedClasses.Reset();
+		SourceFile->AppendDefinedClasses(DefinedClasses);
+		for (UClass* Class : DefinedClasses)
+		{
+			GClassToSourceFileMap.Add(Class, SourceFile);
+		}
 
 		if (GPublicSourceFileSet.Contains(SourceFile))
 		{
 			PublicHeaderGroupIncludes.Add(SourceFile);
 		}
+
+		const FGeneratedCPP& GeneratedCPP = GeneratedCPPs.FindChecked(SourceFile);
+		GeneratedFunctionDeclarations.Log(GeneratedCPP.GeneratedFunctionDeclarations);
 	}
 
 	// Add includes for 'Within' classes
@@ -5734,6 +5838,8 @@ FNativeClassHeaderGenerator::FNativeClassHeaderGenerator(
 		}
 	}
 
+	TSet<FString> PackageHeaderPaths;
+	TArray<FString> TempHeaderPaths;
 	if (bWriteClassesH)
 	{
 		// Write the classes and enums header prefixes.
@@ -5760,9 +5866,11 @@ FNativeClassHeaderGenerator::FNativeClassHeaderGenerator(
 
 		ClassesHText.Log(LINE_TERMINATOR);
 
+		FReferenceGatherers ReferenceGatherers(nullptr, PackageHeaderPaths, TempHeaderPaths);
+
 		// Save the classes header if it has changed.
 		const FString ClassesHeaderPath = PackageManifest->GeneratedIncludeDirectory / (PackageName + TEXT("Classes.h"));
-		SaveHeaderIfChanged(*ClassesHeaderPath, *ClassesHText);
+		SaveHeaderIfChanged(ReferenceGatherers, *ClassesHeaderPath, *ClassesHText);
 	}
 
 	// now export the names for the functions in this package
@@ -5793,38 +5901,52 @@ FNativeClassHeaderGenerator::FNativeClassHeaderGenerator(
 
 	const FManifestModule* ModuleInfo = GPackageToManifestModuleMap.FindChecked(Package);
 
+	TArray<FGeneratedCPP*> GeneratedCPPArray;
+	GeneratedCPPArray.Reserve(GeneratedCPPs.Num());
+	for (TPair<FUnrealSourceFile*, FGeneratedCPP>& Pair : GeneratedCPPs)
+	{
+		GeneratedCPPArray.Add(&Pair.Value);
+	}
+
 	// Generate CPP files
-	TArray<FString> GeneratedCPPNames;
-	for (const TPair<FUnrealSourceFile*, FGeneratedCPP>& GeneratedCPP : GeneratedCPPs)
+	ParallelFor(GeneratedCPPArray.Num(), [ConstThis, &GeneratedCPPArray](int32 Index)
 	{
+		FGeneratedCPP* GeneratedCPP = GeneratedCPPArray[Index];
+		FReferenceGatherers ReferenceGatherers(nullptr, GeneratedCPP->PackageHeaderPaths, GeneratedCPP->TempHeaderPaths);
+
 		FUHTStringBuilder FileText;
 
 		FString GeneratedIncludes;
-		for (const FString& RelativeInclude : GeneratedCPP.Value.RelativeIncludes)
+		for (const FString& RelativeInclude : GeneratedCPP->RelativeIncludes)
 		{
 			GeneratedIncludes += FString::Printf(TEXT("#include \"%s\"\r\n"), *RelativeInclude);
 		}
 
-		FString CleanFilename = FPaths::GetCleanFilename(GeneratedCPP.Value.GeneratedCppFullFilename);
-
-		GeneratedCPPNames.Add(CleanFilename);
+		FString CleanFilename = FPaths::GetCleanFilename(GeneratedCPP->GeneratedCppFullFilename);
 
 		CleanFilename.ReplaceInline(TEXT(".gen.cpp"), TEXT(""), ESearchCase::CaseSensitive);
 		CleanFilename.ReplaceInline(TEXT("."), TEXT("_"), ESearchCase::CaseSensitive);
 
 		ExportGeneratedCPP(
 			FileText,
-			GeneratedCPP.Value.CrossModuleReferences,
+			GeneratedCPP->CrossModuleReferences,
 			*CleanFilename,
-			*GeneratedCPP.Value.GeneratedText,
+			*GeneratedCPP->GeneratedText,
 			*GeneratedIncludes
 		);
 
-		SaveHeaderIfChanged(*GeneratedCPP.Value.GeneratedCppFullFilename, *FileText);
-	}
+		ConstThis->SaveHeaderIfChanged(ReferenceGatherers, *GeneratedCPP->GeneratedCppFullFilename, *FileText);
+	});
 
 	if (bAllowSaveExportedHeaders)
 	{
+		TArray<FString> GeneratedCPPNames;
+		GeneratedCPPNames.Reserve(GeneratedCPPs.Num());
+		for (const TPair<FUnrealSourceFile*, FGeneratedCPP>& GeneratedCPP : GeneratedCPPs)
+		{
+			GeneratedCPPNames.Add(FPaths::GetCleanFilename(GeneratedCPP.Value.GeneratedCppFullFilename));
+		}
+
 		// Delete old generated .cpp files which we don't need because we generated less code than last time.
 		TArray<FString> FoundFiles;
 		FString BaseDir = FPaths::GetPath(ModuleInfo->GeneratedCPPFilenameBase);
@@ -5841,19 +5963,24 @@ FNativeClassHeaderGenerator::FNativeClassHeaderGenerator(
 		}
 	}
 
+	for (TPair<FUnrealSourceFile*, FGeneratedCPP>& GeneratedCPP : GeneratedCPPs)
+	{
+		TempHeaderPaths.Append(MoveTemp(GeneratedCPP.Value.TempHeaderPaths));
+		PackageHeaderPaths.Append(MoveTemp(GeneratedCPP.Value.PackageHeaderPaths));
+	}
+
 	// Export all changed headers from their temp files to the .h files
-	ExportUpdatedHeaders(PackageName);
+	ExportUpdatedHeaders(PackageName, TempHeaderPaths);
 
 	// Delete stale *.generated.h files
-	DeleteUnusedGeneratedHeaders();
+	DeleteUnusedGeneratedHeaders(PackageHeaderPaths);
 }
 
-void FNativeClassHeaderGenerator::DeleteUnusedGeneratedHeaders()
+void FNativeClassHeaderGenerator::DeleteUnusedGeneratedHeaders(const TSet<FString>& PackageHeaderPathSet)
 {
 	TSet<FString> AllIntermediateFolders;
-	TSet<FString> PackageHeaderPathSet(PackageHeaderPaths);
 
-	for (const FString& PackageHeader : PackageHeaderPaths)
+	for (const FString& PackageHeader : PackageHeaderPathSet)
 	{
 		const FString IntermediatePath = FPaths::GetPath(PackageHeader);
 
@@ -5895,7 +6022,7 @@ void FNativeClassHeaderGenerator::DeleteUnusedGeneratedHeaders()
  */
 ECompilationResult::Type GCompilationResult = ECompilationResult::OtherCompilationError;
 
-bool FNativeClassHeaderGenerator::SaveHeaderIfChanged(const TCHAR* HeaderPath, const TCHAR* InNewHeaderContents)
+bool FNativeClassHeaderGenerator::SaveHeaderIfChanged(FReferenceGatherers& OutReferenceGatherers, const TCHAR* HeaderPath, const TCHAR* InNewHeaderContents) const
 {
 	if ( !bAllowSaveExportedHeaders )
 	{
@@ -6021,13 +6148,13 @@ bool FNativeClassHeaderGenerator::SaveHeaderIfChanged(const TCHAR* HeaderPath, c
 			UE_LOG_WARNING_UHT(TEXT("Failed to save header export preview: '%s'"), *TmpHeaderFilename);
 		}
 
-		TempHeaderPaths.Add(MoveTemp(TmpHeaderFilename));
+		OutReferenceGatherers.TempHeaderPaths.Add(MoveTemp(TmpHeaderFilename));
 	}
 
 	// Remember this header filename to be able to check for any old (unused) headers later.
 	HeaderPathStr.ReplaceInline(TEXT("\\"), TEXT("/"), ESearchCase::CaseSensitive);
 
-	PackageHeaderPaths.Add(MoveTemp(HeaderPathStr));
+	OutReferenceGatherers.PackageHeaderPaths.Add(MoveTemp(HeaderPathStr));
 
 	return bHasChanged;
 }
@@ -6052,10 +6179,11 @@ FString FNativeClassHeaderGenerator::GenerateTempHeaderName( const FString& Curr
 *
 * @param	PackageName	Name of the package being saved
 */
-void FNativeClassHeaderGenerator::ExportUpdatedHeaders(const FString& PackageName)
+void FNativeClassHeaderGenerator::ExportUpdatedHeaders(const FString& PackageName, const TArray<FString>& TempHeaderPaths)
 {
-	for (const FString& TmpFilename : TempHeaderPaths)
+	ParallelFor(TempHeaderPaths.Num(), [&](int32 Index)
 	{
+		const FString& TmpFilename = TempHeaderPaths[Index];
 		FString Filename = GenerateTempHeaderName( TmpFilename, true );
 		if (!IFileManager::Get().Move(*Filename, *TmpFilename, true, true))
 		{
@@ -6065,7 +6193,7 @@ void FNativeClassHeaderGenerator::ExportUpdatedHeaders(const FString& PackageNam
 		{
 			UE_LOG(LogCompile, Log, TEXT("Exported updated C++ header: %s"), *Filename);
 		}
-	}
+	});
 }
 
 /**
diff --git a/Engine/Source/Programs/UnrealHeaderTool/Private/HeaderParser.cpp b/Engine/Source/Programs/UnrealHeaderTool/Private/HeaderParser.cpp
--- a/Engine/Source/Programs/UnrealHeaderTool/Private/HeaderParser.cpp
+++ b/Engine/Source/Programs/UnrealHeaderTool/Private/HeaderParser.cpp
@@ -8409,8 +8409,7 @@ TSet<FUnrealSourceFile*> GetSourceFilesWithInheritanceOrdering(UPackage* Current
 
 		FUnrealSourceFile& SourceFile = (*DefinitionInfoPtr)->GetUnrealSourceFile();
 
-		if (!SourceFiles.Contains(&SourceFile)
-			&& SourceFile.GetScope()->ContainsTypes())
+		if (SourceFile.GetScope()->ContainsTypes())
 		{
 			SourceFiles.Add(&SourceFile);
 		}
@@ -8422,7 +8421,6 @@ TSet<FUnrealSourceFile*> GetSourceFilesWithInheritanceOrdering(UPackage* Current
 		auto& SourceFile = Pair.Value.Get();
 
 		if (SourceFile.GetPackage() == CurrentPackage
-			&& !SourceFiles.Contains(&SourceFile)
 			&& SourceFile.GetScope()->ContainsTypes())
 		{
 			SourceFiles.Add(&SourceFile);
@@ -9129,7 +9127,7 @@ void FHeaderParser::SimplifiedClassParse(const TCHAR* Filename, const TCHAR* InB
 		else if ( bProcess && FParse::Command(&Str,TEXT("#include")) )
 		{
 			// Handle #include directives as if they were 'dependson' keywords.
-			FString DependsOnHeaderName = Str;
+			const FString& DependsOnHeaderName = Str;
 
 			if (DependsOnHeaderName != TEXT("\"UObject/DefineUPropertyMacros.h\"") && DependsOnHeaderName != TEXT("\"UObject/UndefineUPropertyMacros.h\""))
 			{
@@ -9147,7 +9145,7 @@ void FHeaderParser::SimplifiedClassParse(const TCHAR* Filename, const TCHAR* InB
 					if (HeaderFilenameEnd != INDEX_NONE)
 					{
 						// Include the extension in the name so that we later know where this entry came from.
-						DependentOn.Add(FHeaderProvider(EHeaderProviderSourceType::FileName, *FPaths::GetCleanFilename(DependsOnHeaderName.Mid(1, HeaderFilenameEnd - 1))));
+						DependentOn.Add(FHeaderProvider(EHeaderProviderSourceType::FileName, FPaths::GetCleanFilename(DependsOnHeaderName.Mid(1, HeaderFilenameEnd - 1))));
 					}
 				}
 			}
diff --git a/Engine/Source/Programs/UnrealHeaderTool/Private/NativeClassExporter.h b/Engine/Source/Programs/UnrealHeaderTool/Private/NativeClassExporter.h
--- a/Engine/Source/Programs/UnrealHeaderTool/Private/NativeClassExporter.h
+++ b/Engine/Source/Programs/UnrealHeaderTool/Private/NativeClassExporter.h
@@ -93,17 +93,29 @@ private:
 
 	const UPackage* Package;
 
-	/** Set of already exported cross-module references, to prevent duplicates */
-	TSet<FString>* UniqueCrossModuleReferences;
+	/** A collection of structures used to gather various kinds of references conveniently grouped together to make passing easier */
+	struct FReferenceGatherers
+	{
+		FReferenceGatherers(TSet<FString>* InUniqueCrossModuleReferences, TSet<FString>& InPackageHeaderPaths, TArray<FString>& InTempHeaderPaths)
+			: UniqueCrossModuleReferences(InUniqueCrossModuleReferences)
+			, PackageHeaderPaths(InPackageHeaderPaths)
+			, TempHeaderPaths(InTempHeaderPaths)
+		{
+		}
 
-	/** the existing disk version of the header for this package's names */
-	FString OriginalNamesHeader;
+		/** Set of already exported cross-module references, to prevent duplicates */
+		TSet<FString>* UniqueCrossModuleReferences;
+		/** Array of all header filenames from the current package. */
+		TSet<FString>& PackageHeaderPaths;
+		/** Array of temp filenames that for files to overwrite headers */
+		TArray<FString>& TempHeaderPaths;
+		/** Forward declarations that we need. */
+		TSet<FString> ForwardDeclarations;
 
-	/** Array of temp filenames that for files to overwrite headers */
-	TArray<FString> TempHeaderPaths;
+	};
 
-	/** Array of all header filenames from the current package. */
-	TArray<FString> PackageHeaderPaths;
+	/** the existing disk version of the header for this package's names */
+	FString OriginalNamesHeader;
 
 	/** If false, exported headers will not be saved to disk */
 	bool bAllowSaveExportedHeaders;
@@ -111,9 +123,6 @@ private:
 	/** If true, any change in the generated headers will result in UHT failure. */
 	bool bFailIfGeneratedCodeChanges;
 
-	/** Forward declarations that we need. */
-	TSet<FString> ForwardDeclarations;
-
 	/**
 	 * Exports the struct's C++ properties to the HeaderText output device and adds special
 	 * compiler directives for GCC to pack as we expect.
@@ -125,13 +134,13 @@ private:
 	static void ExportProperties(FOutputDevice& Out, UStruct* Struct, int32 TextIndent);
 
 	/** Return the name of the singleton function that returns the UObject for Item */
-	const FString& GetPackageSingletonName(const UPackage* Item);
+	static const FString& GetPackageSingletonName(const UPackage* Item, TSet<FString>* UniqueCrossModuleReferences);
 
 	/** Return the name of the singleton function that returns the UObject for Item */
-	FString GetSingletonName(UField* Item, bool bRequiresValidObject=true);
+	static const FString& GetSingletonName(UField* Item, TSet<FString>* UniqueCrossModuleReferences, bool bRequiresValidObject=true);
 
 	/** Return the address of the singleton function - handles nullptr */
-	FString GetSingletonNameFuncAddr(UField* Item, bool bRequiresValidObject=true);
+	static FString GetSingletonNameFuncAddr(UField* Item, TSet<FString>* UniqueCrossModuleReferences, bool bRequiresValidObject=true);
 
 	/**
 	 * Returns the name (overridden if marked up) with TEXT("") or "" wrappers for use in a string literal.
@@ -164,7 +173,7 @@ private:
 	/**
 	 * Export functions used to find and call C++ or script implementation of a script function in the interface 
 	 */
-	void ExportInterfaceCallFunctions(FOutputDevice& OutCpp, FUHTStringBuilder& Out, const TArray<UFunction*>& CallbackFunctions, const TCHAR* ClassName);
+	void ExportInterfaceCallFunctions(FOutputDevice& OutCpp, FUHTStringBuilder& Out, FReferenceGatherers& OutReferenceGatherers, const TArray<UFunction*>& CallbackFunctions, const TCHAR* ClassName) const;
 
 	/**
 	 * Export UInterface boilerplate.
@@ -186,16 +195,17 @@ private:
 		FOutputDevice&           OutGeneratedHeaderText,
 		FOutputDevice&           OutputGetter,
 		FOutputDevice&           OutDeclarations,
+		FReferenceGatherers&     OutReferenceGatherers, 
 		FClass*                  Class,
 		const FUnrealSourceFile& SourceFile
-	);
+	) const;
 
 	/**
 	 * After all of the dependency checking, but before actually exporting the class, set up the generated code
 	 *
 	 * @param	SourceFile			Source file to export.
 	 */
-	bool WriteHeader(const TCHAR* Path, const FString& InBodyText, const TSet<FString>& InFwdDecl);
+	bool WriteHeader(const TCHAR* Path, const FString& InBodyText, FReferenceGatherers& InOutReferenceGatherers) const;
 
 	/**
 	 * Returns a string in the format CLASS_Something|CLASS_Something which represents all class flags that are set for the specified
@@ -209,7 +219,7 @@ private:
 	 * @param	Out		the output device for the mirror struct
 	 * @param	Enums	the enum to export
 	 */
-	void ExportEnum(FOutputDevice& Out, UEnum* Enum);
+	void ExportEnum(FOutputDevice& Out, UEnum* Enum) const;
 
 	/**
 	 * Exports the inl text for enums declared in non-UClass headers.
@@ -217,7 +227,7 @@ private:
 	 * @param	OutputGetter	The function to call to get the output.
 	 * @param	Enum			the enum to export
 	 */
-	void ExportGeneratedEnumInitCode(FOutputDevice& Out, const FUnrealSourceFile& SourceFile, UEnum* Enum);
+	void ExportGeneratedEnumInitCode(FOutputDevice& Out, FReferenceGatherers& OutReferenceGatherers, const FUnrealSourceFile& SourceFile, UEnum* Enum) const;
 
 	/**
 	 * Exports the macro declarations for GENERATED_BODY() for each Foo in the struct specified
@@ -225,7 +235,7 @@ private:
 	 * @param	Out				output device
 	 * @param	Struct			The struct to export
 	 */
-	void ExportGeneratedStructBodyMacros(FOutputDevice& OutGeneratedHeaderText, FOutputDevice& Out, const FUnrealSourceFile& SourceFile, UScriptStruct* Struct);
+	void ExportGeneratedStructBodyMacros(FOutputDevice& OutGeneratedHeaderText, FOutputDevice& Out, FReferenceGatherers& OutReferenceGatherers, const FUnrealSourceFile& SourceFile, UScriptStruct* Struct) const;
 
 	/**
 	 * Exports a local mirror of the specified struct; used to get offsets
@@ -251,7 +261,7 @@ private:
 	 * @param	SourceFile			Source file of the delegate.
 	 * @param	DelegateFunctions	the functions that have parameters which need to be exported
 	 */
-	void ExportDelegateDeclaration(FOutputDevice& Out, const FUnrealSourceFile& SourceFile, UFunction* Function);
+	void ExportDelegateDeclaration(FOutputDevice& Out, FReferenceGatherers& OutReferenceGatherers, const FUnrealSourceFile& SourceFile, UFunction* Function) const;
 
 	/**
 	 * Exports C++ type definitions for delegates
@@ -260,7 +270,7 @@ private:
 	 * @param	SourceFile			Source file of the delegate.
 	 * @param	DelegateFunctions	the functions that have parameters which need to be exported
 	 */
-	void ExportDelegateDefinition(FOutputDevice& Out, const FUnrealSourceFile& SourceFile, UFunction* Function);
+	void ExportDelegateDefinition(FOutputDevice& Out, FReferenceGatherers& OutReferenceGatherers, const FUnrealSourceFile& SourceFile, UFunction* Function) const;
 
 	/**
 	 * Exports the parameter struct declarations for the given function.
@@ -277,7 +287,7 @@ private:
 	* 
 	* @param	PackageName	Name of the package being saved
 	*/
-	void ExportUpdatedHeaders( const FString& PackageName  );
+	static void ExportUpdatedHeaders(const FString& PackageName, const TArray<FString>& TempHeaderPaths);
 
 	/**
 	 * Exports the generated cpp file for all functions/events/delegates in package.
@@ -323,7 +333,7 @@ private:
 	* @param	ValidatePosition		Position in source file of _Validate function for function described by FunctionData.
 	* @param	SourceFile				Currently analyzed source file.
 	*/
-	void CheckRPCFunctions(const FFuncInfo& FunctionData, const FString& ClassName, int32 ImplementationPosition, int32 ValidatePosition, const FUnrealSourceFile& SourceFile);
+	void CheckRPCFunctions(FReferenceGatherers& OutReferenceGatherers, const FFuncInfo& FunctionData, const FString& ClassName, int32 ImplementationPosition, int32 ValidatePosition, const FUnrealSourceFile& SourceFile) const;
 
 	/**
 	 * Exports the native stubs for the list of functions specified
@@ -332,7 +342,7 @@ private:
 	 * @param Class			class
 	 * @param ClassData		class data
 	 */
-	void ExportNativeFunctions(FOutputDevice& OutGeneratedHeaderText, FOutputDevice& OutGeneratedCPPText, FOutputDevice& OutMacroCalls, FOutputDevice& OutNoPureDeclsMacroCalls, const FUnrealSourceFile& SourceFile, UClass* Class, FClassMetaData* ClassData);
+	void ExportNativeFunctions(FOutputDevice& OutGeneratedHeaderText, FOutputDevice& OutGeneratedCPPText, FOutputDevice& OutMacroCalls, FOutputDevice& OutNoPureDeclsMacroCalls, FReferenceGatherers& OutReferenceGatherers, const FUnrealSourceFile& SourceFile, UClass* Class, FClassMetaData* ClassData) const;
 
 	/**
 	 * Export the actual internals to a standard thunk function
@@ -343,7 +353,7 @@ private:
 	 * @param Parameters list of parameters in the function
 	 * @param Return return parameter for the function
 	 */
-	void ExportFunctionThunk(FUHTStringBuilder& RPCWrappers, UFunction* Function, const FFuncInfo& FunctionData, const TArray<FProperty*>& Parameters, FProperty* Return);
+	void ExportFunctionThunk(FUHTStringBuilder& RPCWrappers, FReferenceGatherers& OutReferenceGatherers, UFunction* Function, const FFuncInfo& FunctionData, const TArray<FProperty*>& Parameters, FProperty* Return) const;
 
 	/** Exports the native function registration code for the given class. */
 	static void ExportNatives(FOutputDevice& Out, FClass* Class);
@@ -356,7 +366,7 @@ private:
 	 * @param	Class			Class to export
 	 * @param	OutFriendText	(Output parameter) Friend text
 	 */
-	void ExportNativeGeneratedInitCode(FOutputDevice& Out, FOutputDevice& OutDeclarations, const FUnrealSourceFile& SourceFile, FClass* Class, FUHTStringBuilder& OutFriendText);
+	void ExportNativeGeneratedInitCode(FOutputDevice& Out, FOutputDevice& OutDeclarations, FReferenceGatherers& OutReferenceGatherers, const FUnrealSourceFile& SourceFile, FClass* Class, FUHTStringBuilder& OutFriendText) const;
 
 	/**
 	 * Export given function.
@@ -365,7 +375,7 @@ private:
 	 * @param	Function		Given function.
 	 * @param	bIsNoExport		Is in NoExport class.
 	 */
-	void ExportFunction(FOutputDevice& Out, const FUnrealSourceFile& SourceFile, UFunction* Function, bool bIsNoExport);
+	void ExportFunction(FOutputDevice& Out, FReferenceGatherers& OutReferenceGatherers, const FUnrealSourceFile& SourceFile, UFunction* Function, bool bIsNoExport) const;
 
 	/**
 	 * Exports a generated singleton function to setup the package for compiled-in classes.
@@ -387,7 +397,7 @@ private:
 	 *
 	 * @return      A pair of strings which represents the pointer and a count of the emitted properties.
 	 */
-	TTuple<FString, FString> OutputProperties(FOutputDevice& DeclOut, FOutputDevice& Out, const TCHAR* Scope, const TArray<FProperty*>& Properties, const TCHAR* DeclSpaces, const TCHAR* Spaces);
+	TTuple<FString, FString> OutputProperties(FOutputDevice& DeclOut, FOutputDevice& Out, FReferenceGatherers& OutReferenceGatherers, const TCHAR* Scope, const TArray<FProperty*>& Properties, const TCHAR* DeclSpaces, const TCHAR* Spaces) const;
 
 	/**
 	 * Function to output the C++ code necessary to set up a property
@@ -399,7 +409,7 @@ private:
 	 * @param	DeclSpaces		String of spaces to use as an indent for the declaration
 	 * @param	Spaces			String of spaces to use as an indent
 	**/
-	void OutputProperty(FOutputDevice& DeclOut, FOutputDevice& Out, const TCHAR* Scope, TArray<FPropertyNamePointerPair>& PropertyNamesAndPointers, FProperty* Prop, const TCHAR* DeclSpaces, const TCHAR* Spaces);
+	void OutputProperty(FOutputDevice& DeclOut, FOutputDevice& Out, FReferenceGatherers& OutReferenceGatherers, const TCHAR* Scope, TArray<FPropertyNamePointerPair>& PropertyNamesAndPointers, FProperty* Prop, const TCHAR* DeclSpaces, const TCHAR* Spaces) const;
 
 	/**
 	 * Function to output the C++ code necessary to set up a property, including an array property and its inner, array dimensions, etc.
@@ -413,7 +423,7 @@ private:
 	 * @param	Spaces			String of spaces to use as an indent
 	 * @param	SourceStruct	Structure that the property offset is relative to
 	**/
-	void PropertyNew(FOutputDevice& DeclOut, FOutputDevice& Out, FProperty* Prop, const TCHAR* OffsetStr, const TCHAR* Name, const TCHAR* DeclSpaces, const TCHAR* Spaces, const TCHAR* SourceStruct = NULL);
+	void PropertyNew(FOutputDevice& DeclOut, FOutputDevice& Out, FReferenceGatherers& OutReferenceGatherers, FProperty* Prop, const TCHAR* OffsetStr, const TCHAR* Name, const TCHAR* DeclSpaces, const TCHAR* Spaces, const TCHAR* SourceStruct = nullptr) const;
 
 	/**
 	 * Exports the proxy definitions for the list of enums specified
@@ -457,12 +467,12 @@ private:
 	 * @param NewHeaderContents	Contents of the generated header.
 	 * @return True if the header contents has changed, false otherwise.
 	 */
-	bool SaveHeaderIfChanged(const TCHAR* HeaderPath, const TCHAR* NewHeaderContents);
+	bool SaveHeaderIfChanged(FReferenceGatherers& OutReferenceGatherers, const TCHAR* HeaderPath, const TCHAR* NewHeaderContents) const;
 
 	/**
 	 * Deletes all .generated.h files which do not correspond to any of the classes.
 	 */
-	void DeleteUnusedGeneratedHeaders();
+	static void DeleteUnusedGeneratedHeaders(const TSet<FString>& PackageHeaderPathSet);
 
 	/**
 	 * Exports macros that manages UObject constructors.
@@ -523,7 +533,7 @@ public:
 	 * @param Function Function to get return type of.
 	 * @return FString with function return type.
 	 */
-	FString GetFunctionReturnString(UFunction* Function);
+	static FString GetFunctionReturnString(UFunction* Function, FReferenceGatherers& OutReferenceGatherers);
 
 	/**
 	* Gets string with function parameters (with names).
@@ -531,7 +541,7 @@ public:
 	* @param Function Function to get parameters of.
 	* @return FString with function parameters.
 	*/
-	FString GetFunctionParameterString(UFunction* Function);
+	static FString GetFunctionParameterString(UFunction* Function, FReferenceGatherers& OutReferenceGatherers);
 
 	/**
 	 * Checks if function is missing "virtual" specifier.
diff --git a/Engine/Source/Programs/UnrealHeaderTool/Private/ParserClass.cpp b/Engine/Source/Programs/UnrealHeaderTool/Private/ParserClass.cpp
--- a/Engine/Source/Programs/UnrealHeaderTool/Private/ParserClass.cpp
+++ b/Engine/Source/Programs/UnrealHeaderTool/Private/ParserClass.cpp
@@ -124,13 +124,16 @@ bool FClass::IsOwnedByDynamicType(const FField* Field)
 	return false;
 }
 
-static TMap<const UField*, FString> UFieldTypePackageNames;
-static TMap<const FField*, FString> FFieldTypePackageNames;
+static TMap<const UField*, TUniquePtr<FString>> UFieldTypePackageNames;
+static TMap<const FField*, TUniquePtr<FString>> FFieldTypePackageNames;
 
 template <typename T>
-const FString& GetTypePackageName_Inner(const T* Field, TMap<const T*, FString>& TypePackageNames)
+const FString& GetTypePackageName_Inner(const T* Field, TMap<const T*, TUniquePtr<FString>>& TypePackageNames)
 {
-	FString* TypePackageName = TypePackageNames.Find(Field);
+	static FRWLock Lock;
+	FRWScopeLock TypeNameLock(Lock, SLT_ReadOnly);
+
+	TUniquePtr<FString>* TypePackageName = TypePackageNames.Find(Field);
 	if (TypePackageName == nullptr)
 	{
 		FString PackageName = Field->GetMetaData(FClass::NAME_ReplaceConverted);
@@ -147,9 +150,17 @@ const FString& GetTypePackageName_Inner(const T* Field, TMap<const T*, FString>&
 		{
 			PackageName = Field->GetOutermost()->GetName();
 		}
-		TypePackageName = &TypePackageNames.Add(Field, MoveTemp(PackageName));
+
+		TypeNameLock.ReleaseReadOnlyLockAndAcquireWriteLock_USE_WITH_CAUTION();
+
+		// Check the map again in case another thread had also been waiting on writing this data and got the write lock first
+		TypePackageName = TypePackageNames.Find(Field);
+		if (TypePackageName == nullptr)
+		{
+			TypePackageName = &TypePackageNames.Add(Field, MakeUnique<FString>(MoveTemp(PackageName)));
+		}
 	}
-	return *TypePackageName;
+	return **TypePackageName;
 }
 
 const FString& FClass::GetTypePackageName(const UField* Field)
diff --git a/Engine/Source/Programs/UnrealHeaderTool/Private/ParserHelper.cpp b/Engine/Source/Programs/UnrealHeaderTool/Private/ParserHelper.cpp
--- a/Engine/Source/Programs/UnrealHeaderTool/Private/ParserHelper.cpp
+++ b/Engine/Source/Programs/UnrealHeaderTool/Private/ParserHelper.cpp
@@ -286,21 +286,31 @@ FTokenData* FPropertyData::Set(FProperty* InKey, const FTokenData& InValue, FUnr
 const TCHAR* FNameLookupCPP::GetNameCPP(UStruct* Struct, bool bForceInterface /*= false */)
 {
 	TMap<UStruct*, TCHAR*>& NameMap = (bForceInterface ? InterfaceNameMap : StructNameMap);
+	FRWScopeLock Lock(bForceInterface ? InterfaceNameLock : StructNameLock, SLT_ReadOnly);
 
 	TCHAR* NameCPP = NameMap.FindRef(Struct);
-	if (NameCPP)
+	if (NameCPP == nullptr)
 	{
-		return NameCPP;
-	}
+		const FString TempName = FString::Printf(TEXT("%s%s"), (bForceInterface ? TEXT("I") : Struct->GetPrefixCPP()), *Struct->GetName());
+		const int32 StringLength = TempName.Len();
 
-	const FString TempName = FString::Printf(TEXT("%s%s"), (bForceInterface ? TEXT("I") : Struct->GetPrefixCPP()), *Struct->GetName());
-	const int32 StringLength = TempName.Len();
+		NameCPP = new TCHAR[StringLength + 1];
+		FCString::Strcpy(NameCPP, StringLength + 1, *TempName);
+		NameCPP[StringLength] = 0;
 
-	NameCPP = new TCHAR[StringLength + 1];
-	FCString::Strcpy(NameCPP, StringLength + 1, *TempName);
-	NameCPP[StringLength] = 0;
+		Lock.ReleaseReadOnlyLockAndAcquireWriteLock_USE_WITH_CAUTION();
 
-	NameMap.Add(Struct, NameCPP);
+		// Check the map again in case another thread had also been waiting on writing this data and got the write lock first
+		if (TCHAR* NameInMap = NameMap.FindRef(Struct))
+		{
+			delete NameCPP;
+			NameCPP = NameInMap;
+		}
+		else
+		{
+			NameMap.Add(Struct, NameCPP);
+		}
+	}
 
 	return NameCPP;
 }
diff --git a/Engine/Source/Programs/UnrealHeaderTool/Private/ParserHelper.h b/Engine/Source/Programs/UnrealHeaderTool/Private/ParserHelper.h
--- a/Engine/Source/Programs/UnrealHeaderTool/Private/ParserHelper.h
+++ b/Engine/Source/Programs/UnrealHeaderTool/Private/ParserHelper.h
@@ -1923,6 +1923,9 @@ private:
 	TMap<UStruct*,TCHAR*> StructNameMap;
 	TMap<UStruct*, TCHAR*> InterfaceNameMap;
 
+	FRWLock InterfaceNameLock;
+	FRWLock StructNameLock;
+
 	FUnrealSourceFile* UnrealSourceFile;
 };
 
