UHT Additional cleanup
* Remove FNameLookupCPP struct and interface name maps. Overhead of managing/locking maps provides no benefit to the very simple Printf that it is ultimately doing.
* Remove unused FNameLookupCPP source file
* Remove GClassStrippedHeaderTextMap and GClassHeaderNameWithNoPathMap as they are unused
* No longer add Enum hashes to global table as they are not read.
* CheckedMetadataSpecifiers now stored as a FName keyed TMap as the incoming values are names not strings, so no need to convert name to string
#rnx

diff --git a/Engine/Source/Programs/UnrealHeaderTool/Private/ClassMaps.cpp b/Engine/Source/Programs/UnrealHeaderTool/Private/ClassMaps.cpp
--- a/Engine/Source/Programs/UnrealHeaderTool/Private/ClassMaps.cpp
+++ b/Engine/Source/Programs/UnrealHeaderTool/Private/ClassMaps.cpp
@@ -8,8 +8,6 @@ TMap<FString, TSharedRef<FUnrealSourceFile> > GUnrealSourceFilesMap;
 TMap<UField*, TSharedRef<FUnrealTypeDefinitionInfo> > GTypeDefinitionInfoMap;
 TMap<const UPackage*, TArray<UField*>> GPackageSingletons;
 FCriticalSection GPackageSingletonsCriticalSection;
-TMap<UClass*, FString> GClassStrippedHeaderTextMap;
-TMap<UClass*, FString> GClassHeaderNameWithNoPathMap;
 TSet<FUnrealSourceFile*> GPublicSourceFileSet;
 TMap<FProperty*, FString> GArrayDimensions;
 TMap<UPackage*,  const FManifestModule*> GPackageToManifestModuleMap;
diff --git a/Engine/Source/Programs/UnrealHeaderTool/Private/ClassMaps.h b/Engine/Source/Programs/UnrealHeaderTool/Private/ClassMaps.h
--- a/Engine/Source/Programs/UnrealHeaderTool/Private/ClassMaps.h
+++ b/Engine/Source/Programs/UnrealHeaderTool/Private/ClassMaps.h
@@ -36,8 +36,6 @@ extern TMap<FString, TSharedRef<FUnrealSourceFile> > GUnrealSourceFilesMap;
 extern TMap<UField*, TSharedRef<FUnrealTypeDefinitionInfo> > GTypeDefinitionInfoMap;
 extern TMap<const UPackage*, TArray<UField*>> GPackageSingletons;
 extern FCriticalSection GPackageSingletonsCriticalSection;
-extern TMap<UClass*, FString> GClassStrippedHeaderTextMap;
-extern TMap<UClass*, FString> GClassHeaderNameWithNoPathMap;
 extern TSet<FUnrealSourceFile*> GPublicSourceFileSet;
 extern TMap<FProperty*, FString> GArrayDimensions;
 extern TMap<UPackage*,  const FManifestModule*> GPackageToManifestModuleMap;
diff --git a/Engine/Source/Programs/UnrealHeaderTool/Private/CodeGenerator.cpp b/Engine/Source/Programs/UnrealHeaderTool/Private/CodeGenerator.cpp
--- a/Engine/Source/Programs/UnrealHeaderTool/Private/CodeGenerator.cpp
+++ b/Engine/Source/Programs/UnrealHeaderTool/Private/CodeGenerator.cpp
@@ -79,9 +79,6 @@ static TSharedRef<FUnrealSourceFile> PerformInitialParseOnHeader(UPackage* InPar
 
 FCompilerMetadataManager GScriptHelper;
 
-/** C++ name lookup helper */
-FNameLookupCPP NameLookupCPP;
-
 bool HasIdentifierExactMatch(const TCHAR* StringBegin, const TCHAR* StringEnd, const FString& Find);
 
 namespace
@@ -999,8 +996,7 @@ private:
 
 			if (Cast<UClass>(Outer) || Cast<UScriptStruct>(Outer))
 			{
-				const TCHAR* OuterName = NameLookupCPP.GetNameCPP(Cast<UStruct>(Outer));
-				Result = FString::Printf(TEXT("%s%s"), OuterName, *Result);
+				Result = FNameLookupCPP::GetNameCPP(Cast<UStruct>(Outer)) + Result;
 
 				// Structs can also have UPackage outer.
 				if (Cast<UClass>(Outer) || Cast<UPackage>(Outer->GetOuter()))
@@ -1020,8 +1016,8 @@ private:
 			Result = FPackageName::GetShortName(Result);
 		}
 
-		const TCHAR* ClassString = NameLookupCPP.GetNameCPP(Item->GetClass());
-		return FString::Printf(TEXT("Z_Construct_%s_%s%s()"), ClassString, *Result, Suffix);
+		const FString ClassString = FNameLookupCPP::GetNameCPP(Item->GetClass());
+		return FString::Printf(TEXT("Z_Construct_%s_%s%s()"), *ClassString, *Result, Suffix);
 	}
 
 	static TMap<FTypeSingletonCacheKey, TUniquePtr<FTypeSingleton>> CacheData;
@@ -1881,7 +1877,7 @@ void FNativeClassHeaderGenerator::OutputProperty(FOutputDevice& DeclOut, FOutput
 		}
 		else
 		{
-			SourceStruct = NameLookupCPP.GetNameCPP(CastChecked<UStruct>(Prop->GetOwner<UObject>()));
+			SourceStruct = FNameLookupCPP::GetNameCPP(CastChecked<UStruct>(Prop->GetOwner<UObject>()));
 		}
 
 		FString PropName = Prop->GetName();
@@ -2040,12 +2036,10 @@ const FString& FNativeClassHeaderGenerator::GetPackageSingletonName(const UPacka
 	TUniquePtr<FPackageSingletonStrings>* PackageSingletonStrings = PackageSingletonNames.Find(InPackage);
 	if (PackageSingletonStrings == nullptr)
 	{
-		static const TCHAR* ClassString = NameLookupCPP.GetNameCPP(UPackage::StaticClass());
-
 		FString PackageName = InPackage->GetName();
 		PackageName.ReplaceInline(TEXT("/"), TEXT("_"), ESearchCase::CaseSensitive);
 
-		TUniquePtr<FPackageSingletonStrings> NewPackageSingletonStrings(MakeUnique<FPackageSingletonStrings>(FString::Printf(TEXT("Z_Construct_%s_%s()"), ClassString, *PackageName)));
+		TUniquePtr<FPackageSingletonStrings> NewPackageSingletonStrings(MakeUnique<FPackageSingletonStrings>(FString::Printf(TEXT("Z_Construct_UPackage_%s()"), *PackageName)));
 
 		Lock.ReleaseReadOnlyLockAndAcquireWriteLock_USE_WITH_CAUTION();
 		
@@ -2140,7 +2134,7 @@ void FNativeClassHeaderGenerator::ExportNativeGeneratedInitCode(FOutputDevice& O
 
 	const bool   bIsNoExport  = Class->HasAnyClassFlags(CLASS_NoExport);
 	const bool   bIsDynamic   = FClass::IsDynamic(Class);
-	const TCHAR* ClassNameCPP = NameLookupCPP.GetNameCPP(Class);
+	const FString ClassNameCPP = FNameLookupCPP::GetNameCPP(Class);
 
 	const FString& ApiString = GetAPIString();
 
@@ -2183,7 +2177,7 @@ void FNativeClassHeaderGenerator::ExportNativeGeneratedInitCode(FOutputDevice& O
 
 			GeneratedClassRegisterFunctionText.Logf(TEXT("\tUClass* %s\r\n"), *SingletonNameNoRegister);
 			GeneratedClassRegisterFunctionText.Logf(TEXT("\t{\r\n"));
-			GeneratedClassRegisterFunctionText.Logf(TEXT("\t\treturn %s::StaticClass();\r\n"), ClassNameCPP);
+			GeneratedClassRegisterFunctionText.Logf(TEXT("\t\treturn %s::StaticClass();\r\n"), *ClassNameCPP);
 			GeneratedClassRegisterFunctionText.Logf(TEXT("\t}\r\n"));
 		}
 		const FString& SingletonName = GetSingletonName(Class, OutReferenceGatherers.UniqueCrossModuleReferences);
@@ -2315,7 +2309,7 @@ void FNativeClassHeaderGenerator::ExportNativeGeneratedInitCode(FOutputDevice& O
 				FString OffsetString;
 				if (Inter.PointerOffset)
 				{
-					OffsetString = FString::Printf(TEXT("(int32)VTABLE_OFFSET(%s, %s)"), ClassNameCPP, NameLookupCPP.GetNameCPP(Inter.Class, true));
+					OffsetString = FString::Printf(TEXT("(int32)VTABLE_OFFSET(%s, %s)"), *ClassNameCPP, *FNameLookupCPP::GetNameCPP(Inter.Class, true));
 				}
 				else
 				{
@@ -2342,7 +2336,7 @@ void FNativeClassHeaderGenerator::ExportNativeGeneratedInitCode(FOutputDevice& O
 		GeneratedClassRegisterFunctionText.Logf(TEXT("\t\tstatic const FCppClassTypeInfoStatic StaticCppClassTypeInfo;\r\n"));
 
 		StaticDefinitions.Logf(TEXT("\tconst FCppClassTypeInfoStatic %s::StaticCppClassTypeInfo = {\r\n"), *StaticsStructName);
-		StaticDefinitions.Logf(TEXT("\t\tTCppClassTypeTraits<%s>::IsAbstract,\r\n"), NameLookupCPP.GetNameCPP(Class, Class->HasAllClassFlags(CLASS_Interface)));
+		StaticDefinitions.Logf(TEXT("\t\tTCppClassTypeTraits<%s>::IsAbstract,\r\n"), *FNameLookupCPP::GetNameCPP(Class, Class->HasAllClassFlags(CLASS_Interface)));
 		StaticDefinitions.Logf(TEXT("\t};\r\n"));
 
 		GeneratedClassRegisterFunctionText.Log (TEXT("\t\tstatic const UE4CodeGen_Private::FClassParams ClassParams;\r\n"));
@@ -2355,7 +2349,7 @@ void FNativeClassHeaderGenerator::ExportNativeGeneratedInitCode(FOutputDevice& O
 		ClassFlags = ClassFlags & CLASS_SaveInCompiledInClasses;
 
 		StaticDefinitions.Logf(TEXT("\tconst UE4CodeGen_Private::FClassParams %s::ClassParams = {\r\n"), *StaticsStructName);
-		StaticDefinitions.Logf(TEXT("\t\t&%s::StaticClass,\r\n"), ClassNameCPP);
+		StaticDefinitions.Logf(TEXT("\t\t&%s::StaticClass,\r\n"), *ClassNameCPP);
 		StaticDefinitions.Logf(TEXT("\t\t%s,\r\n"), (Class->ClassConfigName != NAME_None) ? *CreateUTF8LiteralString(Class->ClassConfigName.ToString()) : TEXT("nullptr"));
 		StaticDefinitions.Log (TEXT("\t\t&StaticCppClassTypeInfo,\r\n"));
 		StaticDefinitions.Logf(TEXT("\t\t%s,\r\n"), SingletonsArray);
@@ -2417,7 +2411,7 @@ void FNativeClassHeaderGenerator::ExportNativeGeneratedInitCode(FOutputDevice& O
 
 	if (OutFriendText.Len() && bIsNoExport)
 	{
-		Out.Logf(TEXT("\t/* friend declarations for pasting into noexport class %s\r\n"), ClassNameCPP);
+		Out.Logf(TEXT("\t/* friend declarations for pasting into noexport class %s\r\n"), *ClassNameCPP);
 		Out.Log(OutFriendText);
 		Out.Logf(TEXT("\t*/\r\n"));
 		OutFriendText.Reset();
@@ -2431,7 +2425,7 @@ void FNativeClassHeaderGenerator::ExportNativeGeneratedInitCode(FOutputDevice& O
 	const FString EmptyString;
 	const FString& InitSearchableValuesFunctionName = bIsDynamic ? Class->GetMetaData(NAME_InitializeStaticSearchableValues) : EmptyString;
 	const FString InitSearchableValuesFunctionParam = InitSearchableValuesFunctionName.IsEmpty() ? FString(TEXT("nullptr")) :
-		FString::Printf(TEXT("&%s::%s"), ClassNameCPP, *InitSearchableValuesFunctionName);
+		FString::Printf(TEXT("&%s::%s"), *ClassNameCPP, *InitSearchableValuesFunctionName);
 
 	// Append base class' hash at the end of the generated code, this will force update derived classes
 	// when base class changes during hot-reload.
@@ -2483,25 +2477,25 @@ void FNativeClassHeaderGenerator::ExportNativeGeneratedInitCode(FOutputDevice& O
 	// Emit the IMPLEMENT_CLASS macro to go in the generated cpp file.
 	if (!bIsDynamic)
 	{
-		Out.Logf(TEXT("\tIMPLEMENT_CLASS(%s, %u);\r\n"), ClassNameCPP, ClassHash);
+		Out.Logf(TEXT("\tIMPLEMENT_CLASS(%s, %u);\r\n"), *ClassNameCPP, ClassHash);
 	}
 	else
 	{
-		Out.Logf(TEXT("\tIMPLEMENT_DYNAMIC_CLASS(%s, TEXT(\"%s\"), %u);\r\n"), ClassNameCPP, *OverriddenClassName, ClassHash);
+		Out.Logf(TEXT("\tIMPLEMENT_DYNAMIC_CLASS(%s, TEXT(\"%s\"), %u);\r\n"), *ClassNameCPP, *OverriddenClassName, ClassHash);
 	}
 
-	Out.Logf(TEXT("\ttemplate<> %sUClass* StaticClass<%s>()\r\n"), *GetAPIString(), ClassNameCPP);
+	Out.Logf(TEXT("\ttemplate<> %sUClass* StaticClass<%s>()\r\n"), *GetAPIString(), *ClassNameCPP);
 	Out.Logf(TEXT("\t{\r\n"));
-	Out.Logf(TEXT("\t\treturn %s::StaticClass();\r\n"), ClassNameCPP);
+	Out.Logf(TEXT("\t\treturn %s::StaticClass();\r\n"), *ClassNameCPP);
 	Out.Logf(TEXT("\t}\r\n"));
 
 	if (bIsDynamic)
 	{
 		const FString& ClassPackageName = FClass::GetTypePackageName(Class);
 		Out.Logf(TEXT("\tstatic FCompiledInDefer Z_CompiledInDefer_UClass_%s(%s, &%s::StaticClass, TEXT(\"%s\"), TEXT(\"%s\"), true, %s, %s, %s);\r\n"),
-			ClassNameCPP,
+			*ClassNameCPP,
 			*SingletonName,
-			ClassNameCPP,
+			*ClassNameCPP,
 			*ClassPackageName,
 			*OverriddenClassName,
 			*AsTEXT(ClassPackageName),
@@ -2511,11 +2505,11 @@ void FNativeClassHeaderGenerator::ExportNativeGeneratedInitCode(FOutputDevice& O
 	else
 	{
 		Out.Logf(TEXT("\tstatic FCompiledInDefer Z_CompiledInDefer_UClass_%s(%s, &%s::StaticClass, TEXT(\"%s\"), TEXT(\"%s\"), false, nullptr, nullptr, %s);\r\n"),
-			ClassNameCPP,
+			*ClassNameCPP,
 			*SingletonName,
-			ClassNameCPP,
+			*ClassNameCPP,
 			*Class->GetOutermost()->GetName(),
-			ClassNameCPP,
+			*ClassNameCPP,
 			*InitSearchableValuesFunctionParam);
 	}
 
@@ -2528,7 +2522,7 @@ void FNativeClassHeaderGenerator::ExportNativeGeneratedInitCode(FOutputDevice& O
 			"\r\n"
 			"\tvoid %s::ValidateGeneratedRepEnums(const TArray<struct FRepRecord>& ClassReps) const\r\n"
 			"\t{\r\n"
-		), ClassNameCPP);
+		), *ClassNameCPP);
 
 		FUHTStringBuilder NameBuilder;
 
@@ -2559,7 +2553,7 @@ void FNativeClassHeaderGenerator::ExportNativeGeneratedInitCode(FOutputDevice& O
 			"%s\r\n" // ValidationBuilder
 			"\t\tcheckf(bIsValid, TEXT(\"UHT Generated Rep Indices do not match runtime populated Rep Indices for properties in %s\"));\r\n"
 			"\t}\r\n"
-		), *NameBuilder, *ValidationBuilder, ClassNameCPP);
+		), *NameBuilder, *ValidationBuilder, *ClassNameCPP);
 	}
 }
 
@@ -2702,10 +2696,10 @@ void FNativeClassHeaderGenerator::ExportFunction(FOutputDevice& Out, FReferenceG
 
 void FNativeClassHeaderGenerator::ExportNatives(FOutputDevice& Out, FClass* Class)
 {
-	const TCHAR* ClassCPPName = NameLookupCPP.GetNameCPP(Class);
-	FString TypeName = Class->HasAnyClassFlags(CLASS_Interface) ? *FString::Printf(TEXT("I%s"), *Class->GetName()) : ClassCPPName;
+	const FString ClassCPPName = FNameLookupCPP::GetNameCPP(Class);
+	FString TypeName = Class->HasAnyClassFlags(CLASS_Interface) ? FString::Printf(TEXT("I%s"), *Class->GetName()) : ClassCPPName;
 
-	Out.Logf(TEXT("\tvoid %s::StaticRegisterNatives%s()\r\n"), ClassCPPName, ClassCPPName);
+	Out.Logf(TEXT("\tvoid %s::StaticRegisterNatives%s()\r\n"), *ClassCPPName, *ClassCPPName);
 	Out.Log(TEXT("\t{\r\n"));
 
 	{
@@ -2733,7 +2727,7 @@ void FNativeClassHeaderGenerator::ExportNatives(FOutputDevice& Out, FClass* Clas
 			FMacroBlockEmitter EditorOnly(Out, TEXT("WITH_EDITOR"));
 			EditorOnly(bAllEditorOnly);
 
-			Out.Logf(TEXT("\t\tUClass* Class = %s::StaticClass();\r\n"), ClassCPPName);
+			Out.Logf(TEXT("\t\tUClass* Class = %s::StaticClass();\r\n"), *ClassCPPName);
 			Out.Log(TEXT("\t\tstatic const FNameNativePtrPair Funcs[] = {\r\n"));
 
 			for (const TTuple<UFunction*, FString>& Func : NamedFunctionsToExport)
@@ -2785,7 +2779,7 @@ void FNativeClassHeaderGenerator::ExportInterfaceCallFunctions(FOutputDevice& Ou
 		ExportNativeFunctionHeader(Out, OutReferenceGatherers.ForwardDeclarations, FunctionData, EExportFunctionType::Interface, EExportFunctionHeaderStyle::Declaration, *ExtraParam, *APIString);
 		Out.Logf( TEXT(";") LINE_TERMINATOR );
 
-		FString FunctionNameName = FString::Printf(TEXT("NAME_%s_%s"), NameLookupCPP.GetNameCPP(CastChecked<UStruct>(Function->GetOuter())), *FunctionName);
+		FString FunctionNameName = FString::Printf(TEXT("NAME_%s_%s"), *FNameLookupCPP::GetNameCPP(CastChecked<UStruct>(Function->GetOuter())), *FunctionName);
 		OutCpp.Logf(TEXT("\tstatic FName %s = FName(TEXT(\"%s\"));") LINE_TERMINATOR, *FunctionNameName, *GetOverriddenFName(Function).ToString());
 
 		ExportNativeFunctionHeader(OutCpp, OutReferenceGatherers.ForwardDeclarations, FunctionData, EExportFunctionType::Interface, EExportFunctionHeaderStyle::Definition, *ExtraParam, *APIString);
@@ -3026,8 +3020,8 @@ void FNativeClassHeaderGenerator::ExportClassFromSourceFileInner(
 	FClass* SuperClass = Class->GetSuperClass();
 
 	// the name for the C++ version of the UClass
-	const TCHAR* ClassCPPName = NameLookupCPP.GetNameCPP(Class);
-	const TCHAR* SuperClassCPPName = (SuperClass != nullptr) ? NameLookupCPP.GetNameCPP(SuperClass) : nullptr;
+	const FString ClassCPPName = FNameLookupCPP::GetNameCPP(Class);
+	const FString SuperClassCPPName = (SuperClass ? FNameLookupCPP::GetNameCPP(SuperClass) : TEXT("None"));
 
 	FString APIArg = API;
 	if (!Class->HasAnyClassFlags(CLASS_MinimalAPI))
@@ -3076,9 +3070,9 @@ void FNativeClassHeaderGenerator::ExportClassFromSourceFileInner(
 			BoilerPlateCPP.Logf(TEXT("#if %s\r\n"), *EnclosingDefines);
 		}
 
-		Boilerplate.Logf(TEXT("\t%s(%s, %s_API)\r\n"), MacroNameHeader, ClassCPPName, *APIArg);
+		Boilerplate.Logf(TEXT("\t%s(%s, %s_API)\r\n"), MacroNameHeader, *ClassCPPName, *APIArg);
 		OutGeneratedHeaderText.Log(Macroize(*GeneratedSerializeFunctionHeaderMacroName, *Boilerplate));
-		BoilerPlateCPP.Logf(TEXT("\t%s(%s)\r\n"), MacroNameCPP, ClassCPPName);
+		BoilerPlateCPP.Logf(TEXT("\t%s(%s)\r\n"), MacroNameCPP, *ClassCPPName);
 
 		if (EnclosingDefines.Len())
 		{
@@ -3096,22 +3090,22 @@ void FNativeClassHeaderGenerator::ExportClassFromSourceFileInner(
 
 		// Export the class's native function registration.
 		Boilerplate.Logf(TEXT("private:\r\n"));
-		Boilerplate.Logf(TEXT("\tstatic void StaticRegisterNatives%s();\r\n"), ClassCPPName);
+		Boilerplate.Logf(TEXT("\tstatic void StaticRegisterNatives%s();\r\n"), *ClassCPPName);
 		Boilerplate.Log(*FriendText);
 		Boilerplate.Logf(TEXT("public:\r\n"));
 
 		const bool bCastedClass = Class->HasAnyCastFlag(CASTCLASS_AllFlags) && SuperClass && Class->ClassCastFlags != SuperClass->ClassCastFlags;
 
 		Boilerplate.Logf(TEXT("\tDECLARE_CLASS(%s, %s, COMPILED_IN_FLAGS(%s%s), %s, TEXT(\"%s\"), %s_API)\r\n"),
-			ClassCPPName,
-			SuperClassCPPName ? SuperClassCPPName : TEXT("None"),
+			*ClassCPPName,
+			*SuperClassCPPName,
 			Class->HasAnyClassFlags(CLASS_Abstract) ? TEXT("CLASS_Abstract") : TEXT("0"),
 			*GetClassFlagExportText(Class),
-			bCastedClass ? *FString::Printf(TEXT("CASTCLASS_%s"), ClassCPPName) : TEXT("CASTCLASS_None"),
+			bCastedClass ? *FString::Printf(TEXT("CASTCLASS_%s"), *ClassCPPName) : TEXT("CASTCLASS_None"),
 			*FClass::GetTypePackageName(Class),
 			*APIArg);
 
-		Boilerplate.Logf(TEXT("\tDECLARE_SERIALIZER(%s)\r\n"), ClassCPPName);
+		Boilerplate.Logf(TEXT("\tDECLARE_SERIALIZER(%s)\r\n"), *ClassCPPName);
 
 		// Add the serialization function declaration if we generated one
 		if (GeneratedSerializeFunctionHeaderMacroName.Len() > 0)
@@ -3121,7 +3115,7 @@ void FNativeClassHeaderGenerator::ExportClassFromSourceFileInner(
 
 		if (SuperClass && Class->ClassWithin != SuperClass->ClassWithin)
 		{
-			Boilerplate.Logf(TEXT("\tDECLARE_WITHIN(%s)\r\n"), NameLookupCPP.GetNameCPP(Class->GetClassWithin()));
+			Boilerplate.Logf(TEXT("\tDECLARE_WITHIN(%s)\r\n"), *FNameLookupCPP::GetNameCPP(Class->GetClassWithin()));
 		}
 
 		if (Class->HasAnyClassFlags(CLASS_Interface))
@@ -3180,7 +3174,7 @@ void FNativeClassHeaderGenerator::ExportClassFromSourceFileInner(
 			FUHTStringBuilder InterfaceBoilerplate;
 
 			InterfaceBoilerplate.Logf(TEXT("protected:\r\n\tvirtual ~%s() {}\r\n"), *InterfaceCPPName);
-			InterfaceBoilerplate.Logf(TEXT("public:\r\n\ttypedef %s UClassType;\r\n"), ClassCPPName);
+			InterfaceBoilerplate.Logf(TEXT("public:\r\n\ttypedef %s UClassType;\r\n"), *ClassCPPName);
 			InterfaceBoilerplate.Logf(TEXT("\ttypedef %s ThisClass;\r\n"), *InterfaceCPPName);
 
 			ExportInterfaceCallFunctions(OutCpp, InterfaceBoilerplate, OutReferenceGatherers, CallbackFunctions, *Class->GetName());
@@ -3196,7 +3190,7 @@ void FNativeClassHeaderGenerator::ExportClassFromSourceFileInner(
 
 			if (ClassHasReplicatedProperties(Class))
 			{
-				WriteReplicatedMacroData(ClassRange, ClassCPPName, *APIArg, Class, SuperClass, InterfaceBoilerplate, SourceFile);
+				WriteReplicatedMacroData(ClassRange, *ClassCPPName, *APIArg, Class, SuperClass, InterfaceBoilerplate, SourceFile);
 			}
 
 			FString NoPureDeclsMacroName = SourceFile.GetGeneratedMacroName(ClassData, TEXT("_INCLASS_IINTERFACE_NO_PURE_DECLS"));
@@ -3218,12 +3212,12 @@ void FNativeClassHeaderGenerator::ExportClassFromSourceFileInner(
 			// export implementation of _getUObject for classes that implement interfaces
 			if (Class->Interfaces.Num() > 0)
 			{
-				Boilerplate.Logf(TEXT("\tvirtual UObject* _getUObject() const override { return const_cast<%s*>(this); }\r\n"), ClassCPPName);
+				Boilerplate.Logf(TEXT("\tvirtual UObject* _getUObject() const override { return const_cast<%s*>(this); }\r\n"), *ClassCPPName);
 			}
 
 			if (ClassHasReplicatedProperties(Class))
 			{
-				WriteReplicatedMacroData(ClassRange, ClassCPPName, *APIArg, Class, SuperClass, Boilerplate, SourceFile);
+				WriteReplicatedMacroData(ClassRange, *ClassCPPName, *APIArg, Class, SuperClass, Boilerplate, SourceFile);
 			}
 
 			{
@@ -3289,7 +3283,7 @@ void FNativeClassHeaderGenerator::ExportClassFromSourceFileInner(
 	}
 
 	// Forward declare the StaticClass specialisation in the header
-	OutGeneratedHeaderText.Logf(TEXT("template<> %sUClass* StaticClass<class %s>();\r\n\r\n"), *GetAPIString(), ClassCPPName);
+	OutGeneratedHeaderText.Logf(TEXT("template<> %sUClass* StaticClass<class %s>();\r\n\r\n"), *GetAPIString(), *ClassCPPName);
 
 	// If there is a serialization function implementation for the CPP file, add it now
 	if (GeneratedSerializeFunctionCPP.Len())
@@ -3456,7 +3450,7 @@ FString GetBuildPath(FUnrealSourceFile& SourceFile)
 
 void FNativeClassHeaderGenerator::ExportConstructorsMacros(FOutputDevice& OutGeneratedHeaderText, FOutputDevice& Out, FOutputDevice& StandardUObjectConstructorsMacroCall, FOutputDevice& EnhancedUObjectConstructorsMacroCall, const FString& ConstructorsMacroPrefix, FClass* Class, const TCHAR* APIArg)
 {
-	const TCHAR* ClassCPPName = NameLookupCPP.GetNameCPP(Class);
+	const FString ClassCPPName = FNameLookupCPP::GetNameCPP(Class);
 
 	FClassMetaData* ClassData = GScriptHelper.FindClassData(Class);
 	check(ClassData);
@@ -3466,12 +3460,12 @@ void FNativeClassHeaderGenerator::ExportConstructorsMacros(FOutputDevice& OutGen
 	FString StdMacroName = ConstructorsMacroPrefix + TEXT("_STANDARD_CONSTRUCTORS");
 	FString EnhMacroName = ConstructorsMacroPrefix + TEXT("_ENHANCED_CONSTRUCTORS");
 
-	ExportStandardConstructorsMacro(StdMacro, Class, ClassData, APIArg, ClassCPPName);
-	ExportEnhancedConstructorsMacro(EnhMacro, Class, ClassData, APIArg, ClassCPPName);
+	ExportStandardConstructorsMacro(StdMacro, Class, ClassData, APIArg, *ClassCPPName);
+	ExportEnhancedConstructorsMacro(EnhMacro, Class, ClassData, APIArg, *ClassCPPName);
 
 	if (!ClassData->bCustomVTableHelperConstructorDeclared)
 	{
-		Out.Logf(TEXT("\tDEFINE_VTABLE_PTR_HELPER_CTOR(%s);" LINE_TERMINATOR), ClassCPPName);
+		Out.Logf(TEXT("\tDEFINE_VTABLE_PTR_HELPER_CTOR(%s);" LINE_TERMINATOR), *ClassCPPName);
 	}
 
 	OutGeneratedHeaderText.Log(*Macroize(*StdMacroName, *StdMacro));
@@ -3607,7 +3601,7 @@ void FNativeClassHeaderGenerator::ExportGeneratedStructBodyMacros(FOutputDevice&
 
 	UStruct* BaseStruct = Struct->GetSuperStruct();
 
-	const TCHAR* StructNameCPP = NameLookupCPP.GetNameCPP(Struct);
+	const FString StructNameCPP = FNameLookupCPP::GetNameCPP(Struct);
 
 	const FString& SingletonName = GetSingletonName(Struct, OutReferenceGatherers.UniqueCrossModuleReferences);
 	const FString ChoppedSingletonName = SingletonName.LeftChop(2);
@@ -3670,8 +3664,8 @@ void FNativeClassHeaderGenerator::ExportGeneratedStructBodyMacros(FOutputDevice&
 		for (const FRigVMMethodInfo& MethodInfo : StructRigVMInfo->Methods)
 		{
 			FString ParameterSuffix = MethodInfo.Parameters.Declarations(true, TEXT(", \\\r\n\t\t"));
-			OutGeneratedHeaderText.Logf(TEXT("#define %s_%s() \\\r\n"), StructNameCPP, *MethodInfo.Name);
-			OutGeneratedHeaderText.Logf(TEXT("\t%s %s::Static%s( \\\r\n\t\t%s%s \\\r\n\t)\n"), *MethodInfo.ReturnType, StructNameCPP, *MethodInfo.Name, *StructMembers, *ParameterSuffix);
+			OutGeneratedHeaderText.Logf(TEXT("#define %s_%s() \\\r\n"), *StructNameCPP, *MethodInfo.Name);
+			OutGeneratedHeaderText.Logf(TEXT("\t%s %s::Static%s( \\\r\n\t\t%s%s \\\r\n\t)\n"), *MethodInfo.ReturnType, *StructNameCPP, *MethodInfo.Name, *StructMembers, *ParameterSuffix);
 		}
 		OutGeneratedHeaderText.Log(TEXT("\n"));
 	}
@@ -3701,7 +3695,7 @@ void FNativeClassHeaderGenerator::ExportGeneratedStructBodyMacros(FOutputDevice&
 			}
 		}
 
-		const FString SuperTypedef = BaseStruct ? FString::Printf(TEXT("\ttypedef %s Super;\r\n"), NameLookupCPP.GetNameCPP(BaseStruct)) : FString();
+		const FString SuperTypedef = BaseStruct ? FString::Printf(TEXT("\ttypedef %s Super;\r\n"), *FNameLookupCPP::GetNameCPP(BaseStruct)) : FString();
 
 		FString CombinedLine = FString::Printf(TEXT("%s%s%s%s%s"), *FriendLine, *StaticClassLine, *RigVMMethodsDeclarations, *PrivatePropertiesOffset, *SuperTypedef);
 		const FString MacroName = SourceFile.GetGeneratedBodyMacroName(Struct->StructMacroDeclaredLineNumber);
@@ -3711,7 +3705,7 @@ void FNativeClassHeaderGenerator::ExportGeneratedStructBodyMacros(FOutputDevice&
 
 		FString GetHashName = FString::Printf(TEXT("Get_%s_Hash"), *ChoppedSingletonName);
 
-		Out.Logf(TEXT("class UScriptStruct* %s::StaticStruct()\r\n"), StructNameCPP);
+		Out.Logf(TEXT("class UScriptStruct* %s::StaticStruct()\r\n"), *StructNameCPP);
 		Out.Logf(TEXT("{\r\n"));
 
 		// UStructs can have UClass or UPackage outer (if declared in non-UClass headers).
@@ -3731,7 +3725,7 @@ void FNativeClassHeaderGenerator::ExportGeneratedStructBodyMacros(FOutputDevice&
 		Out.Logf(TEXT("\t\textern %suint32 %s();\r\n"), *FriendApiString, *GetHashName);
 
 		Out.Logf(TEXT("\t\tSingleton = GetStaticStruct(%s, %s, TEXT(\"%s\"), sizeof(%s), %s());\r\n"),
-			*ChoppedSingletonName, *OuterName, *ActualStructName, StructNameCPP, *GetHashName);
+			*ChoppedSingletonName, *OuterName, *ActualStructName, *StructNameCPP, *GetHashName);
 
 		// if this struct has RigVM methods - we need to register the method to our central
 		// registry on construction of the static struct
@@ -3740,7 +3734,7 @@ void FNativeClassHeaderGenerator::ExportGeneratedStructBodyMacros(FOutputDevice&
 			for (const FRigVMMethodInfo& MethodInfo : StructRigVMInfo->Methods)
 			{
 				Out.Logf(TEXT("\t\tFRigVMRegistry::Get().Register(TEXT(\"%s::%s\"), &%s::RigVM%s);\r\n"),
-					StructNameCPP, *MethodInfo.Name, StructNameCPP, *MethodInfo.Name);
+					*StructNameCPP, *MethodInfo.Name, *StructNameCPP, *MethodInfo.Name);
 			}
 		}
 
@@ -3749,20 +3743,20 @@ void FNativeClassHeaderGenerator::ExportGeneratedStructBodyMacros(FOutputDevice&
 		Out.Logf(TEXT("}\r\n"));
 
 		// Forward declare the StaticStruct specialisation in the header
-		OutGeneratedHeaderText.Logf(TEXT("template<> %sUScriptStruct* StaticStruct<struct %s>();\r\n\r\n"), *GetAPIString(), StructNameCPP);
+		OutGeneratedHeaderText.Logf(TEXT("template<> %sUScriptStruct* StaticStruct<struct %s>();\r\n\r\n"), *GetAPIString(), *StructNameCPP);
 
 		// Generate the StaticStruct specialisation
-		Out.Logf(TEXT("template<> %sUScriptStruct* StaticStruct<%s>()\r\n"), *GetAPIString(), StructNameCPP);
+		Out.Logf(TEXT("template<> %sUScriptStruct* StaticStruct<%s>()\r\n"), *GetAPIString(), *StructNameCPP);
 		Out.Logf(TEXT("{\r\n"));
-		Out.Logf(TEXT("\treturn %s::StaticStruct();\r\n"), StructNameCPP);
+		Out.Logf(TEXT("\treturn %s::StaticStruct();\r\n"), *StructNameCPP);
 		Out.Logf(TEXT("}\r\n"));
 
 		if (bIsDynamic)
 		{
 			const FString& StructPackageName = FClass::GetTypePackageName(Struct);
 			Out.Logf(TEXT("static FCompiledInDeferStruct Z_CompiledInDeferStruct_UScriptStruct_%s(%s::StaticStruct, TEXT(\"%s\"), TEXT(\"%s\"), true, %s, %s);\r\n"),
-				StructNameCPP,
-				StructNameCPP,
+				*StructNameCPP,
+				*StructNameCPP,
 				*StructPackageName,
 				*ActualStructName,
 				*AsTEXT(StructPackageName),
@@ -3771,8 +3765,8 @@ void FNativeClassHeaderGenerator::ExportGeneratedStructBodyMacros(FOutputDevice&
 		else
 		{
 			Out.Logf(TEXT("static FCompiledInDeferStruct Z_CompiledInDeferStruct_UScriptStruct_%s(%s::StaticStruct, TEXT(\"%s\"), TEXT(\"%s\"), false, nullptr, nullptr);\r\n"),
-				StructNameCPP,
-				StructNameCPP,
+				*StructNameCPP,
+				*StructNameCPP,
 				*Struct->GetOutermost()->GetName(),
 				*ActualStructName);
 		}
@@ -3781,15 +3775,15 @@ void FNativeClassHeaderGenerator::ExportGeneratedStructBodyMacros(FOutputDevice&
 		if (!Struct->GetOuter()->IsA(UStruct::StaticClass()))
 		{
 			const FString ShortPackageName = FPackageName::GetShortName(Struct->GetOuter()->GetName());
-			Out.Logf(TEXT("static struct FScriptStruct_%s_StaticRegisterNatives%s\r\n"), *ShortPackageName, StructNameCPP);
+			Out.Logf(TEXT("static struct FScriptStruct_%s_StaticRegisterNatives%s\r\n"), *ShortPackageName, *StructNameCPP);
 			Out.Logf(TEXT("{\r\n"));
-			Out.Logf(TEXT("\tFScriptStruct_%s_StaticRegisterNatives%s()\r\n"), *ShortPackageName, StructNameCPP);
+			Out.Logf(TEXT("\tFScriptStruct_%s_StaticRegisterNatives%s()\r\n"), *ShortPackageName, *StructNameCPP);
 			Out.Logf(TEXT("\t{\r\n"));
 
-			Out.Logf(TEXT("\t\tUScriptStruct::DeferCppStructOps(FName(TEXT(\"%s\")),new UScriptStruct::TCppStructOps<%s>);\r\n"), *ActualStructName, StructNameCPP);
+			Out.Logf(TEXT("\t\tUScriptStruct::DeferCppStructOps(FName(TEXT(\"%s\")),new UScriptStruct::TCppStructOps<%s>);\r\n"), *ActualStructName, *StructNameCPP);
 
 			Out.Logf(TEXT("\t}\r\n"));
-			Out.Logf(TEXT("} ScriptStruct_%s_StaticRegisterNatives%s;\r\n"), *ShortPackageName, StructNameCPP);
+			Out.Logf(TEXT("} ScriptStruct_%s_StaticRegisterNatives%s;\r\n"), *ShortPackageName, *StructNameCPP);
 		}
 	}
 
@@ -3845,7 +3839,7 @@ void FNativeClassHeaderGenerator::ExportGeneratedStructBodyMacros(FOutputDevice&
 
 		StaticDefinitions.Logf(TEXT("\tvoid* %s::NewStructOps()\r\n"), *StaticsStructName);
 		StaticDefinitions.Log (TEXT("\t{\r\n"));
-		StaticDefinitions.Logf(TEXT("\t\treturn (UScriptStruct::ICppStructOps*)new UScriptStruct::TCppStructOps<%s>();\r\n"), StructNameCPP);
+		StaticDefinitions.Logf(TEXT("\t\treturn (UScriptStruct::ICppStructOps*)new UScriptStruct::TCppStructOps<%s>();\r\n"), *StructNameCPP);
 		StaticDefinitions.Log (TEXT("\t}\r\n"));
 
 		NewStructOps = TEXT("&NewStructOps");
@@ -3864,8 +3858,8 @@ void FNativeClassHeaderGenerator::ExportGeneratedStructBodyMacros(FOutputDevice&
 	StaticDefinitions.Logf(TEXT("\t\t%s,\r\n"), *GetSingletonNameFuncAddr(BaseStruct, OutReferenceGatherers.UniqueCrossModuleReferences));
 	StaticDefinitions.Logf(TEXT("\t\t%s,\r\n"), *NewStructOps);
 	StaticDefinitions.Logf(TEXT("\t\t%s,\r\n"), *CreateUTF8LiteralString(ActualStructName));
-	StaticDefinitions.Logf(TEXT("\t\tsizeof(%s),\r\n"), StructNameCPP);
-	StaticDefinitions.Logf(TEXT("\t\talignof(%s),\r\n"), StructNameCPP);
+	StaticDefinitions.Logf(TEXT("\t\tsizeof(%s),\r\n"), *StructNameCPP);
+	StaticDefinitions.Logf(TEXT("\t\talignof(%s),\r\n"), *StructNameCPP);
 	StaticDefinitions.Logf(TEXT("\t\t%s,\r\n"), *PropertyRange.Get<0>());
 	StaticDefinitions.Logf(TEXT("\t\t%s,\r\n"), *PropertyRange.Get<1>());
 	StaticDefinitions.Logf(TEXT("\t\t%s,\r\n"), bIsDynamic ? TEXT("RF_Public|RF_Transient") : TEXT("RF_Public|RF_Transient|RF_MarkAsNative"));
@@ -3883,7 +3877,7 @@ void FNativeClassHeaderGenerator::ExportGeneratedStructBodyMacros(FOutputDevice&
 	FString NoExportStructNameCPP;
 	if (NoExportStructs.Contains(Struct))
 	{
-		NoExportStructNameCPP = FString::Printf(TEXT("%s::%s"), *StaticsStructName, StructNameCPP);
+		NoExportStructNameCPP = FString::Printf(TEXT("%s::%s"), *StaticsStructName, *StructNameCPP);
 	}
 	else
 	{
@@ -3938,7 +3932,7 @@ void FNativeClassHeaderGenerator::ExportGeneratedStructBodyMacros(FOutputDevice&
 			FString ParameterSuffix = MethodInfo.Parameters.Names(true, TEXT(",\r\n\t\t"));
 
 			// implement the virtual function body.
-			Out.Logf(TEXT("%s %s::%s(%s)\r\n"), *MethodInfo.ReturnType, StructNameCPP, *MethodInfo.Name, *ParameterDeclaration);
+			Out.Logf(TEXT("%s %s::%s(%s)\r\n"), *MethodInfo.ReturnType, *StructNameCPP, *MethodInfo.Name, *ParameterDeclaration);
 			Out.Log(TEXT("{\r\n"));
 
 			if(RigVMVirtualFuncProlog.Num() > 0)
@@ -3956,7 +3950,7 @@ void FNativeClassHeaderGenerator::ExportGeneratedStructBodyMacros(FOutputDevice&
 			Out.Log(TEXT("\r\n"));
 
 			// implement stub method body
-			Out.Logf(TEXT("%s %s::RigVM%s(const FRigVMOperandArray& RigVMOperands, FRigVMMemoryContainerPtrArray& RigVMMemoryContainer, const FRigVMUserDataArray& RigVMUserData)\r\n"), *MethodInfo.ReturnType, StructNameCPP, *MethodInfo.Name);
+			Out.Logf(TEXT("%s %s::RigVM%s(const FRigVMOperandArray& RigVMOperands, FRigVMMemoryContainerPtrArray& RigVMMemoryContainer, const FRigVMUserDataArray& RigVMUserData)\r\n"), *MethodInfo.ReturnType, *StructNameCPP, *MethodInfo.Name);
 			Out.Log(TEXT("{\r\n"));
 
 			if (MethodInfo.Parameters.Num() > 0)
@@ -4128,7 +4122,6 @@ void FNativeClassHeaderGenerator::ExportGeneratedEnumInitCode(FOutputDevice& Out
 	GeneratedEnumRegisterFunctionText.Logf(TEXT("\t}\r\n"));
 
 	uint32 EnumHash = GenerateTextHash(*GeneratedEnumRegisterFunctionText);
-	AddGeneratedCodeHash(Enum, EnumHash);
 	Out.Logf(TEXT("\tuint32 %s() { return %uU; }\r\n"), *HashFuncName, EnumHash);
 	Out.Log(GeneratedEnumRegisterFunctionText);
 }
@@ -4136,11 +4129,11 @@ void FNativeClassHeaderGenerator::ExportGeneratedEnumInitCode(FOutputDevice& Out
 void FNativeClassHeaderGenerator::ExportMirrorsForNoexportStruct(FOutputDevice& Out, UScriptStruct* Struct, int32 TextIndent)
 {
 	// Export struct.
-	const TCHAR* StructName = NameLookupCPP.GetNameCPP(Struct);
-	Out.Logf(TEXT("%sstruct %s"), FCString::Tab(TextIndent), StructName);
+	const FString StructName = FNameLookupCPP::GetNameCPP(Struct);
+	Out.Logf(TEXT("%sstruct %s"), FCString::Tab(TextIndent), *StructName);
 	if (Struct->GetSuperStruct() != NULL)
 	{
-		Out.Logf(TEXT(" : public %s"), NameLookupCPP.GetNameCPP(Struct->GetSuperStruct()));
+		Out.Logf(TEXT(" : public %s"), *FNameLookupCPP::GetNameCPP(Struct->GetSuperStruct()));
 	}
 	Out.Logf(TEXT("\r\n%s{\r\n"), FCString::Tab(TextIndent));
 
@@ -4736,7 +4729,7 @@ void FNativeClassHeaderGenerator::ExportNativeFunctionHeader(
 	FString FunctionName;
 	if (FunctionHeaderStyle == EExportFunctionHeaderStyle::Definition)
 	{
-		FunctionName = FString::Printf(TEXT("%s::"), NameLookupCPP.GetNameCPP(CastChecked<UClass>(Function->GetOuter()), bIsInterface || FunctionType == EExportFunctionType::Interface));
+		FunctionName = FString::Printf(TEXT("%s::"), *FNameLookupCPP::GetNameCPP(CastChecked<UClass>(Function->GetOuter()), bIsInterface || FunctionType == EExportFunctionType::Interface));
 	}
 
 	if (FunctionType == EExportFunctionType::Interface)
@@ -5037,7 +5030,7 @@ void FNativeClassHeaderGenerator::ExportFunctionThunk(FUHTStringBuilder& RPCWrap
 	// export the call to the C++ version
 	if (FunctionData.FunctionExportFlags & FUNCEXPORT_CppStatic)
 	{
-		RPCWrappers.Logf(TEXT("%s::%s(%s);") LINE_TERMINATOR, NameLookupCPP.GetNameCPP(Function->GetOwnerClass()), *FunctionData.CppImplName, *ParameterList);
+		RPCWrappers.Logf(TEXT("%s::%s(%s);") LINE_TERMINATOR, *FNameLookupCPP::GetNameCPP(Function->GetOwnerClass()), *FunctionData.CppImplName, *ParameterList);
 	}
 	else
 	{
@@ -5091,7 +5084,7 @@ void FNativeClassHeaderGenerator::ExportNativeFunctions(FOutputDevice& OutGenera
 	FNativeFunctionStringBuilder RuntimeStringBuilders;
 	FNativeFunctionStringBuilder EditorStringBuilders;
 
-	const TCHAR* ClassCPPName = NameLookupCPP.GetNameCPP(Class, Class->HasAnyClassFlags(CLASS_Interface));
+	const FString ClassCPPName = FNameLookupCPP::GetNameCPP(Class, Class->HasAnyClassFlags(CLASS_Interface));
 
 	ClassDefinitionRange ClassRange;
 	if (ClassDefinitionRanges.Contains(Class))
@@ -5261,7 +5254,7 @@ void FNativeClassHeaderGenerator::ExportNativeFunctions(FOutputDevice& OutGenera
 
 		// export the script wrappers
 		FuncStringBuilders.RPCWrappers.Logf(TEXT("\tDECLARE_FUNCTION(%s);"), *FunctionData.UnMarshallAndCallName);
-		FuncStringBuilders.RPCImplementations.Logf(TEXT("\tDEFINE_FUNCTION(%s::%s)"), ClassCPPName, *FunctionData.UnMarshallAndCallName);
+		FuncStringBuilders.RPCImplementations.Logf(TEXT("\tDEFINE_FUNCTION(%s::%s)"), *ClassCPPName, *FunctionData.UnMarshallAndCallName);
 		FuncStringBuilders.RPCImplementations += LINE_TERMINATOR TEXT("\t{") LINE_TERMINATOR;
 
 		FParmsAndReturnProperties Parameters = GetFunctionParmsAndReturn(FunctionData.FunctionReference);
@@ -5380,7 +5373,7 @@ void FNativeClassHeaderGenerator::ExportCallbackFunctions(
 		const FFuncInfo& FunctionData = CompilerInfo->GetFunctionData();
 		FString          FunctionName = Function->GetName();
 		UClass*          Class        = CastChecked<UClass>(Function->GetOuter());
-		const TCHAR*     ClassName    = NameLookupCPP.GetNameCPP(Class);
+		const FString    ClassName    = FNameLookupCPP::GetNameCPP(Class);
 
 		if (FunctionData.FunctionFlags & FUNC_NetResponse)
 		{
@@ -5403,7 +5396,7 @@ void FNativeClassHeaderGenerator::ExportCallbackFunctions(
 		FString FunctionNameName;
 		if (ExportCallbackType != EExportCallbackType::Interface)
 		{
-			FunctionNameName = FString::Printf(TEXT("NAME_%s_%s"), ClassName, *FunctionName);
+			FunctionNameName = FString::Printf(TEXT("NAME_%s_%s"), *ClassName, *FunctionName);
 			OutCpp.Logf(TEXT("\tstatic FName %s = FName(TEXT(\"%s\"));") LINE_TERMINATOR, *FunctionNameName, *GetOverriddenFName(Function).ToString());
 		}
 
@@ -5419,7 +5412,7 @@ void FNativeClassHeaderGenerator::ExportCallbackFunctions(
 				// Cast away const just in case, because ProcessEvent isn't const
 				OutCpp.Logf(
 					TEXT("\t\t%sProcessEvent(FindFunctionChecked(%s),%s);\r\n"),
-					(Function->HasAllFunctionFlags(FUNC_Const)) ? *FString::Printf(TEXT("const_cast<%s*>(this)->"), ClassName) : TEXT(""),
+					(Function->HasAllFunctionFlags(FUNC_Const)) ? *FString::Printf(TEXT("const_cast<%s*>(this)->"), *ClassName) : TEXT(""),
 					*FunctionNameName,
 					Parameters.HasParms() ? TEXT("&Parms") : TEXT("NULL")
 				);
@@ -5699,8 +5692,6 @@ FNativeClassHeaderGenerator::FNativeClassHeaderGenerator(
 
 		FUHTStringBuilderLineCounter& OutText = GeneratedCPP.GeneratedText;
 
-		NameLookupCPP.SetCurrentSourceFile(SourceFile);
-
 		TArray<UEnum*> Enums;
 		TArray<UScriptStruct*> Structs;
 		TArray<UDelegateFunction*> DelegateFunctions;
@@ -6797,7 +6788,6 @@ UClass* ProcessParsedClass(bool bClassIsAnInterface, TArray<FHeaderProvider>& De
 
 		// Create new class.
 		ResultClass = new(EC_InternalUseOnlyConstructor, InParent, *ClassNameStripped, Flags) UClass(FObjectInitializer(), nullptr);
-		GClassHeaderNameWithNoPathMap.Add(ResultClass, ClassNameStripped);
 
 		// add CLASS_Interface flag if the class is an interface
 		// NOTE: at this pre-parsing/importing stage, we cannot know if our super class is an interface or not,
diff --git a/Engine/Source/Programs/UnrealHeaderTool/Private/HeaderParser.cpp b/Engine/Source/Programs/UnrealHeaderTool/Private/HeaderParser.cpp
--- a/Engine/Source/Programs/UnrealHeaderTool/Private/HeaderParser.cpp
+++ b/Engine/Source/Programs/UnrealHeaderTool/Private/HeaderParser.cpp
@@ -963,9 +963,9 @@ namespace
 	 * @param	InKey			the metadata key being added
 	 * @param	InValue			the value string that will be associated with the InKey
 	 */
-	void ValidateMetaDataFormat(FFieldVariant Field, const FString& InKey, const FString& InValue)
+	void ValidateMetaDataFormat(FFieldVariant Field, const FName InKey, const FString& InValue)
 	{
-		switch ((ECheckedMetadataSpecifier)Algo::FindSortedStringCaseInsensitive(*InKey, GCheckedMetadataSpecifierStrings))
+		switch (GetCheckedMetadataSpecifier(InKey))
 		{
 			default:
 			{
@@ -980,7 +980,7 @@ namespace
 			{
 				if (!InValue.IsNumeric())
 				{
-					FError::Throwf(TEXT("Metadata value for '%s' is non-numeric : '%s'"), *InKey, *InValue);
+					FError::Throwf(TEXT("Metadata value for '%s' is non-numeric : '%s'"), *InKey.ToString(), *InValue);
 				}
 			}
 			break;
@@ -1001,7 +1001,7 @@ namespace
 
 						if (Class != nullptr && Class->GetName() == TEXT("BlueprintFunctionLibrary"))
 						{
-							FError::Throwf(TEXT("%s doesn't make sense on static method '%s' in a blueprint function library"), *InKey, *Function->GetName());
+							FError::Throwf(TEXT("%s doesn't make sense on static method '%s' in a blueprint function library"), *InKey.ToString(), *Function->GetName());
 						}
 					}
 				}
@@ -1118,7 +1118,7 @@ namespace
 				const FString ExperimentalValue(TEXT("Experimental"));
 				if ((InValue != EarlyAccessValue) && (InValue != ExperimentalValue))
 				{
-					FError::Throwf(TEXT("'%s' metadata was '%s' but it must be %s or %s"), *InKey, *InValue, *ExperimentalValue, *EarlyAccessValue);
+					FError::Throwf(TEXT("'%s' metadata was '%s' but it must be %s or %s"), *InKey.ToString(), *InValue, *ExperimentalValue, *EarlyAccessValue);
 				}
 			}
 			break;
@@ -1144,7 +1144,7 @@ namespace
 				const TCHAR* StrictValue = TEXT("Strict");
 				if (InValue != StrictValue)
 				{
-					FError::Throwf(TEXT("'%s' metadata was '%s' but it must be %s"), *InKey, *InValue, *StrictValue);
+					FError::Throwf(TEXT("'%s' metadata was '%s' but it must be %s"), *InKey.ToString(), *InValue, *StrictValue);
 				}
 			}
 			break;
@@ -1154,9 +1154,9 @@ namespace
 	// Ensures at script compile time that the metadata formatting is correct
 	void ValidateMetaDataFormat(FFieldVariant Field, const TMap<FName, FString>& MetaData)
 	{
-		for (const auto& Pair : MetaData)
+		for (const TPair<FName, FString>& Pair : MetaData)
 		{
-			ValidateMetaDataFormat(Field, Pair.Key.ToString(), Pair.Value);
+			ValidateMetaDataFormat(Field, Pair.Key, Pair.Value);
 		}
 	}
 
@@ -5662,7 +5662,7 @@ bool FHeaderParser::CompileDeclaration(FClasses& AllClasses, TArray<UDelegateFun
 				bSkippedAPIToken = true;
 			}
 
-			if (ConstructorToken.Matches(NameLookupCPP.GetNameCPP(Class), ESearchCase::IgnoreCase))
+			if (ConstructorToken.Matches(*FNameLookupCPP::GetNameCPP(Class), ESearchCase::IgnoreCase))
 			{
 				if (TryToMatchConstructorParameterList(ConstructorToken))
 				{
@@ -8109,7 +8109,6 @@ void FHeaderParser::FinalizeScriptExposedFunctions(UClass* Class)
 ECompilationResult::Type FHeaderParser::ParseHeader(FClasses& AllClasses, FUnrealSourceFile* SourceFile)
 {
 	SetCurrentSourceFile(SourceFile);
-	NameLookupCPP.SetCurrentSourceFile(SourceFile);
 	FUnrealSourceFile* CurrentSrcFile = SourceFile;
 	if (CurrentSrcFile->IsParsed())
 	{
diff --git a/Engine/Source/Programs/UnrealHeaderTool/Private/ParserHelper.cpp b/Engine/Source/Programs/UnrealHeaderTool/Private/ParserHelper.cpp
--- a/Engine/Source/Programs/UnrealHeaderTool/Private/ParserHelper.cpp
+++ b/Engine/Source/Programs/UnrealHeaderTool/Private/ParserHelper.cpp
@@ -281,36 +281,4 @@ FTokenData* FPropertyData::Set(FProperty* InKey, const FTokenData& InValue, FUnr
 	}
 
 	return Result;
-}
-
-const TCHAR* FNameLookupCPP::GetNameCPP(UStruct* Struct, bool bForceInterface /*= false */)
-{
-	TMap<UStruct*, TCHAR*>& NameMap = (bForceInterface ? InterfaceNameMap : StructNameMap);
-	FRWScopeLock Lock(bForceInterface ? InterfaceNameLock : StructNameLock, SLT_ReadOnly);
-
-	TCHAR* NameCPP = NameMap.FindRef(Struct);
-	if (NameCPP == nullptr)
-	{
-		const FString TempName = FString::Printf(TEXT("%s%s"), (bForceInterface ? TEXT("I") : Struct->GetPrefixCPP()), *Struct->GetName());
-		const int32 StringLength = TempName.Len();
-
-		NameCPP = new TCHAR[StringLength + 1];
-		FCString::Strcpy(NameCPP, StringLength + 1, *TempName);
-		NameCPP[StringLength] = 0;
-
-		Lock.ReleaseReadOnlyLockAndAcquireWriteLock_USE_WITH_CAUTION();
-
-		// Check the map again in case another thread had also been waiting on writing this data and got the write lock first
-		if (TCHAR* NameInMap = NameMap.FindRef(Struct))
-		{
-			delete NameCPP;
-			NameCPP = NameInMap;
-		}
-		else
-		{
-			NameMap.Add(Struct, NameCPP);
-		}
-	}
-
-	return NameCPP;
-}
+}
\ No newline at end of file
diff --git a/Engine/Source/Programs/UnrealHeaderTool/Private/ParserHelper.h b/Engine/Source/Programs/UnrealHeaderTool/Private/ParserHelper.h
--- a/Engine/Source/Programs/UnrealHeaderTool/Private/ParserHelper.h
+++ b/Engine/Source/Programs/UnrealHeaderTool/Private/ParserHelper.h
@@ -1887,51 +1887,18 @@ struct FScriptLocation
  */
 struct FNameLookupCPP
 {
-	/**
-	 * Destructor, cleaning up allocated memory.
-	 */
-	~FNameLookupCPP()
-	{
-		for (TMap<UStruct*,TCHAR*>::TIterator It(StructNameMap); It; ++It)
-		{
-			TCHAR* Name = It.Value();
-			delete [] Name;
-		}
-	
-		for (TMap<UStruct*, TCHAR*>::TIterator It(InterfaceNameMap); It; ++It)
-		{
-			TCHAR* Name = It.Value();
-			delete [] Name;
-		}
-	}
-
 	/**
 	 * Returns the name used for declaring the passed in struct in C++
 	 *
 	 * @param	Struct	UStruct to obtain C++ name for
 	 * @return	Name used for C++ declaration
 	 */
-	const TCHAR* GetNameCPP( UStruct* Struct, bool bForceInterface = false );
-
-	void SetCurrentSourceFile(FUnrealSourceFile* InUnrealSourceFile)
+	static FString GetNameCPP(UStruct* Struct, bool bForceInterface = false)
 	{
-		UnrealSourceFile = InUnrealSourceFile;
+		return FString::Printf(TEXT("%s%s"), (bForceInterface ? TEXT("I") : Struct->GetPrefixCPP()), *Struct->GetName());
 	}
-
-private:
-	/** Map of UStruct pointers to C++ names */
-	TMap<UStruct*,TCHAR*> StructNameMap;
-	TMap<UStruct*, TCHAR*> InterfaceNameMap;
-
-	FRWLock InterfaceNameLock;
-	FRWLock StructNameLock;
-
-	FUnrealSourceFile* UnrealSourceFile;
 };
 
-extern FNameLookupCPP NameLookupCPP;
-
-
 /////////////////////////////////////////////////////
 // FAdvancedDisplayParameterHandler - used by FHeaderParser::ParseParameterList, to check if a property if a function parameter has 'AdvancedDisplay' flag
 
diff --git a/Engine/Source/Programs/UnrealHeaderTool/Private/Specifiers/CheckedMetadataSpecifiers.h b/Engine/Source/Programs/UnrealHeaderTool/Private/Specifiers/CheckedMetadataSpecifiers.h
--- a/Engine/Source/Programs/UnrealHeaderTool/Private/Specifiers/CheckedMetadataSpecifiers.h
+++ b/Engine/Source/Programs/UnrealHeaderTool/Private/Specifiers/CheckedMetadataSpecifiers.h
@@ -15,4 +15,6 @@ enum class ECheckedMetadataSpecifier
 	, Max
 };
 
-extern const TCHAR* GCheckedMetadataSpecifierStrings[(int32)ECheckedMetadataSpecifier::Max];
+extern const TMap<FName, ECheckedMetadataSpecifier> GCheckedMetadataSpecifiers;
+
+ECheckedMetadataSpecifier GetCheckedMetadataSpecifier(FName Key);
\ No newline at end of file
diff --git a/Engine/Source/Programs/UnrealHeaderTool/Private/Specifiers/SpecifierArrays.cpp b/Engine/Source/Programs/UnrealHeaderTool/Private/Specifiers/SpecifierArrays.cpp
--- a/Engine/Source/Programs/UnrealHeaderTool/Private/Specifiers/SpecifierArrays.cpp
+++ b/Engine/Source/Programs/UnrealHeaderTool/Private/Specifiers/SpecifierArrays.cpp
@@ -1,7 +1,7 @@
 // Copyright Epic Games, Inc. All Rights Reserved.
 
-#include "UnrealHeaderTool.h"
 #include "CoreMinimal.h"
+#include "UnrealHeaderTool.h"
 #include "Misc/AssertionMacros.h"
 #include "Misc/CString.h"
 
@@ -41,12 +41,12 @@ const TCHAR* GVariableSpecifierStrings[(int32)EVariableSpecifier::Max] =
 	#undef VARIABLE_SPECIFIER
 };
 
-const TCHAR* GCheckedMetadataSpecifierStrings[(int32)ECheckedMetadataSpecifier::Max] =
+const TMap<FName, ECheckedMetadataSpecifier> GCheckedMetadataSpecifiers(
 {
-	#define CHECKED_METADATA_SPECIFIER(SpecifierName) TEXT(#SpecifierName),
+	#define CHECKED_METADATA_SPECIFIER(SpecifierName) { FName(#SpecifierName), ECheckedMetadataSpecifier::SpecifierName },
 		#include "CheckedMetadataSpecifiers.def"
 	#undef CHECKED_METADATA_SPECIFIER
-};
+});
 
 const TCHAR* GClassMetadataSpecifierStrings[(int32)EClassMetadataSpecifier::Max] =
 {
@@ -55,6 +55,15 @@ const TCHAR* GClassMetadataSpecifierStrings[(int32)EClassMetadataSpecifier::Max]
 	#undef CLASS_METADATA_SPECIFIER
 };
 
+ECheckedMetadataSpecifier GetCheckedMetadataSpecifier(FName Key)
+{
+	if (const ECheckedMetadataSpecifier* Specifier = GCheckedMetadataSpecifiers.Find(Key))
+	{
+		return *Specifier;
+	}
+	return ECheckedMetadataSpecifier::Max;
+}
+
 struct FCStringsLessThanCaseInsensitive
 {
 	FORCEINLINE bool operator()(const TCHAR* Lhs, const TCHAR* Rhs)
@@ -67,5 +76,4 @@ const bool GIsGFunctionSpecifierStringsSorted        = ensure(Algo::IsSorted(GFu
 const bool GIsGStructSpecifierStringsSorted          = ensure(Algo::IsSorted(GStructSpecifierStrings,          FCStringsLessThanCaseInsensitive()));
 const bool GIsGInterfaceSpecifierStringsSorted       = ensure(Algo::IsSorted(GInterfaceSpecifierStrings,       FCStringsLessThanCaseInsensitive()));
 const bool GIsGVariableSpecifierStringsSorted        = ensure(Algo::IsSorted(GVariableSpecifierStrings,        FCStringsLessThanCaseInsensitive()));
-const bool GIsGCheckedMetadataSpecifierStringsSorted = ensure(Algo::IsSorted(GCheckedMetadataSpecifierStrings, FCStringsLessThanCaseInsensitive()));
-const bool GIsGClassMetadataSpecifierStringsSorted   = ensure(Algo::IsSorted(GClassMetadataSpecifierStrings,   FCStringsLessThanCaseInsensitive()));
+const bool GIsGClassMetadataSpecifierStringsSorted   = ensure(Algo::IsSorted(GClassMetadataSpecifierStrings,   FCStringsLessThanCaseInsensitive()));
\ No newline at end of file
