GitHub #6052 : Fix comments for EBuildModuleType

diff --git a/Engine/Source/Programs/UnrealHeaderTool/Public/IScriptGeneratorPluginInterface.h b/Engine/Source/Programs/UnrealHeaderTool/Public/IScriptGeneratorPluginInterface.h
--- a/Engine/Source/Programs/UnrealHeaderTool/Public/IScriptGeneratorPluginInterface.h
+++ b/Engine/Source/Programs/UnrealHeaderTool/Public/IScriptGeneratorPluginInterface.h
@@ -9,7 +9,7 @@
 class FArchive;
 class UClass;
 
-/** Build module type, mirrored in UEBuildModule.cs, enum UEBuildModuletype */
+/** Build module type, mirrored in ExternalExecution.cs, enum UHTModuleType */
 struct EBuildModuleType
 {
 	enum Type
