Fix crash throwing error based on a required TCHAR symbol by using %c instead of %s in the string format.
#rnx

diff --git a/Engine/Source/Programs/UnrealHeaderTool/Private/BaseParser.cpp b/Engine/Source/Programs/UnrealHeaderTool/Private/BaseParser.cpp
--- a/Engine/Source/Programs/UnrealHeaderTool/Private/BaseParser.cpp
+++ b/Engine/Source/Programs/UnrealHeaderTool/Private/BaseParser.cpp
@@ -960,7 +960,7 @@ void FBaseParser::RequireSymbol( const TCHAR Match, const TCHAR* Tag, ESymbolPar
 {
 	if (!MatchSymbol(Match, bParseTemplateClosingBracket))
 	{
-		FError::Throwf(TEXT("Missing '%s' in %s"), Match, Tag );
+		FError::Throwf(TEXT("Missing '%c' in %s"), Match, Tag );
 	}
 }
 
@@ -968,7 +968,7 @@ void FBaseParser::RequireSymbol(const TCHAR Match, TFunctionRef<FString()> Error
 {
 	if (!MatchSymbol(Match, bParseTemplateClosingBracket))
 	{
-		FError::Throwf(TEXT("Missing '%s' in %s"), Match, *ErrorGetter());
+		FError::Throwf(TEXT("Missing '%c' in %s"), Match, *ErrorGetter());
 	}
 }
 
