Fix for nonunity compile error.

diff --git a/Engine/Source/Programs/UnrealHeaderTool/Private/UnrealHeaderTool.h b/Engine/Source/Programs/UnrealHeaderTool/Private/UnrealHeaderTool.h
--- a/Engine/Source/Programs/UnrealHeaderTool/Private/UnrealHeaderTool.h
+++ b/Engine/Source/Programs/UnrealHeaderTool/Private/UnrealHeaderTool.h
@@ -1,6 +1,6 @@
 // Copyright Epic Games, Inc. All Rights Reserved.
-
 #pragma once
+#include "CoreMinimal.h"
 
 struct FHeaderParserNames
 {
