Make sure Replicate Data is set up in UHT outside of the parallel loops.
This should prevent race conditions from potentially modifying classes from multiple threads.

diff --git a/Engine/Source/Programs/UnrealHeaderTool/Private/CodeGenerator.cpp b/Engine/Source/Programs/UnrealHeaderTool/Private/CodeGenerator.cpp
--- a/Engine/Source/Programs/UnrealHeaderTool/Private/CodeGenerator.cpp
+++ b/Engine/Source/Programs/UnrealHeaderTool/Private/CodeGenerator.cpp
@@ -2529,8 +2529,6 @@ void FNativeClassHeaderGenerator::ExportNativeGeneratedInitCode(FOutputDevice& O
 
 	if (ClassHasReplicatedProperties(Class))
 	{
-		Class->SetUpUhtReplicationData(nullptr);
-
 		Out.Logf(TEXT(
 			"\r\n"
 			"\tvoid %s::ValidateGeneratedRepEnums(const TArray<struct FRepRecord>& ClassReps) const\r\n"
@@ -5684,6 +5682,17 @@ FNativeClassHeaderGenerator::FNativeClassHeaderGenerator(
 
 		FGeneratedCPP& GeneratedCPP = GeneratedCPPs.Emplace(SourceFile, MoveTemp(GeneratedSourceFilename));
 		GeneratedCPP.RelativeIncludes.Add(MoveTemp(ModuleRelativeFilename));
+
+		// This needs to be done outside of parallel blocks because it will modify UClass memory.
+		// Later calls to SetUpUhtReplicationData inside parallel blocks should be fine, because
+		// they will see the memory has already been set up, and just return the parent pointer.
+		for (UClass* Class : SourceFile->GetDefinedClasses())
+		{
+			if (ClassHasReplicatedProperties(Class))
+			{
+				Class->SetUpUhtReplicationData(nullptr);
+			}
+		}
 	}
 
 	const FManifestModule* ConstPackageManifest = GetPackageManifest(PackageName);
