UnrealHeader Tool Feature: Static-Virtual-Methods on UStruct

Within Control Rig we use a static method to describe the math to perform a single node - yet we store the data on a struct. To simplify the implementation we add a virtual method and have the reflection system generate a static method providing access to all struct members as parameters.

diff --git a/Engine/Source/Programs/UnrealHeaderTool/Private/CodeGenerator.cpp b/Engine/Source/Programs/UnrealHeaderTool/Private/CodeGenerator.cpp
--- a/Engine/Source/Programs/UnrealHeaderTool/Private/CodeGenerator.cpp
+++ b/Engine/Source/Programs/UnrealHeaderTool/Private/CodeGenerator.cpp
@@ -3299,6 +3299,63 @@ void FNativeClassHeaderGenerator::ExportGeneratedStructBodyMacros(FOutputDevice&
 
 	const FString SingletonName = GetSingletonName(Struct);
 
+	FString StaticVirtualMacroPrefix(StructNameCPP);
+	StaticVirtualMacroPrefix = FString::Printf(TEXT("UE_%s"), *StaticVirtualMacroPrefix.Mid(1));
+
+	const TArray<FStaticVirtualMethodInfo>* StructStaticVirtualMethods = FHeaderParser::StructStaticVirtualMethods.Find(Struct);
+	if(StructStaticVirtualMethods)
+	{
+		const TCHAR* Comma = TEXT(", ");
+		TArray<FString> PropertyNames, PropertyArguments;
+		for (UProperty* Prop : TFieldRange<UProperty>(Struct))
+		{
+			bool bIsInput = Prop->HasMetaData(TEXT("Input"));
+			bool bIsOutput = Prop->HasMetaData(TEXT("Output"));
+			bool bIsConstant = Prop->HasMetaData(TEXT("Constant"));
+
+			FString Name = Prop->GetName();
+			PropertyNames.Add(Name);
+
+			FString CPPType;
+			FString ExtendedCPPType;
+			CPPType = Prop->GetCPPType(&ExtendedCPPType);
+			CPPType += ExtendedCPPType;
+
+			if (CPPType.StartsWith(TEXT("F"), ESearchCase::CaseSensitive) ||
+				CPPType.StartsWith(TEXT("T"), ESearchCase::CaseSensitive) ||
+				CPPType.StartsWith(TEXT("E"), ESearchCase::CaseSensitive))
+			{
+				if (bIsConstant || (bIsInput && !bIsOutput))
+				{
+					PropertyArguments.Add(FString::Printf(TEXT("const %s& %s"), *CPPType, *Name));
+				}
+				else
+				{
+					PropertyArguments.Add(FString::Printf(TEXT("%s& %s"), *CPPType, *Name));
+				}
+			}
+			else if (bIsConstant || (bIsInput && !bIsOutput))
+			{
+				PropertyArguments.Add(FString::Printf(TEXT("const %s %s"), *CPPType, *Name));
+			}
+			else
+			{
+				PropertyArguments.Add(FString::Printf(TEXT("%s& %s"), *CPPType, *Name));
+			}
+		}
+
+		OutGeneratedHeaderText.Log(TEXT("\n"));
+		OutGeneratedHeaderText.Logf(TEXT("#define %s_MEMBER_NAMES %s\n"), *StaticVirtualMacroPrefix, *FString::Join(PropertyNames, Comma));
+		OutGeneratedHeaderText.Logf(TEXT("#define %s_MEMBER_ARGUMENTS %s\n"), *StaticVirtualMacroPrefix, *FString::Join(PropertyArguments, Comma));
+		OutGeneratedHeaderText.Logf(TEXT("#define %s_DECLARE_STATIC_VIRTUAL_METHOD(ReturnType, FunctionName, ...) \\\n"), *StaticVirtualMacroPrefix);
+		OutGeneratedHeaderText.Logf(TEXT("    static ReturnType Static##FunctionName(%s_MEMBER_ARGUMENTS, ##__VA_ARGS__);\n"), *StaticVirtualMacroPrefix);
+		OutGeneratedHeaderText.Logf(TEXT("#define %s_IMPLEMENT_STATIC_VIRTUAL_METHOD(ReturnType, FunctionName, ...) \\\n"), *StaticVirtualMacroPrefix);
+		OutGeneratedHeaderText.Logf(TEXT("    ReturnType %s::Static##FunctionName(%s_MEMBER_ARGUMENTS, ##__VA_ARGS__)\n"), StructNameCPP, *StaticVirtualMacroPrefix);
+		OutGeneratedHeaderText.Logf(TEXT("#define %s_INVOKE_STATIC_VIRTUAL_METHOD(FunctionName, ...) \\\n"), *StaticVirtualMacroPrefix);
+		OutGeneratedHeaderText.Logf(TEXT("    %s::Static##FunctionName(%s_MEMBER_NAMES, ##__VA_ARGS__);\n"), StructNameCPP, *StaticVirtualMacroPrefix);
+		OutGeneratedHeaderText.Log(TEXT("\n"));
+	}
+
 	// Export struct.
 	if (Struct->StructFlags & STRUCT_Native)
 	{
@@ -3313,9 +3370,26 @@ void FNativeClassHeaderGenerator::ExportGeneratedStructBodyMacros(FOutputDevice&
 		const FString FriendLine = FString::Printf(TEXT("\tfriend struct %s_Statics;\r\n"), *SingletonName.LeftChop(2));
 		const FString StaticClassLine = FString::Printf(TEXT("\t%sstatic class UScriptStruct* StaticStruct();\r\n"), *RequiredAPI);
 		const FString PrivatePropertiesOffset = PrivatePropertiesOffsetGetters(Struct, StructNameCPP);
+		
+		FString StaticVirtualMethodsDeclarations;
+		if (StructStaticVirtualMethods)
+		{
+			for (const FStaticVirtualMethodInfo& Info : (*StructStaticVirtualMethods))
+			{
+				if (Info.Params.IsEmpty())
+				{
+					StaticVirtualMethodsDeclarations += FString::Printf(TEXT("\t%s_DECLARE_STATIC_VIRTUAL_METHOD(%s, %s)\n"), *StaticVirtualMacroPrefix, *Info.ReturnType, *Info.Name);
+				}
+				else
+				{
+					StaticVirtualMethodsDeclarations += FString::Printf(TEXT("\t%s_DECLARE_STATIC_VIRTUAL_METHOD(%s, %s, %s)\n"), *StaticVirtualMacroPrefix, *Info.ReturnType, *Info.Name, *Info.Params);
+				}
+			}
+		}
+
 		const FString SuperTypedef = BaseStruct ? FString::Printf(TEXT("\ttypedef %s Super;\r\n"), NameLookupCPP.GetNameCPP(BaseStruct)) : FString();
 
-		const FString CombinedLine = FriendLine + StaticClassLine + PrivatePropertiesOffset + SuperTypedef;
+		const FString CombinedLine = FriendLine + StaticClassLine + StaticVirtualMethodsDeclarations + PrivatePropertiesOffset + SuperTypedef;
 		const FString MacroName = SourceFile.GetGeneratedBodyMacroName(Struct->StructMacroDeclaredLineNumber);
 
 		const FString Macroized = Macroize(*MacroName, *CombinedLine);
@@ -3518,6 +3592,62 @@ void FNativeClassHeaderGenerator::ExportGeneratedStructBodyMacros(FOutputDevice&
 
 	Out.Log(GeneratedStructRegisterFunctionText);
 	Out.Logf(TEXT("\tuint32 %s() { return %uU; }\r\n"), *HashFuncName, StructHash);
+
+	if (StructStaticVirtualMethods)
+	{
+		Out.Log(TEXT("\n"));
+		for (const FStaticVirtualMethodInfo& Info : (*StructStaticVirtualMethods))
+		{
+			// implement the virtual function body.
+			Out.Logf(TEXT("%s %s::%s(%s)\n"), *Info.ReturnType, StructNameCPP, *Info.Name, *Info.Params);
+			Out.Log(TEXT("{\n"));
+			if (Info.Params.IsEmpty())
+			{
+				if(Info.ReturnType.IsEmpty() || (Info.ReturnType == TEXT("void")))
+				{
+					Out.Logf(TEXT("    %s_INVOKE_STATIC_VIRTUAL_METHOD(%s);\n"), *StaticVirtualMacroPrefix, *Info.Name);
+				}
+				else
+				{
+					Out.Logf(TEXT("    return %s_INVOKE_STATIC_VIRTUAL_METHOD(%s);\n"), *StaticVirtualMacroPrefix, *Info.Name);
+				}
+			}
+			else
+			{
+				TArray<FString> Params;
+				FString ParamPrev, ParamLeft, ParamRight;
+				ParamPrev = Info.Params;
+				while (ParamPrev.Contains(TEXT(",")))
+				{
+					ParamPrev.Split(TEXT(","), &ParamLeft, &ParamRight);
+					Params.Add(ParamLeft.TrimStartAndEnd());
+					ParamPrev = ParamRight;
+				} 
+				Params.Add(ParamPrev.TrimStartAndEnd());
+
+				for (int32 ParamIndex=0;ParamIndex<Params.Num(); ParamIndex++)
+				{
+					int32 LastSpace = INDEX_NONE;
+					if (Params[ParamIndex].FindLastChar(TCHAR(' '), LastSpace))
+					{
+						Params[ParamIndex] = Params[ParamIndex].Mid(LastSpace + 1);
+					}
+				}
+
+				if(Info.ReturnType.IsEmpty() || (Info.ReturnType == TEXT("void")))
+				{
+					Out.Logf(TEXT("    %s_INVOKE_STATIC_VIRTUAL_METHOD(%s, %s);\n"), *StaticVirtualMacroPrefix, *Info.Name, *FString::Join(Params, TEXT(", ")));
+				}
+				else
+				{
+					Out.Logf(TEXT("    return %s_INVOKE_STATIC_VIRTUAL_METHOD(%s, %s);\n"), *StaticVirtualMacroPrefix, *Info.Name, *FString::Join(Params, TEXT(", ")));
+				}
+			}
+			Out.Log(TEXT("}\n"));
+		}
+
+		Out.Log(TEXT("\n"));
+	}
 }
 
 void FNativeClassHeaderGenerator::ExportGeneratedEnumInitCode(FOutputDevice& Out, const FUnrealSourceFile& SourceFile, UEnum* Enum)
diff --git a/Engine/Source/Programs/UnrealHeaderTool/Private/HeaderParser.cpp b/Engine/Source/Programs/UnrealHeaderTool/Private/HeaderParser.cpp
--- a/Engine/Source/Programs/UnrealHeaderTool/Private/HeaderParser.cpp
+++ b/Engine/Source/Programs/UnrealHeaderTool/Private/HeaderParser.cpp
@@ -72,6 +72,7 @@ static const FName NAME_DocumentationPolicy(TEXT("DocumentationPolicy"));
 EGeneratedCodeVersion FHeaderParser::DefaultGeneratedCodeVersion = EGeneratedCodeVersion::V1;
 TArray<FString> FHeaderParser::StructsWithNoPrefix;
 TArray<FString> FHeaderParser::StructsWithTPrefix;
+TMap<UStruct*, TArray<FStaticVirtualMethodInfo>> FHeaderParser::StructStaticVirtualMethods;
 TArray<FString> FHeaderParser::DelegateParameterCountStrings;
 TMap<FString, FString> FHeaderParser::TypeRedirectMap;
 TArray<FString> FHeaderParser::PropertyCPPTypesRequiringUIRanges = { TEXT("float"), TEXT("double") };
@@ -2494,6 +2495,10 @@ UScriptStruct* FHeaderParser::CompileStructDeclaration(FClasses& AllClasses)
 		{
 			FError::Throwf(TEXT("USTRUCTs cannot contain UFUNCTIONs."));
 		}
+		else if (Token.Matches(TEXT("STATIC_VIRTUAL_METHOD"), ESearchCase::CaseSensitive))
+		{
+			CompileStaticMethodDeclaration(AllClasses, Struct);
+		}
 		else if (Token.Matches(TEXT("GENERATED_USTRUCT_BODY")) || Token.Matches(TEXT("GENERATED_BODY")))
 		{
 			// Match 'GENERATED_USTRUCT_BODY' '(' [StructName] ')' or 'GENERATED_BODY' '(' [StructName] ')'
@@ -6160,6 +6165,79 @@ void FHeaderParser::CompileInterfaceDeclaration(FClasses& AllClasses)
 	PushNest(ENestType::Interface, InterfaceClass);
 }
 
+void FHeaderParser::CompileStaticMethodDeclaration(FClasses& AllClasses, UStruct* Struct)
+{
+	if (!MatchSymbol(TEXT("(")))
+	{
+		FError::Throwf(TEXT("Bad STATICMETHOD definition"));
+	}
+
+	// find the next close brace
+	while (!MatchSymbol(TEXT(")")))
+	{
+		FToken Token;
+		if (!GetToken(Token))
+		{
+			break;
+		}
+	}
+
+	FToken PrefixToken, ReturnTypeToken, NameToken, PostfixToken;
+	if (!GetToken(PrefixToken))
+	{
+		return;
+	}
+
+	if (FString(PrefixToken.Identifier).Equals(TEXT("virtual")))
+	{
+		if (!GetToken(ReturnTypeToken))
+		{
+			return;
+		}
+	}
+	else
+	{
+		ReturnTypeToken = PrefixToken;
+	}
+
+	if (!GetToken(NameToken))
+	{
+		return;
+	}
+
+	if (!MatchSymbol(TEXT("(")))
+	{
+		FError::Throwf(TEXT("Bad STATICMETHOD definition"));
+	}
+
+	TArray<FString> ParamsContent;
+	while (!MatchSymbol(TEXT(")")))
+	{
+		FToken Token;
+		if (!GetToken(Token))
+		{
+			break;
+		}
+		ParamsContent.Add(FString(Token.Identifier));
+	}
+
+	while (!FString(PostfixToken.Identifier).Equals(TEXT(";")))
+	{
+		if (!GetToken(PostfixToken))
+		{
+			return;
+		}
+	}
+
+	FStaticVirtualMethodInfo Info;
+	Info.ReturnType = ReturnTypeToken.Identifier;
+	Info.Name = NameToken.Identifier;
+	Info.Params = FString::Join(ParamsContent, TEXT(" "));
+
+	TArray<FStaticVirtualMethodInfo>& Infos = StructStaticVirtualMethods.FindOrAdd(Struct);
+	Infos.Add(Info);
+}
+
 // Returns true if the token is a dynamic delegate declaration
 bool FHeaderParser::IsValidDelegateDeclaration(const FToken& Token) const
 {
diff --git a/Engine/Source/Programs/UnrealHeaderTool/Private/HeaderParser.h b/Engine/Source/Programs/UnrealHeaderTool/Private/HeaderParser.h
--- a/Engine/Source/Programs/UnrealHeaderTool/Private/HeaderParser.h
+++ b/Engine/Source/Programs/UnrealHeaderTool/Private/HeaderParser.h
@@ -117,6 +117,12 @@ struct FIndexRange
 	int32 Count;
 };
 
+struct FStaticVirtualMethodInfo
+{
+	FString ReturnType;
+	FString Name;
+	FString Params;
+};
 
 struct ClassDefinitionRange
 {
@@ -235,6 +241,7 @@ public:
 
 protected:
 	friend struct FScriptLocation;
+	friend struct FNativeClassHeaderGenerator;
 
 	// For compiling messages and errors.
 	FFeedbackContext* Warn;
@@ -419,6 +426,9 @@ protected:
 	// List of all net service functions with undeclared response functions 
 	TMap<int32, FString> RPCsNeedingHookup;
 
+	// List of all static virtual methods defined on structs
+	static TMap<UStruct*, TArray<FStaticVirtualMethodInfo>> StructStaticVirtualMethods;
+
 	// Constructor.
 	explicit FHeaderParser(FFeedbackContext* InWarn, const FManifestModule& InModule);
 
@@ -562,6 +572,7 @@ protected:
 	void CompileFunctionDeclaration(FClasses& AllClasses);
 	void CompileVariableDeclaration (FClasses& AllClasses, UStruct* Struct);
 	void CompileInterfaceDeclaration(FClasses& AllClasses);
+	void CompileStaticMethodDeclaration(FClasses& AllClasses, UStruct* Struct);
 
 	FClass* ParseInterfaceNameDeclaration(FClasses& AllClasses, FString& DeclaredInterfaceName, FString& RequiredAPIMacroIfPresent);
 	bool TryParseIInterfaceClass(FClasses& AllClasses);
diff --git a/Engine/Source/Programs/UnrealHeaderTool/Resources/UHTDebugging.manifest b/Engine/Source/Programs/UnrealHeaderTool/Resources/UHTDebugging.manifest
--- a/Engine/Source/Programs/UnrealHeaderTool/Resources/UHTDebugging.manifest
+++ b/Engine/Source/Programs/UnrealHeaderTool/Resources/UHTDebugging.manifest
@@ -34,7 +34,8 @@
 				".\\UHTDebugging\\TestInterfaceObject.h",
 				".\\UHTDebugging\\StructOnlyHeader.h",
 				".\\UHTDebugging\\EnumOnlyHeader.h",
-				".\\UHTDebugging\\TestObjectSerializers.h"
+				".\\UHTDebugging\\TestObjectSerializers.h",
+				".\\UHTDebugging\\StaticMethodStructHeader.h"
 			],
 			"PublicHeaders": [
 			],
diff --git a/Engine/Source/Programs/UnrealHeaderTool/Resources/UHTDebugging/StaticMethodStructHeader.h b/Engine/Source/Programs/UnrealHeaderTool/Resources/UHTDebugging/StaticMethodStructHeader.h
new file mode 100644
--- /dev/null
+++ b/Engine/Source/Programs/UnrealHeaderTool/Resources/UHTDebugging/StaticMethodStructHeader.h
@@ -0,0 +1,50 @@
+// Copyright 1998-2019 Epic Games, Inc. All Rights Reserved.
+
+#pragma once
+
+#include "StaticMethodStructHeader.generated.h"
+
+USTRUCT()
+struct FStaticMethodStructBase
+{
+	GENERATED_BODY()
+		
+	UPROPERTY(meta = (Input))
+	float Inherited;
+
+	UPROPERTY(meta = (Output))
+	float InheritedOutput;
+};
+
+USTRUCT()
+struct FStaticMethodStruct : public FStaticMethodStructBase
+{
+	GENERATED_BODY()
+
+	STATIC_VIRTUAL_METHOD()
+	void Clear();
+
+	STATIC_VIRTUAL_METHOD()
+	virtual void Execute(bool bAdditionalFlag, const FString& InString) override;
+
+	STATIC_VIRTUAL_METHOD()
+	void Compute(float TestFloat);
+
+	UPROPERTY(meta = (Input))
+	float A;
+
+	UPROPERTY(meta = (Output))
+	float B;
+
+	UPROPERTY(meta = (Input))
+	FVector C;
+
+	UPROPERTY(meta = (Output))
+	FVector D;
+
+	UPROPERTY(meta = (Input, Output))
+	TArray<FVector> E;
+
+	UPROPERTY()
+	float Cache;
+};
diff --git a/Engine/Source/Runtime/CoreUObject/Public/UObject/ObjectMacros.h b/Engine/Source/Runtime/CoreUObject/Public/UObject/ObjectMacros.h
--- a/Engine/Source/Runtime/CoreUObject/Public/UObject/ObjectMacros.h
+++ b/Engine/Source/Runtime/CoreUObject/Public/UObject/ObjectMacros.h
@@ -581,6 +581,7 @@ struct COREUOBJECT_API FReferencerInformationList
 #define UPARAM(...)
 #define UENUM(...)
 #define UDELEGATE(...)
+#define STATIC_VIRTUAL_METHOD(...)
 
 // This pair of macros is used to help implement GENERATED_BODY() and GENERATED_USTRUCT_BODY()
 #define BODY_MACRO_COMBINE_INNER(A,B,C,D) A##B##C##D
