Control Rig: Adapting static virtual method to use TArrayView for array access - testing on kalman filter nodes.

TArrayView can be used to unmanaged memory as well - right now we are just mapping it to the struct array property, however in the future we might want to map to fragmented memory inside the upcoming VM memory storage.

diff --git a/Engine/Plugins/Experimental/ControlRig/Source/ControlRig/Private/Drawing/ControlRigDrawInterface.cpp b/Engine/Plugins/Experimental/ControlRig/Source/ControlRig/Private/Drawing/ControlRigDrawInterface.cpp
--- a/Engine/Plugins/Experimental/ControlRig/Source/ControlRig/Private/Drawing/ControlRigDrawInterface.cpp
+++ b/Engine/Plugins/Experimental/ControlRig/Source/ControlRig/Private/Drawing/ControlRigDrawInterface.cpp
@@ -10,7 +10,7 @@ void FControlRigDrawInterface::DrawPoint(const FTransform& WorldOffset, const FV
 	DrawInstructions.Add(Instruction);
 }
 
-void FControlRigDrawInterface::DrawPoints(const FTransform& WorldOffset, const TArray<FVector>& Points, float Size, const FLinearColor& Color)
+void FControlRigDrawInterface::DrawPoints(const FTransform& WorldOffset, const TArrayView<FVector>& Points, float Size, const FLinearColor& Color)
 {
 	FDrawIntruction Instruction(EDrawType_Point, Color, Size);
 	for(const FVector& Point : Points)
@@ -28,7 +28,7 @@ void FControlRigDrawInterface::DrawLine(const FTransform& WorldOffset, const FVe
 	DrawInstructions.Add(Instruction);
 }
 
-void FControlRigDrawInterface::DrawLines(const FTransform& WorldOffset, const TArray<FVector>& Positions, const FLinearColor& Color, float Thickness)
+void FControlRigDrawInterface::DrawLines(const FTransform& WorldOffset, const TArrayView<FVector>& Positions, const FLinearColor& Color, float Thickness)
 {
 	FDrawIntruction Instruction(EDrawType_Lines, Color, Thickness);
 	for (const FVector& Point : Positions)
@@ -38,7 +38,7 @@ void FControlRigDrawInterface::DrawLines(const FTransform& WorldOffset, const TA
 	DrawInstructions.Add(Instruction);
 }
 
-void FControlRigDrawInterface::DrawLineStrip(const FTransform& WorldOffset, const TArray<FVector>& Positions, const FLinearColor& Color, float Thickness)
+void FControlRigDrawInterface::DrawLineStrip(const FTransform& WorldOffset, const TArrayView<FVector>& Positions, const FLinearColor& Color, float Thickness)
 {
 	FDrawIntruction Instruction(EDrawType_LineStrip, Color, Thickness);
 	for (const FVector& Point : Positions)
diff --git a/Engine/Plugins/Experimental/ControlRig/Source/ControlRig/Private/Units/Highlevel/Hierarchy/RigUnit_TransformConstraint.cpp b/Engine/Plugins/Experimental/ControlRig/Source/ControlRig/Private/Units/Highlevel/Hierarchy/RigUnit_TransformConstraint.cpp
--- a/Engine/Plugins/Experimental/ControlRig/Source/ControlRig/Private/Units/Highlevel/Hierarchy/RigUnit_TransformConstraint.cpp
+++ b/Engine/Plugins/Experimental/ControlRig/Source/ControlRig/Private/Units/Highlevel/Hierarchy/RigUnit_TransformConstraint.cpp
@@ -93,7 +93,7 @@ UE_RigUnit_TransformConstraint_IMPLEMENT_STATIC_VIRTUAL_METHOD(void, Execute, co
 	}
 }
 
-void FRigUnit_TransformConstraint::AddConstraintData(const TArray<FConstraintTarget>& Targets, ETransformConstraintType ConstraintType, const int32 TargetIndex, const FTransform& SourceTransform, const FTransform& InBaseTransform, TArray<FConstraintData>& OutConstraintData, TMap<int32, int32>& OutConstraintDataToTargets)
+void FRigUnit_TransformConstraint::AddConstraintData(const TArrayView<FConstraintTarget>& Targets, ETransformConstraintType ConstraintType, const int32 TargetIndex, const FTransform& SourceTransform, const FTransform& InBaseTransform, TArray<FConstraintData>& OutConstraintData, TMap<int32, int32>& OutConstraintDataToTargets)
 {
 	const FConstraintTarget& Target = Targets[TargetIndex];
 
diff --git a/Engine/Plugins/Experimental/ControlRig/Source/ControlRig/Private/Units/Highlevel/Hierarchy/RigUnit_TransformConstraint.h b/Engine/Plugins/Experimental/ControlRig/Source/ControlRig/Private/Units/Highlevel/Hierarchy/RigUnit_TransformConstraint.h
--- a/Engine/Plugins/Experimental/ControlRig/Source/ControlRig/Private/Units/Highlevel/Hierarchy/RigUnit_TransformConstraint.h
+++ b/Engine/Plugins/Experimental/ControlRig/Source/ControlRig/Private/Units/Highlevel/Hierarchy/RigUnit_TransformConstraint.h
@@ -66,7 +66,7 @@ struct FRigUnit_TransformConstraint : public FRigUnit_HighlevelBaseMutable
 	TArray<FConstraintTarget> Targets;
 
 private:
-	static void AddConstraintData(const TArray<FConstraintTarget>& Targets, ETransformConstraintType ConstraintType, const int32 TargetIndex, const FTransform& SourceTransform, const FTransform& InBaseTransform, TArray<FConstraintData>& OutConstraintData, TMap<int32, int32>& OutConstraintDataToTargets);
+	static void AddConstraintData(const TArrayView<FConstraintTarget>& Targets, ETransformConstraintType ConstraintType, const int32 TargetIndex, const FTransform& SourceTransform, const FTransform& InBaseTransform, TArray<FConstraintData>& OutConstraintData, TMap<int32, int32>& OutConstraintDataToTargets);
 
 	// note that Targets.Num () != ConstraintData.Num()
 	UPROPERTY(transient)
diff --git a/Engine/Plugins/Experimental/ControlRig/Source/ControlRig/Private/Units/Simulation/RigUnit_Kalman.cpp b/Engine/Plugins/Experimental/ControlRig/Source/ControlRig/Private/Units/Simulation/RigUnit_Kalman.cpp
--- a/Engine/Plugins/Experimental/ControlRig/Source/ControlRig/Private/Units/Simulation/RigUnit_Kalman.cpp
+++ b/Engine/Plugins/Experimental/ControlRig/Source/ControlRig/Private/Units/Simulation/RigUnit_Kalman.cpp
@@ -17,35 +17,38 @@ UE_RigUnit_KalmanFloat_IMPLEMENT_STATIC_VIRTUAL_METHOD(void, Execute, const FRig
 		return;
 	}
 
-	int32 MaxSize = FMath::Clamp<int32>(BufferSize, 1, 512);
 	if (Context.State == EControlRigState::Init)
 	{
-		Buffer.Reset();
-		Buffer.Reserve(MaxSize);
+		for (float& Element : Buffer)
+		{
+			Element = FLT_MAX;
+			LastInsertIndex = -1;
+		}
 	}
 	else
 	{
-		if(Buffer.Num() < MaxSize)
-		{
-			Buffer.Add(Value);
-			LastInsertIndex = 0;
-		}
-		else
+		if(LastInsertIndex == Buffer.Num() - 1)
 		{
-			Buffer[LastInsertIndex++] = Value;
-			if(LastInsertIndex == Buffer.Num())
-			{
-				LastInsertIndex = 0;
-			}
+			LastInsertIndex = -1;
 		}
+		Buffer[++LastInsertIndex] = Value;
 
 		Result = 0.f;
+		int32 NumberValidEntries = 0;
 		for(const float F : Buffer)
 		{
+			if (F == FLT_MAX)
+			{
+				break;
+			}
 			Result += F;
+			NumberValidEntries++;
 		}
 
-		Result = Result / float(Buffer.Num());
+		if (NumberValidEntries > 0)
+		{
+			Result = Result / float(NumberValidEntries);
+		}
 	}
 }
 
@@ -63,35 +66,38 @@ UE_RigUnit_KalmanVector_IMPLEMENT_STATIC_VIRTUAL_METHOD(void, Execute, const FRi
 		return;
 	}
 
-	int32 MaxSize = FMath::Clamp<int32>(BufferSize, 1, 512);
 	if (Context.State == EControlRigState::Init)
 	{
-		Buffer.Reset();
-		Buffer.Reserve(MaxSize);
+		for (FVector& Element : Buffer)
+		{
+			Element.X = FLT_MAX;
+			LastInsertIndex = -1;
+		}
 	}
 	else
 	{
-		if(Buffer.Num() < MaxSize)
+		if(LastInsertIndex == Buffer.Num() - 1)
 		{
-			Buffer.Add(Value);
-			LastInsertIndex = 0;
+			LastInsertIndex = -1;
 		}
-		else
+		Buffer[++LastInsertIndex] = Value;
+
+		Result = FVector::ZeroVector;
+		int32 NumberValidEntries = 0;
+		for (const FVector& F : Buffer)
 		{
-			Buffer[LastInsertIndex++] = Value;
-			if(LastInsertIndex == Buffer.Num())
+			if (F.X == FLT_MAX)
 			{
-				LastInsertIndex = 0;
+				break;
 			}
+			Result += F;
+			NumberValidEntries++;
 		}
 
-		Result = FVector::ZeroVector;
-		for(const FVector& F : Buffer)
+		if (NumberValidEntries > 0)
 		{
-			Result += F;
+			Result = Result / float(NumberValidEntries);
 		}
-
-		Result = Result / float(Buffer.Num());
 	}
 }
 
@@ -109,46 +115,50 @@ UE_RigUnit_KalmanTransform_IMPLEMENT_STATIC_VIRTUAL_METHOD(void, Execute, const
 		return;
 	}
 
-	int32 MaxSize = FMath::Clamp<int32>(BufferSize, 1, 512);
 	if (Context.State == EControlRigState::Init)
 	{
-		Buffer.Reset();
-		Buffer.Reserve(MaxSize);
+		for (FTransform& Element : Buffer)
+		{
+			Element.SetTranslation(FVector(FLT_MAX));
+			LastInsertIndex = -1;
+		}
 	}
 	else
 	{
-		if(Buffer.Num() < MaxSize)
-		{
-			Buffer.Add(Value);
-			LastInsertIndex = 0;
-		}
-		else
+		if(LastInsertIndex == Buffer.Num() - 1)
 		{
-			Buffer[LastInsertIndex++] = Value;
-			if(LastInsertIndex == Buffer.Num())
-			{
-				LastInsertIndex = 0;
-			}
+			LastInsertIndex = -1;
 		}
+		Buffer[++LastInsertIndex] = Value;
 
 		FVector Location = FVector::ZeroVector;
 		FVector AxisX = FVector::ZeroVector;
 		FVector AxisY = FVector::ZeroVector;
 		FVector Scale = FVector::ZeroVector;
 		
+		int32 NumberValidEntries = 0;
+
 		for(const FTransform& F : Buffer)
 		{
+			if (F.GetLocation().X == FLT_MAX)
+			{
+				break;
+			}
 			Location += F.GetLocation();
 			AxisX += F.TransformVectorNoScale(FVector(1.f, 0.f, 0.f));
 			AxisY += F.TransformVectorNoScale(FVector(0.f, 1.f, 0.f));
 			AxisY += F.GetLocation();
 			Scale += F.GetScale3D();
+			NumberValidEntries++;
 		}
 
-		Location = Location / float(Buffer.Num());
-		AxisX = (AxisX / float(Buffer.Num())).GetSafeNormal();
-		AxisY = (AxisY / float(Buffer.Num())).GetSafeNormal();
-		Scale = Scale / float(Buffer.Num());
+		if (NumberValidEntries > 0)
+		{
+			Location = Location / float(NumberValidEntries);
+			AxisX = (AxisX / float(NumberValidEntries)).GetSafeNormal();
+			AxisY = (AxisY / float(NumberValidEntries)).GetSafeNormal();
+			Scale = Scale / float(NumberValidEntries);
+		}
 
 		Result.SetLocation(Location);
 		Result.SetRotation(FRotationMatrix::MakeFromXY(AxisX, AxisY).ToQuat());
diff --git a/Engine/Plugins/Experimental/ControlRig/Source/ControlRig/Private/Units/Simulation/RigUnit_Kalman.h b/Engine/Plugins/Experimental/ControlRig/Source/ControlRig/Private/Units/Simulation/RigUnit_Kalman.h
--- a/Engine/Plugins/Experimental/ControlRig/Source/ControlRig/Private/Units/Simulation/RigUnit_Kalman.h
+++ b/Engine/Plugins/Experimental/ControlRig/Source/ControlRig/Private/Units/Simulation/RigUnit_Kalman.h
@@ -33,7 +33,7 @@ struct FRigUnit_KalmanFloat : public FRigUnit_SimBase
 	UPROPERTY(meta=(Output))
 	float Result;
 
-	UPROPERTY()
+	UPROPERTY(meta=(MaxArraySize = "FMath::Clamp<int32>(BufferSize, 1, 512)"))
 	TArray<float> Buffer;
 
 	UPROPERTY()
@@ -68,7 +68,7 @@ struct FRigUnit_KalmanVector : public FRigUnit_SimBase
 	UPROPERTY(meta=(Output))
 	FVector Result;
 
-	UPROPERTY()
+	UPROPERTY(meta=(MaxArraySize = "FMath::Clamp<int32>(BufferSize, 1, 512)"))
 	TArray<FVector> Buffer;
 
 	UPROPERTY()
@@ -103,7 +103,7 @@ struct FRigUnit_KalmanTransform : public FRigUnit_SimBase
 	UPROPERTY(meta=(Output))
 	FTransform Result;
 
-	UPROPERTY()
+	UPROPERTY(meta = (MaxArraySize = "FMath::Clamp<int32>(BufferSize, 1, 512)"))
 	TArray<FTransform> Buffer;
 
 	UPROPERTY()
diff --git a/Engine/Plugins/Experimental/ControlRig/Source/ControlRig/Public/Drawing/ControlRigDrawInterface.h b/Engine/Plugins/Experimental/ControlRig/Source/ControlRig/Public/Drawing/ControlRigDrawInterface.h
--- a/Engine/Plugins/Experimental/ControlRig/Source/ControlRig/Public/Drawing/ControlRigDrawInterface.h
+++ b/Engine/Plugins/Experimental/ControlRig/Source/ControlRig/Public/Drawing/ControlRigDrawInterface.h
@@ -29,10 +29,10 @@ public:
 
 
 	void DrawPoint(const FTransform& WorldOffset, const FVector& Position, float Size, const FLinearColor& Color);
-	void DrawPoints(const FTransform& WorldOffset, const TArray<FVector>& Points, float Size, const FLinearColor& Color);
+	void DrawPoints(const FTransform& WorldOffset, const TArrayView<FVector>& Points, float Size, const FLinearColor& Color);
 	void DrawLine(const FTransform& WorldOffset, const FVector& LineStart, const FVector& LineEnd, const FLinearColor& Color, float Thickness = 0.f);
-	void DrawLines(const FTransform& WorldOffset, const TArray<FVector>& Positions, const FLinearColor& Color, float Thickness = 0.f);
-	void DrawLineStrip(const FTransform& WorldOffset, const TArray<FVector>& Positions, const FLinearColor& Color, float Thickness = 0.f);
+	void DrawLines(const FTransform& WorldOffset, const TArrayView<FVector>& Positions, const FLinearColor& Color, float Thickness = 0.f);
+	void DrawLineStrip(const FTransform& WorldOffset, const TArrayView<FVector>& Positions, const FLinearColor& Color, float Thickness = 0.f);
 	void DrawBox(const FTransform& WorldOffset, const FTransform& Transform, const FLinearColor& Color, float Thickness = 0.f);
 	void DrawAxes(const FTransform& WorldOffset, const FTransform& Transform, float Size, float Thickness = 0.f);
 	void DrawRectangle(const FTransform& WorldOffset, const FTransform& Transform, float Size, const FLinearColor& Color, float Thickness);
diff --git a/Engine/Source/Programs/UnrealHeaderTool/Private/CodeGenerator.cpp b/Engine/Source/Programs/UnrealHeaderTool/Private/CodeGenerator.cpp
--- a/Engine/Source/Programs/UnrealHeaderTool/Private/CodeGenerator.cpp
+++ b/Engine/Source/Programs/UnrealHeaderTool/Private/CodeGenerator.cpp
@@ -3302,57 +3302,83 @@ void FNativeClassHeaderGenerator::ExportGeneratedStructBodyMacros(FOutputDevice&
 	FString StaticVirtualMacroPrefix(StructNameCPP);
 	StaticVirtualMacroPrefix = FString::Printf(TEXT("UE_%s"), *StaticVirtualMacroPrefix.Mid(1));
 
+	const TCHAR* Comma = TEXT(", ");
+	TArray<FString> MemberNames, MemberDeclarations, MemberCasts, MemberCastProlog;
+	TMap<FString, FString> MemberMaxArraySizeGetters;
 	const TArray<FStaticVirtualMethodInfo>* StructStaticVirtualMethods = FHeaderParser::StructStaticVirtualMethods.Find(Struct);
 	if(StructStaticVirtualMethods)
 	{
-		const TCHAR* Comma = TEXT(", ");
-		TArray<FString> PropertyNames, PropertyArguments;
+		const TCHAR* FText = TEXT("F");
+		const TCHAR* TText = TEXT("T");
+		const TCHAR* EText = TEXT("E");
+		const TCHAR* InputText = TEXT("Input");
+		const TCHAR* OutputText = TEXT("Output");
+		const TCHAR* ConstantText = TEXT("Constant");
+		const TCHAR* MaxArraySizeText = TEXT("MaxArraySize");
+		const TCHAR* TArrayText = TEXT("TArray");
+		const TCHAR* TArrayViewText = TEXT("TArrayView");
+
 		for (UProperty* Prop : TFieldRange<UProperty>(Struct))
 		{
-			bool bIsInput = Prop->HasMetaData(TEXT("Input"));
-			bool bIsOutput = Prop->HasMetaData(TEXT("Output"));
-			bool bIsConstant = Prop->HasMetaData(TEXT("Constant"));
+			bool bIsInput = Prop->HasMetaData(InputText);
+			bool bIsOutput = Prop->HasMetaData(OutputText);
+			bool bIsConstant = Prop->HasMetaData(ConstantText);
+			bIsConstant = bIsConstant || (bIsInput && !bIsOutput);
 
 			FString Name = Prop->GetName();
-			PropertyNames.Add(Name);
+			MemberNames.Add(Name);
 
-			FString CPPType;
+			FString MemberCPPType;
 			FString ExtendedCPPType;
-			CPPType = Prop->GetCPPType(&ExtendedCPPType);
-			CPPType += ExtendedCPPType;
+			MemberCPPType = Prop->GetCPPType(&ExtendedCPPType);
+			MemberCPPType += ExtendedCPPType;
+			
+			FString ParameterCPPType = MemberCPPType;
+			FString MemberCast = Name;
 
-			if (CPPType.StartsWith(TEXT("F"), ESearchCase::CaseSensitive) ||
-				CPPType.StartsWith(TEXT("T"), ESearchCase::CaseSensitive) ||
-				CPPType.StartsWith(TEXT("E"), ESearchCase::CaseSensitive))
+			if (MemberCPPType.StartsWith(TArrayText))
 			{
-				if (bIsConstant || (bIsInput && !bIsOutput))
+				if (bIsConstant || Prop->HasMetaData(MaxArraySizeText))
 				{
-					PropertyArguments.Add(FString::Printf(TEXT("const %s& %s"), *CPPType, *Name));
+					ParameterCPPType = FString::Printf(TEXT("%s%s"), TArrayViewText, *MemberCPPType.RightChop(6));
+					if (!bIsConstant && Prop->HasMetaData(MaxArraySizeText))
+					{
+						MemberCast = FString::Printf(TEXT("%s_ArrayViewCast"), *Name);
+						MemberMaxArraySizeGetters.Add(Name, Prop->GetMetaData(MaxArraySizeText));
+						MemberCastProlog.Add(FString::Printf(TEXT("%s.SetNumUninitialized(%s_MaxArraySize());"), *Name, *Name));
+						MemberCastProlog.Add(FString::Printf(TEXT("%s %s(%s);"), *ParameterCPPType, *MemberCast, *Name));
+					}
+				}
+			}
+
+			if (ParameterCPPType.StartsWith(FText, ESearchCase::CaseSensitive) ||
+				ParameterCPPType.StartsWith(TText, ESearchCase::CaseSensitive) ||
+				ParameterCPPType.StartsWith(EText, ESearchCase::CaseSensitive))
+			{
+				if (bIsConstant)
+				{
+					MemberDeclarations.Add(FString::Printf(TEXT("const %s& %s"), *ParameterCPPType, *Name));
 				}
 				else
 				{
-					PropertyArguments.Add(FString::Printf(TEXT("%s& %s"), *CPPType, *Name));
+					MemberDeclarations.Add(FString::Printf(TEXT("%s& %s"), *ParameterCPPType, *Name));
 				}
 			}
-			else if (bIsConstant || (bIsInput && !bIsOutput))
+			else if (bIsConstant)
 			{
-				PropertyArguments.Add(FString::Printf(TEXT("const %s %s"), *CPPType, *Name));
+				MemberDeclarations.Add(FString::Printf(TEXT("const %s %s"), *ParameterCPPType, *Name));
 			}
 			else
 			{
-				PropertyArguments.Add(FString::Printf(TEXT("%s& %s"), *CPPType, *Name));
+				MemberDeclarations.Add(FString::Printf(TEXT("%s& %s"), *ParameterCPPType, *Name));
 			}
+
+			MemberCasts.Add(MemberCast);
 		}
 
 		OutGeneratedHeaderText.Log(TEXT("\n"));
-		OutGeneratedHeaderText.Logf(TEXT("#define %s_MEMBER_NAMES %s\n"), *StaticVirtualMacroPrefix, *FString::Join(PropertyNames, Comma));
-		OutGeneratedHeaderText.Logf(TEXT("#define %s_MEMBER_ARGUMENTS %s\n"), *StaticVirtualMacroPrefix, *FString::Join(PropertyArguments, Comma));
-		OutGeneratedHeaderText.Logf(TEXT("#define %s_DECLARE_STATIC_VIRTUAL_METHOD(ReturnType, FunctionName, ...) \\\n"), *StaticVirtualMacroPrefix);
-		OutGeneratedHeaderText.Logf(TEXT("    static ReturnType Static##FunctionName(%s_MEMBER_ARGUMENTS, ##__VA_ARGS__);\n"), *StaticVirtualMacroPrefix);
 		OutGeneratedHeaderText.Logf(TEXT("#define %s_IMPLEMENT_STATIC_VIRTUAL_METHOD(ReturnType, FunctionName, ...) \\\n"), *StaticVirtualMacroPrefix);
-		OutGeneratedHeaderText.Logf(TEXT("    ReturnType %s::Static##FunctionName(%s_MEMBER_ARGUMENTS, ##__VA_ARGS__)\n"), StructNameCPP, *StaticVirtualMacroPrefix);
-		OutGeneratedHeaderText.Logf(TEXT("#define %s_INVOKE_STATIC_VIRTUAL_METHOD(FunctionName, ...) \\\n"), *StaticVirtualMacroPrefix);
-		OutGeneratedHeaderText.Logf(TEXT("    %s::Static##FunctionName(%s_MEMBER_NAMES, ##__VA_ARGS__);\n"), StructNameCPP, *StaticVirtualMacroPrefix);
+		OutGeneratedHeaderText.Logf(TEXT("\tReturnType %s::Static##FunctionName(\\\n\t\t%s, ##__VA_ARGS__)\n"), StructNameCPP, *FString::Join(MemberDeclarations, TEXT(",\\\n\t\t")));
 		OutGeneratedHeaderText.Log(TEXT("\n"));
 	}
 
@@ -3376,14 +3402,13 @@ void FNativeClassHeaderGenerator::ExportGeneratedStructBodyMacros(FOutputDevice&
 		{
 			for (const FStaticVirtualMethodInfo& Info : (*StructStaticVirtualMethods))
 			{
-				if (Info.Params.IsEmpty())
-				{
-					StaticVirtualMethodsDeclarations += FString::Printf(TEXT("\t%s_DECLARE_STATIC_VIRTUAL_METHOD(%s, %s)\n"), *StaticVirtualMacroPrefix, *Info.ReturnType, *Info.Name);
-				}
-				else
-				{
-					StaticVirtualMethodsDeclarations += FString::Printf(TEXT("\t%s_DECLARE_STATIC_VIRTUAL_METHOD(%s, %s, %s)\n"), *StaticVirtualMacroPrefix, *Info.ReturnType, *Info.Name, *Info.Params);
-				}
+				FString ParamsSuffix = Info.Params.IsEmpty() ? TEXT("") : FString::Printf(TEXT(", %s"), *Info.Params);
+				StaticVirtualMethodsDeclarations += FString::Printf(TEXT("\tstatic %s Static%s(\r\n\t\t%s%s);\r\n"), *Info.ReturnType, *Info.Name, *FString::Join(MemberDeclarations, TEXT(",\r\n\t\t")), *ParamsSuffix);
+			}
+
+			for (const TPair<FString, FString>& MemberArraySizeGetter : MemberMaxArraySizeGetters)
+			{
+				StaticVirtualMethodsDeclarations += FString::Printf(TEXT("\tint32 %s_MaxArraySize() const { return %s; }\r\n"), *MemberArraySizeGetter.Key, *MemberArraySizeGetter.Value);
 			}
 		}
 
@@ -3601,18 +3626,15 @@ void FNativeClassHeaderGenerator::ExportGeneratedStructBodyMacros(FOutputDevice&
 			// implement the virtual function body.
 			Out.Logf(TEXT("%s %s::%s(%s)\n"), *Info.ReturnType, StructNameCPP, *Info.Name, *Info.Params);
 			Out.Log(TEXT("{\n"));
-			if (Info.Params.IsEmpty())
+			for (const FString& MemberCastPrologLine : MemberCastProlog)
 			{
-				if(Info.ReturnType.IsEmpty() || (Info.ReturnType == TEXT("void")))
-				{
-					Out.Logf(TEXT("    %s_INVOKE_STATIC_VIRTUAL_METHOD(%s);\n"), *StaticVirtualMacroPrefix, *Info.Name);
-				}
-				else
-				{
-					Out.Logf(TEXT("    return %s_INVOKE_STATIC_VIRTUAL_METHOD(%s);\n"), *StaticVirtualMacroPrefix, *Info.Name);
-				}
+				Out.Logf(TEXT("\t%s\n"), *MemberCastPrologLine);
 			}
-			else
+
+			FString ReturnPrefix = (Info.ReturnType.IsEmpty() || (Info.ReturnType == TEXT("void"))) ? TEXT("") : TEXT("return ");
+			FString ParamSuffix;
+
+			if (!Info.Params.IsEmpty())
 			{
 				TArray<FString> Params;
 				FString ParamPrev, ParamLeft, ParamRight;
@@ -3633,16 +3655,11 @@ void FNativeClassHeaderGenerator::ExportGeneratedStructBodyMacros(FOutputDevice&
 						Params[ParamIndex] = Params[ParamIndex].Mid(LastSpace + 1);
 					}
 				}
-
-				if(Info.ReturnType.IsEmpty() || (Info.ReturnType == TEXT("void")))
-				{
-					Out.Logf(TEXT("    %s_INVOKE_STATIC_VIRTUAL_METHOD(%s, %s);\n"), *StaticVirtualMacroPrefix, *Info.Name, *FString::Join(Params, TEXT(", ")));
-				}
-				else
-				{
-					Out.Logf(TEXT("    return %s_INVOKE_STATIC_VIRTUAL_METHOD(%s, %s);\n"), *StaticVirtualMacroPrefix, *Info.Name, *FString::Join(Params, TEXT(", ")));
-				}
+				ParamSuffix = FString::Printf(TEXT(", %s"), *FString::Join(Params, Comma));
 			}
+
+			Out.Logf(TEXT("    %s%s::Static%s(%s%s);\n"), *ReturnPrefix, StructNameCPP, *Info.Name, *FString::Join(MemberCasts, Comma), *ParamSuffix);
+
 			Out.Log(TEXT("}\n"));
 		}
 
diff --git a/Engine/Source/Programs/UnrealHeaderTool/Resources/UHTDebugging/StaticMethodStructHeader.h b/Engine/Source/Programs/UnrealHeaderTool/Resources/UHTDebugging/StaticMethodStructHeader.h
--- a/Engine/Source/Programs/UnrealHeaderTool/Resources/UHTDebugging/StaticMethodStructHeader.h
+++ b/Engine/Source/Programs/UnrealHeaderTool/Resources/UHTDebugging/StaticMethodStructHeader.h
@@ -45,6 +45,12 @@ struct FStaticMethodStruct : public FStaticMethodStructBase
 	UPROPERTY(meta = (Input, Output))
 	TArray<FVector> E;
 
+	UPROPERTY(meta = (Input))
+	TArray<FVector> F;
+
+	UPROPERTY(meta = (Output, MaxArraySize = 8))
+	TArray<FVector> G;
+
 	UPROPERTY()
 	float Cache;
 };
