UHT: Remove use of MaxArraySize getters for SVM

diff --git a/Engine/Source/Programs/UnrealHeaderTool/Private/CodeGenerator.cpp b/Engine/Source/Programs/UnrealHeaderTool/Private/CodeGenerator.cpp
--- a/Engine/Source/Programs/UnrealHeaderTool/Private/CodeGenerator.cpp
+++ b/Engine/Source/Programs/UnrealHeaderTool/Private/CodeGenerator.cpp
@@ -3304,7 +3304,6 @@ void FNativeClassHeaderGenerator::ExportGeneratedStructBodyMacros(FOutputDevice&
 
 	const TCHAR* Comma = TEXT(", ");
 	TArray<FString> MemberNames, MemberDeclarations, MemberCasts, MemberCastProlog;
-	TMap<FString, FString> MemberMaxArraySizeGetters;
 	const TArray<FStaticVirtualMethodInfo>* StructStaticVirtualMethods = FHeaderParser::StructStaticVirtualMethods.Find(Struct);
 	if(StructStaticVirtualMethods)
 	{
@@ -3318,8 +3317,11 @@ void FNativeClassHeaderGenerator::ExportGeneratedStructBodyMacros(FOutputDevice&
 		const TCHAR* TArrayText = TEXT("TArray");
 		const TCHAR* TArrayViewText = TEXT("TArrayView");
 
+		int32 MemberCastSuffix = 0;
 		for (UProperty* Prop : TFieldRange<UProperty>(Struct))
 		{
+			++MemberCastSuffix;
+
 			bool bIsInput = Prop->HasMetaData(InputText);
 			bool bIsOutput = Prop->HasMetaData(OutputText);
 			bool bIsConstant = Prop->HasMetaData(ConstantText);
@@ -3343,9 +3345,8 @@ void FNativeClassHeaderGenerator::ExportGeneratedStructBodyMacros(FOutputDevice&
 					ParameterCPPType = FString::Printf(TEXT("%s%s"), TArrayViewText, *MemberCPPType.RightChop(6));
 					if (!bIsConstant && Prop->HasMetaData(MaxArraySizeText))
 					{
-						MemberCast = FString::Printf(TEXT("%s_ArrayViewCast"), *Name);
-						MemberMaxArraySizeGetters.Add(Name, Prop->GetMetaData(MaxArraySizeText));
-						MemberCastProlog.Add(FString::Printf(TEXT("%s.SetNumUninitialized(%s_MaxArraySize());"), *Name, *Name));
+						MemberCast = FString::Printf(TEXT("%s_%d_View"), *Name, MemberCastSuffix);
+						MemberCastProlog.Add(FString::Printf(TEXT("%s.SetNumUninitialized( %s );"), *Name, *Prop->GetMetaData(MaxArraySizeText)));
 						MemberCastProlog.Add(FString::Printf(TEXT("%s %s(%s);"), *ParameterCPPType, *MemberCast, *Name));
 					}
 				}
@@ -3405,11 +3406,6 @@ void FNativeClassHeaderGenerator::ExportGeneratedStructBodyMacros(FOutputDevice&
 				FString ParamsSuffix = Info.Params.IsEmpty() ? TEXT("") : FString::Printf(TEXT(", %s"), *Info.Params);
 				StaticVirtualMethodsDeclarations += FString::Printf(TEXT("\tstatic %s Static%s(\r\n\t\t%s%s);\r\n"), *Info.ReturnType, *Info.Name, *FString::Join(MemberDeclarations, TEXT(",\r\n\t\t")), *ParamsSuffix);
 			}
-
-			for (const TPair<FString, FString>& MemberArraySizeGetter : MemberMaxArraySizeGetters)
-			{
-				StaticVirtualMethodsDeclarations += FString::Printf(TEXT("\tint32 %s_MaxArraySize() const { return %s; }\r\n"), *MemberArraySizeGetter.Key, *MemberArraySizeGetter.Value);
-			}
 		}
 
 		const FString SuperTypedef = BaseStruct ? FString::Printf(TEXT("\ttypedef %s Super;\r\n"), NameLookupCPP.GetNameCPP(BaseStruct)) : FString();
