Multiplex: Centralized registration of multiplex methods + unit tests of hand written invocations

diff --git a/Engine/Source/Programs/UnrealHeaderTool/Private/CodeGenerator.cpp b/Engine/Source/Programs/UnrealHeaderTool/Private/CodeGenerator.cpp
--- a/Engine/Source/Programs/UnrealHeaderTool/Private/CodeGenerator.cpp
+++ b/Engine/Source/Programs/UnrealHeaderTool/Private/CodeGenerator.cpp
@@ -3299,13 +3299,13 @@ void FNativeClassHeaderGenerator::ExportGeneratedStructBodyMacros(FOutputDevice&
 
 	const FString SingletonName = GetSingletonName(Struct);
 
-	FString StaticVirtualMacroPrefix(StructNameCPP);
-	StaticVirtualMacroPrefix = FString::Printf(TEXT("UE_%s"), *StaticVirtualMacroPrefix.Mid(1));
+	FString MultiplexMacroPrefix(StructNameCPP);
+	MultiplexMacroPrefix = FString::Printf(TEXT("UE_%s"), *MultiplexMacroPrefix.Mid(1));
 
 	const TCHAR* Comma = TEXT(", ");
-	TArray<FString> MemberNames, MemberDeclarations, MemberCasts, MemberCastProlog;
-	const TArray<FStaticVirtualMethodInfo>* StructStaticVirtualMethods = FHeaderParser::StructStaticVirtualMethods.Find(Struct);
-	if(StructStaticVirtualMethods)
+	TArray<FString> MemberNames, MemberDeclarations, MemberCasts, MemberCastProlog, MemberMulticastProlog;
+	const TArray<FMultiplexMethodInfo>* StructMultiplexMethods = FHeaderParser::StructMultiplexMethods.Find(Struct);
+	if(StructMultiplexMethods)
 	{
 		const TCHAR* FText = TEXT("F");
 		const TCHAR* TText = TEXT("T");
@@ -3318,6 +3318,17 @@ void FNativeClassHeaderGenerator::ExportGeneratedStructBodyMacros(FOutputDevice&
 		const TCHAR* TArrayViewText = TEXT("TArrayView");
 
 		int32 MemberCastSuffix = 0;
+
+		int32 MemberIndex = 0;
+		for (UProperty* Prop : TFieldRange<UProperty>(Struct))
+		{
+			FString Name = Prop->GetName();
+			FString AddressName = FString::Printf(TEXT("%s_Address"), *Name);
+			MemberMulticastProlog.Add(FString::Printf(TEXT("const FMultiplexArgument& %s = MultiplexAddresses[%d];"), *AddressName, MemberIndex++));
+		}
+		MemberMulticastProlog.Insert(FString::Printf(TEXT("ensure(MultiplexAddresses.Num() == %d);"), MemberMulticastProlog.Num()), 0);
+		MemberMulticastProlog.Add(FString());
+
 		for (UProperty* Prop : TFieldRange<UProperty>(Struct))
 		{
 			++MemberCastSuffix;
@@ -3336,18 +3347,24 @@ void FNativeClassHeaderGenerator::ExportGeneratedStructBodyMacros(FOutputDevice&
 			MemberCPPType += ExtendedCPPType;
 			
 			FString ParameterCPPType = MemberCPPType;
+			FString MultiCastCPPType = MemberCPPType;
+			FString MulticastGetter = TEXT("GetRef");
+			FString ExtractedMulticastCPPType;
 			FString MemberCast = Name;
 
 			if (MemberCPPType.StartsWith(TArrayText))
 			{
 				if (bIsConstant || Prop->HasMetaData(MaxArraySizeText))
 				{
-					ParameterCPPType = FString::Printf(TEXT("%s%s"), TArrayViewText, *MemberCPPType.RightChop(6));
+					ParameterCPPType = FString::Printf(TEXT("%s%s"), TArrayViewText, *ExtendedCPPType);
 					if (!bIsConstant && Prop->HasMetaData(MaxArraySizeText))
 					{
 						MemberCast = FString::Printf(TEXT("%s_%d_View"), *Name, MemberCastSuffix);
 						MemberCastProlog.Add(FString::Printf(TEXT("%s.SetNumUninitialized( %s );"), *Name, *Prop->GetMetaData(MaxArraySizeText)));
 						MemberCastProlog.Add(FString::Printf(TEXT("%s %s(%s);"), *ParameterCPPType, *MemberCast, *Name));
+						MultiCastCPPType = ExtendedCPPType.LeftChop(1).RightChop(1);
+						MulticastGetter = TEXT("GetArray");
+						ExtractedMulticastCPPType = FString::Printf(TEXT("TArrayView<%s>"), *MultiCastCPPType);
 					}
 				}
 			}
@@ -3358,28 +3375,37 @@ void FNativeClassHeaderGenerator::ExportGeneratedStructBodyMacros(FOutputDevice&
 			{
 				if (bIsConstant)
 				{
-					MemberDeclarations.Add(FString::Printf(TEXT("const %s& %s"), *ParameterCPPType, *Name));
+					ParameterCPPType = FString::Printf(TEXT("const %s&"), *ParameterCPPType);
 				}
 				else
 				{
-					MemberDeclarations.Add(FString::Printf(TEXT("%s& %s"), *ParameterCPPType, *Name));
+					ParameterCPPType = FString::Printf(TEXT("%s&"), *ParameterCPPType);
 				}
 			}
 			else if (bIsConstant)
 			{
-				MemberDeclarations.Add(FString::Printf(TEXT("const %s %s"), *ParameterCPPType, *Name));
+				ParameterCPPType = FString::Printf(TEXT("const %s"), *ParameterCPPType);
 			}
 			else
 			{
-				MemberDeclarations.Add(FString::Printf(TEXT("%s& %s"), *ParameterCPPType, *Name));
+				ParameterCPPType = FString::Printf(TEXT("%s&"), *ParameterCPPType);
 			}
 
 			MemberCasts.Add(MemberCast);
+
+			if (ExtractedMulticastCPPType.IsEmpty())
+			{
+				ExtractedMulticastCPPType = ParameterCPPType;
+			}
+
+			MemberDeclarations.Add(FString::Printf(TEXT("%s %s"), *ParameterCPPType, *Name));
+			FString AddressName = FString::Printf(TEXT("%s_Address"), *Name);
+			MemberMulticastProlog.Add(FString::Printf(TEXT("%s %s = MultiplexStorage[%s.StorageType()][%s.Index()].%s<%s>();"), *ExtractedMulticastCPPType, *Name, *AddressName, *AddressName, *MulticastGetter, *MultiCastCPPType));
 		}
 
 		OutGeneratedHeaderText.Log(TEXT("\n"));
-		OutGeneratedHeaderText.Logf(TEXT("#define %s_IMPLEMENT_STATIC_VIRTUAL_METHOD(ReturnType, FunctionName, ...) \\\n"), *StaticVirtualMacroPrefix);
-		OutGeneratedHeaderText.Logf(TEXT("\tReturnType %s::Static##FunctionName(\\\n\t\t%s, ##__VA_ARGS__)\n"), StructNameCPP, *FString::Join(MemberDeclarations, TEXT(",\\\n\t\t")));
+		OutGeneratedHeaderText.Logf(TEXT("#define %s_IMPLEMENT_MULTIPLEX(ReturnType, FunctionName, ...) \\\n"), *MultiplexMacroPrefix);
+		OutGeneratedHeaderText.Logf(TEXT("\tReturnType %s::Static##FunctionName(\\\r\n\t\t%s, ##__VA_ARGS__)\n"), StructNameCPP, *FString::Join(MemberDeclarations, TEXT(",\\\r\n\t\t")));
 		OutGeneratedHeaderText.Log(TEXT("\n"));
 	}
 
@@ -3398,19 +3424,20 @@ void FNativeClassHeaderGenerator::ExportGeneratedStructBodyMacros(FOutputDevice&
 		const FString StaticClassLine = FString::Printf(TEXT("\t%sstatic class UScriptStruct* StaticStruct();\r\n"), *RequiredAPI);
 		const FString PrivatePropertiesOffset = PrivatePropertiesOffsetGetters(Struct, StructNameCPP);
 		
-		FString StaticVirtualMethodsDeclarations;
-		if (StructStaticVirtualMethods)
+		FString MultiplexMethodsDeclarations;
+		if (StructMultiplexMethods)
 		{
-			for (const FStaticVirtualMethodInfo& Info : (*StructStaticVirtualMethods))
+			for (const FMultiplexMethodInfo& Info : (*StructMultiplexMethods))
 			{
 				FString ParamsSuffix = Info.Params.IsEmpty() ? TEXT("") : FString::Printf(TEXT(", %s"), *Info.Params);
-				StaticVirtualMethodsDeclarations += FString::Printf(TEXT("\tstatic %s Static%s(\r\n\t\t%s%s);\r\n"), *Info.ReturnType, *Info.Name, *FString::Join(MemberDeclarations, TEXT(",\r\n\t\t")), *ParamsSuffix);
+				MultiplexMethodsDeclarations += FString::Printf(TEXT("\tstatic %s Static%s(\r\n\t\t%s%s);\r\n"), *Info.ReturnType, *Info.Name, *FString::Join(MemberDeclarations, TEXT(",\r\n\t\t")), *ParamsSuffix);
+				MultiplexMethodsDeclarations += FString::Printf(TEXT("\tstatic %s Multiplex%s(const TArrayView<FMultiplexArgument>& MultiplexAddresses, FMultiplexStorage* MultiplexStorage, const TArrayView<void*>& AdditionalArgs);\r\n"), *Info.ReturnType, *Info.Name);
 			}
 		}
 
 		const FString SuperTypedef = BaseStruct ? FString::Printf(TEXT("\ttypedef %s Super;\r\n"), NameLookupCPP.GetNameCPP(BaseStruct)) : FString();
 
-		const FString CombinedLine = FriendLine + StaticClassLine + StaticVirtualMethodsDeclarations + PrivatePropertiesOffset + SuperTypedef;
+		const FString CombinedLine = FriendLine + StaticClassLine + MultiplexMethodsDeclarations + PrivatePropertiesOffset + SuperTypedef;
 		const FString MacroName = SourceFile.GetGeneratedBodyMacroName(Struct->StructMacroDeclaredLineNumber);
 
 		const FString Macroized = Macroize(*MacroName, *CombinedLine);
@@ -3442,6 +3469,15 @@ void FNativeClassHeaderGenerator::ExportGeneratedStructBodyMacros(FOutputDevice&
 		Out.Logf(TEXT("\t\tSingleton = GetStaticStruct(%s, %s, TEXT(\"%s\"), sizeof(%s), %s());\r\n"),
 			*SingletonName.LeftChop(2), *OuterName, *ActualStructName, StructNameCPP, *GetHashName);
 
+		if (StructMultiplexMethods)
+		{
+			for (const FMultiplexMethodInfo& Info : (*StructMultiplexMethods))
+			{
+				Out.Logf(TEXT("\t\tFMultiplexRegistry::Get().Register(TEXT(\"%s::%s\"), &%s::Multiplex%s);\r\n"),
+					StructNameCPP, *Info.Name, StructNameCPP, *Info.Name);
+			}
+		}
+
 		Out.Logf(TEXT("\t}\r\n"));
 		Out.Logf(TEXT("\treturn Singleton;\r\n"));
 		Out.Logf(TEXT("}\r\n"));
@@ -3614,22 +3650,24 @@ void FNativeClassHeaderGenerator::ExportGeneratedStructBodyMacros(FOutputDevice&
 	Out.Log(GeneratedStructRegisterFunctionText);
 	Out.Logf(TEXT("\tuint32 %s() { return %uU; }\r\n"), *HashFuncName, StructHash);
 
-	if (StructStaticVirtualMethods)
+	if (StructMultiplexMethods)
 	{
-		Out.Log(TEXT("\n"));
-		for (const FStaticVirtualMethodInfo& Info : (*StructStaticVirtualMethods))
+		for (const FMultiplexMethodInfo& Info : (*StructMultiplexMethods))
 		{
+			Out.Log(TEXT("\r\n"));
+
 			// implement the virtual function body.
-			Out.Logf(TEXT("%s %s::%s(%s)\n"), *Info.ReturnType, StructNameCPP, *Info.Name, *Info.Params);
-			Out.Log(TEXT("{\n"));
+			Out.Logf(TEXT("%s %s::%s(%s)\r\n"), *Info.ReturnType, StructNameCPP, *Info.Name, *Info.Params);
+			Out.Log(TEXT("{\r\n"));
 			for (const FString& MemberCastPrologLine : MemberCastProlog)
 			{
-				Out.Logf(TEXT("\t%s\n"), *MemberCastPrologLine);
+				Out.Logf(TEXT("\t%s\r\n"), *MemberCastPrologLine);
 			}
 
 			FString ReturnPrefix = (Info.ReturnType.IsEmpty() || (Info.ReturnType == TEXT("void"))) ? TEXT("") : TEXT("return ");
 			FString ParamSuffix;
 
+			TArray<FString> ParameterExtractionProlog;
 			if (!Info.Params.IsEmpty())
 			{
 				TArray<FString> Params;
@@ -3643,23 +3681,79 @@ void FNativeClassHeaderGenerator::ExportGeneratedStructBodyMacros(FOutputDevice&
 				} 
 				Params.Add(ParamPrev.TrimStartAndEnd());
 
+				ParameterExtractionProlog.Add(FString::Printf(TEXT("ensure(AdditionalArgs.Num() == %d);"), Params.Num()));
+
 				for (int32 ParamIndex=0;ParamIndex<Params.Num(); ParamIndex++)
 				{
+					FString FullParam = Params[ParamIndex];
+					int32 LastEqual = INDEX_NONE;
+					if (FullParam.FindLastChar(TCHAR('='), LastEqual))
+					{
+						FullParam = FullParam.Mid(0, LastEqual);
+					}
+
+					FullParam.TrimStartAndEndInline();
+
+					FString ParamType = FullParam;
+					FString ParamName = FullParam;
+
 					int32 LastSpace = INDEX_NONE;
-					if (Params[ParamIndex].FindLastChar(TCHAR(' '), LastSpace))
+					if (ParamName.FindLastChar(TCHAR(' '), LastSpace))
+					{
+						ParamType = ParamType.Mid(0, LastSpace);
+						ParamName = ParamName.Mid(LastSpace + 1);
+					}
+
+					ParamType.TrimStartAndEndInline();
+					ParamName.TrimStartAndEndInline();
+
+					Params[ParamIndex] = ParamName;
+
+					FString ParamPointerType = ParamType;
+					if (ParamPointerType.EndsWith(TEXT("&")))
 					{
-						Params[ParamIndex] = Params[ParamIndex].Mid(LastSpace + 1);
+						ParamPointerType = ParamPointerType.LeftChop(1);
+						ParamPointerType.TrimStartAndEndInline();
 					}
+					ParameterExtractionProlog.Add(FString::Printf(TEXT("%s = *(%s*)AdditionalArgs[%d];"), *FullParam, *ParamPointerType, ParamIndex));
 				}
 				ParamSuffix = FString::Printf(TEXT(", %s"), *FString::Join(Params, Comma));
 			}
 
-			Out.Logf(TEXT("    %s%s::Static%s(%s%s);\n"), *ReturnPrefix, StructNameCPP, *Info.Name, *FString::Join(MemberCasts, Comma), *ParamSuffix);
+			Out.Logf(TEXT("    %sStatic%s(%s%s);\n"), *ReturnPrefix, *Info.Name, *FString::Join(MemberCasts, Comma), *ParamSuffix);
+			Out.Log(TEXT("}\r\n"));
+
+			Out.Log(TEXT("\r\n"));
+
+			FString ParamDeclarationSuffix;
+			if (!Info.Params.IsEmpty())
+			{
+				ParamDeclarationSuffix = FString::Printf(TEXT(", %s"), *Info.Params);
+			}
+
+			Out.Logf(TEXT("%s %s::Multiplex%s(const TArrayView<FMultiplexArgument>& MultiplexAddresses, FMultiplexStorage* MultiplexStorage, const TArrayView<void*>& AdditionalArgs)\r\n"), *Info.ReturnType, StructNameCPP, *Info.Name);
+			Out.Log(TEXT("{\r\n"));
+
+			if (ParameterExtractionProlog.Num() > 0)
+			{
+				for (const FString& ParameterExtractionPrologLine : ParameterExtractionProlog)
+				{
+					Out.Logf(TEXT("\t%s\r\n"), *ParameterExtractionPrologLine);
+				}
+				Out.Log(TEXT("\t\r\n"));
+			}
+
+			for (const FString& MemberMulticastPrologLine : MemberMulticastProlog)
+			{
+				Out.Logf(TEXT("\t%s\r\n"), *MemberMulticastPrologLine);
+			}
+
+			Out.Logf(TEXT("\t%sStatic%s(%s%s);\r\n"), *ReturnPrefix, *Info.Name, *FString::Join(MemberNames, Comma), *ParamSuffix);
+			Out.Log(TEXT("}\r\n"));
 
-			Out.Log(TEXT("}\n"));
 		}
 
-		Out.Log(TEXT("\n"));
+		Out.Log(TEXT("\r\n"));
 	}
 }
 
diff --git a/Engine/Source/Programs/UnrealHeaderTool/Private/HeaderParser.cpp b/Engine/Source/Programs/UnrealHeaderTool/Private/HeaderParser.cpp
--- a/Engine/Source/Programs/UnrealHeaderTool/Private/HeaderParser.cpp
+++ b/Engine/Source/Programs/UnrealHeaderTool/Private/HeaderParser.cpp
@@ -72,7 +72,7 @@ static const FName NAME_DocumentationPolicy(TEXT("DocumentationPolicy"));
 EGeneratedCodeVersion FHeaderParser::DefaultGeneratedCodeVersion = EGeneratedCodeVersion::V1;
 TArray<FString> FHeaderParser::StructsWithNoPrefix;
 TArray<FString> FHeaderParser::StructsWithTPrefix;
-TMap<UStruct*, TArray<FStaticVirtualMethodInfo>> FHeaderParser::StructStaticVirtualMethods;
+TMap<UStruct*, TArray<FMultiplexMethodInfo>> FHeaderParser::StructMultiplexMethods;
 TArray<FString> FHeaderParser::DelegateParameterCountStrings;
 TMap<FString, FString> FHeaderParser::TypeRedirectMap;
 TArray<FString> FHeaderParser::PropertyCPPTypesRequiringUIRanges = { TEXT("float"), TEXT("double") };
@@ -2495,7 +2495,7 @@ UScriptStruct* FHeaderParser::CompileStructDeclaration(FClasses& AllClasses)
 		{
 			FError::Throwf(TEXT("USTRUCTs cannot contain UFUNCTIONs."));
 		}
-		else if (Token.Matches(TEXT("STATIC_VIRTUAL_METHOD"), ESearchCase::CaseSensitive))
+		else if (Token.Matches(TEXT("MULTIPLEX_METHOD"), ESearchCase::CaseSensitive))
 		{
 			CompileStaticMethodDeclaration(AllClasses, Struct);
 		}
@@ -6229,12 +6229,12 @@ void FHeaderParser::CompileStaticMethodDeclaration(FClasses& AllClasses, UStruct
 		}
 	}
 
-	FStaticVirtualMethodInfo Info;
+	FMultiplexMethodInfo Info;
 	Info.ReturnType = ReturnTypeToken.Identifier;
 	Info.Name = NameToken.Identifier;
 	Info.Params = FString::Join(ParamsContent, TEXT(" "));
 
-	TArray<FStaticVirtualMethodInfo>& Infos = StructStaticVirtualMethods.FindOrAdd(Struct);
+	TArray<FMultiplexMethodInfo>& Infos = StructMultiplexMethods.FindOrAdd(Struct);
 	Infos.Add(Info);
 }
 
diff --git a/Engine/Source/Programs/UnrealHeaderTool/Private/HeaderParser.h b/Engine/Source/Programs/UnrealHeaderTool/Private/HeaderParser.h
--- a/Engine/Source/Programs/UnrealHeaderTool/Private/HeaderParser.h
+++ b/Engine/Source/Programs/UnrealHeaderTool/Private/HeaderParser.h
@@ -117,7 +117,7 @@ struct FIndexRange
 	int32 Count;
 };
 
-struct FStaticVirtualMethodInfo
+struct FMultiplexMethodInfo
 {
 	FString ReturnType;
 	FString Name;
@@ -426,8 +426,8 @@ protected:
 	// List of all net service functions with undeclared response functions 
 	TMap<int32, FString> RPCsNeedingHookup;
 
-	// List of all static virtual methods defined on structs
-	static TMap<UStruct*, TArray<FStaticVirtualMethodInfo>> StructStaticVirtualMethods;
+	// List of all multiplex methods defined on structs
+	static TMap<UStruct*, TArray<FMultiplexMethodInfo>> StructMultiplexMethods;
 
 	// Constructor.
 	explicit FHeaderParser(FFeedbackContext* InWarn, const FManifestModule& InModule);
diff --git a/Engine/Source/Programs/UnrealHeaderTool/Resources/UHTDebugging/StaticMethodStructHeader.h b/Engine/Source/Programs/UnrealHeaderTool/Resources/UHTDebugging/StaticMethodStructHeader.h
--- a/Engine/Source/Programs/UnrealHeaderTool/Resources/UHTDebugging/StaticMethodStructHeader.h
+++ b/Engine/Source/Programs/UnrealHeaderTool/Resources/UHTDebugging/StaticMethodStructHeader.h
@@ -21,13 +21,13 @@ struct FStaticMethodStruct : public FStaticMethodStructBase
 {
 	GENERATED_BODY()
 
-	STATIC_VIRTUAL_METHOD()
+	MULTIPLEX_METHOD()
 	void Clear();
 
-	STATIC_VIRTUAL_METHOD()
-	virtual void Execute(bool bAdditionalFlag, const FString& InString) override;
+	MULTIPLEX_METHOD()
+	virtual void Execute(bool bAdditionalFlag = false, const FString& InString = TEXT("")) override;
 
-	STATIC_VIRTUAL_METHOD()
+	MULTIPLEX_METHOD()
 	void Compute(float TestFloat);
 
 	UPROPERTY(meta = (Input))
diff --git a/Engine/Source/Runtime/AnimationCore/Private/MultiplexRegistry.cpp b/Engine/Source/Runtime/AnimationCore/Private/MultiplexRegistry.cpp
new file mode 100644
--- /dev/null
+++ b/Engine/Source/Runtime/AnimationCore/Private/MultiplexRegistry.cpp
@@ -0,0 +1,38 @@
+// Copyright 1998-2019 Epic Games, Inc. All Rights Reserved.
+
+#include "MultiplexRegistry.h"
+#include "UObject/UObjectIterator.h"
+
+FMultiplexRegistry FMultiplexRegistry::s_MultiplexRegistry;
+
+FMultiplexRegistry& FMultiplexRegistry::Get()
+{
+	return s_MultiplexRegistry;
+}
+
+void FMultiplexRegistry::Refresh()
+{
+	// todo
+}
+
+void FMultiplexRegistry::Register(const TCHAR* InName, FMultiplexFunctionPtr InFunctionPtr)
+{
+	if (Find(InName) != nullptr)
+	{
+		return;
+	}
+	Functions.Add(FMultiplexFunction(InName, InFunctionPtr));
+}
+
+FMultiplexFunctionPtr FMultiplexRegistry::Find(const TCHAR* InName) const
+{
+	for (const FMultiplexFunction& Function : Functions)
+	{
+		if (FCString::Strcmp(Function.Name, InName) == 0)
+		{
+			return Function.FunctionPtr;
+		}
+	}
+
+	return nullptr;
+}
diff --git a/Engine/Source/Runtime/AnimationCore/Public/MultiplexRegistry.h b/Engine/Source/Runtime/AnimationCore/Public/MultiplexRegistry.h
new file mode 100644
--- /dev/null
+++ b/Engine/Source/Runtime/AnimationCore/Public/MultiplexRegistry.h
@@ -0,0 +1,73 @@
+// Copyright 1998-2019 Epic Games, Inc. All Rights Reserved.
+
+#pragma once
+
+#include "CoreMinimal.h"
+#include "MultiplexStorage.h"
+
+struct ANIMATIONCORE_API FMultiplexArgument
+{
+public:
+
+	FMultiplexArgument()
+		:Address(INDEX_NONE)
+	{
+	}
+
+	FMultiplexArgument(int32 InAddress)
+		:Address(InAddress)
+	{
+	}
+
+	FORCEINLINE bool IsLiteral() const { return Address < 0; }
+	FORCEINLINE int32 StorageType() const { return Address < 0 ? 1 : 0; }
+	FORCEINLINE int32 Index() const { return Address < 0 ? -(Address + 1) : Address; }
+
+private:
+	int32 Address;
+};
+
+typedef void (*FMultiplexFunctionPtr)(const TArrayView<FMultiplexArgument>&, FMultiplexStorage*, const TArrayView<void*>&);
+
+struct ANIMATIONCORE_API FMultiplexFunction
+{
+	const TCHAR* Name;
+	FMultiplexFunctionPtr FunctionPtr;
+
+	FMultiplexFunction()
+		: Name(nullptr)
+		, FunctionPtr(nullptr)
+	{
+	}
+
+	FMultiplexFunction(const TCHAR* InName, FMultiplexFunctionPtr InFunctionPtr)
+		: Name(InName)
+		, FunctionPtr(InFunctionPtr)
+	{
+	}
+};
+
+struct ANIMATIONCORE_API FMultiplexRegistry
+{
+public:
+	static FMultiplexRegistry& Get();
+
+	void Register(const TCHAR* InName, FMultiplexFunctionPtr InFunctionPtr);
+	void Refresh();
+
+	FMultiplexFunctionPtr Find(const TCHAR* InName) const;
+
+private:
+
+	// disable default constructor
+	FMultiplexRegistry() {}
+	// disable copy constructor
+	FMultiplexRegistry(const FMultiplexRegistry&) {}
+	// disable assignment operator
+	FMultiplexRegistry& operator= (const FMultiplexRegistry &InOther) { return *this; }
+
+	// storage for all functions
+	TArray<FMultiplexFunction> Functions;
+
+	static FMultiplexRegistry s_MultiplexRegistry;
+};
diff --git a/Engine/Source/Runtime/AnimationCore/Public/MultiplexStorage.h b/Engine/Source/Runtime/AnimationCore/Public/MultiplexStorage.h
--- a/Engine/Source/Runtime/AnimationCore/Public/MultiplexStorage.h
+++ b/Engine/Source/Runtime/AnimationCore/Public/MultiplexStorage.h
@@ -1,9 +1,5 @@
 // Copyright 1998-2019 Epic Games, Inc. All Rights Reserved.
 
-/*=============================================================================
-	AngularLimit.h: Angular limit features
-=============================================================================*/
-
 #pragma once
 
 #include "CoreMinimal.h"
diff --git a/Engine/Source/Runtime/CoreUObject/Public/UObject/ObjectMacros.h b/Engine/Source/Runtime/CoreUObject/Public/UObject/ObjectMacros.h
--- a/Engine/Source/Runtime/CoreUObject/Public/UObject/ObjectMacros.h
+++ b/Engine/Source/Runtime/CoreUObject/Public/UObject/ObjectMacros.h
@@ -581,7 +581,7 @@ struct COREUOBJECT_API FReferencerInformationList
 #define UPARAM(...)
 #define UENUM(...)
 #define UDELEGATE(...)
-#define STATIC_VIRTUAL_METHOD(...)
+#define MULTIPLEX_METHOD(...)
 
 // This pair of macros is used to help implement GENERATED_BODY() and GENERATED_USTRUCT_BODY()
 #define BODY_MACRO_COMBINE_INNER(A,B,C,D) A##B##C##D
