UHT: Enable error message for problematic multiplex declarations

diff --git a/Engine/Source/Programs/UnrealHeaderTool/Private/HeaderParser.cpp b/Engine/Source/Programs/UnrealHeaderTool/Private/HeaderParser.cpp
--- a/Engine/Source/Programs/UnrealHeaderTool/Private/HeaderParser.cpp
+++ b/Engine/Source/Programs/UnrealHeaderTool/Private/HeaderParser.cpp
@@ -2497,7 +2497,7 @@ UScriptStruct* FHeaderParser::CompileStructDeclaration(FClasses& AllClasses)
 		}
 		else if (Token.Matches(TEXT("MULTIPLEX_METHOD"), ESearchCase::CaseSensitive))
 		{
-			CompileStaticMethodDeclaration(AllClasses, Struct);
+			CompileMultiplexMethodDeclaration(AllClasses, Struct);
 		}
 		else if (Token.Matches(TEXT("GENERATED_USTRUCT_BODY")) || Token.Matches(TEXT("GENERATED_BODY")))
 		{
@@ -3006,6 +3006,48 @@ void FHeaderParser::FixupDelegateProperties( FClasses& AllClasses, UStruct* Stru
 	TMap<FName, FString> MetaData;
 	MetaData.Add(NAME_ToolTip, Struct->GetMetaData(NAME_ToolTip));
 	CheckDocumentationPolicyForStruct(Struct, MetaData);
+
+	// check if the struct has multi plex methods, and if so perform some additional checks
+	TArray<FMultiplexMethodInfo>* MultiplexMethods = StructMultiplexMethods.Find(Struct);
+	if (MultiplexMethods)
+	{
+		const TCHAR* InputText = TEXT("Input");
+		const TCHAR* OutputText = TEXT("Output");
+		const TCHAR* ConstantText = TEXT("Constant");
+		const TCHAR* MaxArraySizeText = TEXT("MaxArraySize");
+		const TCHAR* TArrayText = TEXT("TArray");
+
+		// validate the property types for this struct
+		for (TFieldIterator<UProperty> It(Struct); It; ++It)
+		{
+			UProperty const* const Prop = *It;
+			FString PropName = Prop->GetName();
+			FString MemberCPPType;
+			FString ExtendedCPPType;
+			MemberCPPType = Prop->GetCPPType(&ExtendedCPPType);
+
+			if (!ExtendedCPPType.IsEmpty())
+			{
+				// we only support arrays - no maps or similar data structures
+				if (MemberCPPType != TArrayText)
+				{
+					UE_LOG_ERROR_UHT(TEXT("Multiplex Struct '%s' - Member '%s' type '%s' not supported by Multiplex."), *Struct->GetName(), *PropName, *MemberCPPType);
+				}
+				else
+				{
+					bool bIsInput = Prop->HasMetaData(InputText);
+					bool bIsOutput = Prop->HasMetaData(OutputText);
+					bool bIsConstant = Prop->HasMetaData(ConstantText);
+					bIsConstant = bIsConstant || (bIsInput && !bIsOutput);
+
+					if (!bIsConstant && !Prop->HasMetaData(MaxArraySizeText))
+					{
+						UE_LOG_ERROR_UHT(TEXT("Multiplex Struct '%s' - Member '%s' requires the 'MaxArraySize' meta tag."), *Struct->GetName(), *PropName);
+					}
+				}
+			}
+		}
+	}
 }
 
 void FHeaderParser::VerifyBlueprintPropertyGetter(UProperty* Prop, UFunction* TargetFunc)
@@ -6165,11 +6207,11 @@ void FHeaderParser::CompileInterfaceDeclaration(FClasses& AllClasses)
 	PushNest(ENestType::Interface, InterfaceClass);
 }
 
-void FHeaderParser::CompileStaticMethodDeclaration(FClasses& AllClasses, UStruct* Struct)
+void FHeaderParser::CompileMultiplexMethodDeclaration(FClasses& AllClasses, UStruct* Struct)
 {
 	if (!MatchSymbol(TEXT("(")))
 	{
-		FError::Throwf(TEXT("Bad STATICMETHOD definition"));
+		FError::Throwf(TEXT("Bad MULTIPLEX_METHOD definition"));
 	}
 
 	// find the next close brace
@@ -6207,7 +6249,7 @@ void FHeaderParser::CompileStaticMethodDeclaration(FClasses& AllClasses, UStruct
 
 	if (!MatchSymbol(TEXT("(")))
 	{
-		FError::Throwf(TEXT("Bad STATICMETHOD definition"));
+		FError::Throwf(TEXT("Bad MULTIPLEX_METHOD definition"));
 	}
 
 	TArray<FString> ParamsContent;
diff --git a/Engine/Source/Programs/UnrealHeaderTool/Private/HeaderParser.h b/Engine/Source/Programs/UnrealHeaderTool/Private/HeaderParser.h
--- a/Engine/Source/Programs/UnrealHeaderTool/Private/HeaderParser.h
+++ b/Engine/Source/Programs/UnrealHeaderTool/Private/HeaderParser.h
@@ -572,7 +572,7 @@ protected:
 	void CompileFunctionDeclaration(FClasses& AllClasses);
 	void CompileVariableDeclaration (FClasses& AllClasses, UStruct* Struct);
 	void CompileInterfaceDeclaration(FClasses& AllClasses);
-	void CompileStaticMethodDeclaration(FClasses& AllClasses, UStruct* Struct);
+	void CompileMultiplexMethodDeclaration(FClasses& AllClasses, UStruct* Struct);
 
 	FClass* ParseInterfaceNameDeclaration(FClasses& AllClasses, FString& DeclaredInterfaceName, FString& RequiredAPIMacroIfPresent);
 	bool TryParseIInterfaceClass(FClasses& AllClasses);
diff --git a/Engine/Source/Programs/UnrealHeaderTool/Resources/UHTDebugging/StaticMethodStructHeader.h b/Engine/Source/Programs/UnrealHeaderTool/Resources/UHTDebugging/StaticMethodStructHeader.h
--- a/Engine/Source/Programs/UnrealHeaderTool/Resources/UHTDebugging/StaticMethodStructHeader.h
+++ b/Engine/Source/Programs/UnrealHeaderTool/Resources/UHTDebugging/StaticMethodStructHeader.h
@@ -42,7 +42,7 @@ struct FStaticMethodStruct : public FStaticMethodStructBase
 	UPROPERTY(meta = (Output))
 	FVector D;
 
-	UPROPERTY(meta = (Input, Output))
+	UPROPERTY(meta = (Input, Output, MaxArraySize = 8))
 	TArray<FVector> E;
 
 	UPROPERTY(meta = (Input))
@@ -51,6 +51,9 @@ struct FStaticMethodStruct : public FStaticMethodStructBase
 	UPROPERTY(meta = (Output, MaxArraySize = 8))
 	TArray<FVector> G;
 
+	UPROPERTY(meta = (MaxArraySize = 8))
+	TArray<FVector> H;
+
 	UPROPERTY()
 	float Cache;
 };
