RigVM: Restructuring Argument data structure

diff --git a/Engine/Source/Programs/UnrealHeaderTool/Private/CodeGenerator.cpp b/Engine/Source/Programs/UnrealHeaderTool/Private/CodeGenerator.cpp
--- a/Engine/Source/Programs/UnrealHeaderTool/Private/CodeGenerator.cpp
+++ b/Engine/Source/Programs/UnrealHeaderTool/Private/CodeGenerator.cpp
@@ -3323,10 +3323,10 @@ void FNativeClassHeaderGenerator::ExportGeneratedStructBodyMacros(FOutputDevice&
 		for (UProperty* Prop : TFieldRange<UProperty>(Struct))
 		{
 			FString Name = Prop->GetName();
-			FString AddressName = FString::Printf(TEXT("%s_Address"), *Name);
-			MemberMulticastProlog.Add(FString::Printf(TEXT("const FRigVMArgument& %s = RigVMAddresses[%d];"), *AddressName, MemberIndex++));
+			FString AddressName = FString::Printf(TEXT("%s_Arg"), *Name);
+			MemberMulticastProlog.Add(FString::Printf(TEXT("const FRigVMArgument& %s = RigVMArguments[%d];"), *AddressName, MemberIndex++));
 		}
-		MemberMulticastProlog.Insert(FString::Printf(TEXT("ensure(RigVMAddresses.Num() == %d);"), MemberMulticastProlog.Num()), 0);
+		MemberMulticastProlog.Insert(FString::Printf(TEXT("ensure(RigVMArguments.Num() == %d);"), MemberMulticastProlog.Num()), 0);
 		MemberMulticastProlog.Add(FString());
 
 		for (UProperty* Prop : TFieldRange<UProperty>(Struct))
@@ -3399,8 +3399,8 @@ void FNativeClassHeaderGenerator::ExportGeneratedStructBodyMacros(FOutputDevice&
 			}
 
 			MemberDeclarations.Add(FString::Printf(TEXT("%s %s"), *ParameterCPPType, *Name));
-			FString AddressName = FString::Printf(TEXT("%s_Address"), *Name);
-			MemberMulticastProlog.Add(FString::Printf(TEXT("%s %s = RigVMStorage[%s.StorageType()]->%s<%s>(%s.Index());"), *ExtractedMulticastCPPType, *Name, *AddressName, *MulticastGetter, *MultiCastCPPType, *AddressName));
+			FString AddressName = FString::Printf(TEXT("%s_Arg"), *Name);
+			MemberMulticastProlog.Add(FString::Printf(TEXT("%s %s = RigVMStorage[%s.StorageType()]->%s<%s>(%s.GetRegisterIndex());"), *ExtractedMulticastCPPType, *Name, *AddressName, *MulticastGetter, *MultiCastCPPType, *AddressName));
 		}
 
 		OutGeneratedHeaderText.Log(TEXT("\n"));
@@ -3431,7 +3431,7 @@ void FNativeClassHeaderGenerator::ExportGeneratedStructBodyMacros(FOutputDevice&
 			{
 				FString ParamsSuffix = Info.Params.IsEmpty() ? TEXT("") : FString::Printf(TEXT(", %s"), *Info.Params);
 				RigVMMethodsDeclarations += FString::Printf(TEXT("\tstatic %s Static%s(\r\n\t\t%s%s);\r\n"), *Info.ReturnType, *Info.Name, *FString::Join(MemberDeclarations, TEXT(",\r\n\t\t")), *ParamsSuffix);
-				RigVMMethodsDeclarations += FString::Printf(TEXT("\tstatic %s RigVM%s(const TArrayView<FRigVMArgument>& RigVMAddresses, FRigVMStorage** RigVMStorage, const TArrayView<void*>& AdditionalArgs);\r\n"), *Info.ReturnType, *Info.Name);
+				RigVMMethodsDeclarations += FString::Printf(TEXT("\tstatic %s RigVM%s(const FRigVMArgumentArray& RigVMArguments, FRigVMStoragePtrArray& RigVMStorage, const FRigVMUserDataArray& RigVMUserData);\r\n"), *Info.ReturnType, *Info.Name);
 			}
 		}
 
@@ -3681,7 +3681,7 @@ void FNativeClassHeaderGenerator::ExportGeneratedStructBodyMacros(FOutputDevice&
 				} 
 				Params.Add(ParamPrev.TrimStartAndEnd());
 
-				ParameterExtractionProlog.Add(FString::Printf(TEXT("ensure(AdditionalArgs.Num() == %d);"), Params.Num()));
+				ParameterExtractionProlog.Add(FString::Printf(TEXT("ensure(RigVMUserData.Num() == %d);"), Params.Num()));
 
 				for (int32 ParamIndex=0;ParamIndex<Params.Num(); ParamIndex++)
 				{
@@ -3715,7 +3715,7 @@ void FNativeClassHeaderGenerator::ExportGeneratedStructBodyMacros(FOutputDevice&
 						ParamPointerType = ParamPointerType.LeftChop(1);
 						ParamPointerType.TrimStartAndEndInline();
 					}
-					ParameterExtractionProlog.Add(FString::Printf(TEXT("%s = *(%s*)AdditionalArgs[%d];"), *FullParam, *ParamPointerType, ParamIndex));
+					ParameterExtractionProlog.Add(FString::Printf(TEXT("%s = *(%s*)RigVMUserData[%d];"), *FullParam, *ParamPointerType, ParamIndex));
 				}
 				ParamSuffix = FString::Printf(TEXT(", %s"), *FString::Join(Params, Comma));
 			}
@@ -3731,7 +3731,7 @@ void FNativeClassHeaderGenerator::ExportGeneratedStructBodyMacros(FOutputDevice&
 				ParamDeclarationSuffix = FString::Printf(TEXT(", %s"), *Info.Params);
 			}
 
-			Out.Logf(TEXT("%s %s::RigVM%s(const TArrayView<FRigVMArgument>& RigVMAddresses, FRigVMStorage** RigVMStorage, const TArrayView<void*>& AdditionalArgs)\r\n"), *Info.ReturnType, StructNameCPP, *Info.Name);
+			Out.Logf(TEXT("%s %s::RigVM%s(const FRigVMArgumentArray& RigVMArguments, FRigVMStoragePtrArray& RigVMStorage, const FRigVMUserDataArray& RigVMUserData)\r\n"), *Info.ReturnType, StructNameCPP, *Info.Name);
 			Out.Log(TEXT("{\r\n"));
 
 			if (ParameterExtractionProlog.Num() > 0)
diff --git a/Engine/Source/Runtime/RigVM/Private/RigVM.cpp b/Engine/Source/Runtime/RigVM/Private/RigVM.cpp
--- a/Engine/Source/Runtime/RigVM/Private/RigVM.cpp
+++ b/Engine/Source/Runtime/RigVM/Private/RigVM.cpp
@@ -2,6 +2,17 @@
 
 #include "RigVM.h"
 
+URigVM::URigVM()
+{
+	Literals.SetLiteralStorage(true);
+	WorkState.SetLiteralStorage(false);
+}
+
+URigVM::~URigVM()
+{
+	Reset();
+}
+
 void URigVM::Reset()
 {
 	Literals.Reset();
@@ -54,7 +65,7 @@ void URigVM::RefreshInstructionsIfRequired()
 	}
 }
 
-bool URigVM::Execute(FRigVMStorage** Storage, TArrayView<void*> AdditionalArgs)
+bool URigVM::Execute(FRigVMStoragePtrArray Storage, TArrayView<void*> AdditionalArgs)
 {
 	ResolveFunctionsIfRequired();
 	RefreshInstructionsIfRequired();
@@ -65,9 +76,9 @@ bool URigVM::Execute(FRigVMStorage** Storage, TArrayView<void*> AdditionalArgs)
 	}
 
 	FRigVMStorage* LocalStorage[] = { &WorkState, &Literals }; 
-	if (Storage == nullptr)
+	if (Storage.Num() == 0)
 	{
-		Storage = LocalStorage;
+		Storage = FRigVMStoragePtrArray(LocalStorage, 2);
 	}
 
 	int32 InstructionIndex = 0;
@@ -78,21 +89,21 @@ bool URigVM::Execute(FRigVMStorage** Storage, TArrayView<void*> AdditionalArgs)
 			case ERigVMOpCode::Copy:
 			{
 				const FRigVMCopyOp& Op = ByteCode.GetOpAt<FRigVMCopyOp>(Instructions[InstructionIndex]);
-				Storage[0]->Copy(Op.Source.Index(), Op.Target.Index(), Storage[Op.Source.StorageType()], Op.SourceOffset, Op.TargetOffset, Op.NumBytes);
+				Storage[0]->Copy(Op.Source.GetRegisterIndex(), Op.Target.GetRegisterIndex(), Storage[Op.Source.StorageType()], Op.SourceOffset, Op.TargetOffset, Op.NumBytes);
 				InstructionIndex++;
 				break;
 			}
 			case ERigVMOpCode::Increment:
 			{
 				const FRigVMIncrementOp& Op = ByteCode.GetOpAt<FRigVMIncrementOp>(Instructions[InstructionIndex]);
-				Storage[0]->GetRef<int32>(Op.Arg.Index())++;
+				Storage[0]->GetRef<int32>(Op.Arg.GetRegisterIndex())++;
 				InstructionIndex++;
 				break;
 			}
 			case ERigVMOpCode::Decrement:
 			{
 				const FRigVMDecrementOp& Op = ByteCode.GetOpAt<FRigVMDecrementOp>(Instructions[InstructionIndex]);
-				Storage[0]->GetRef<int32>(Op.Arg.Index())--;
+				Storage[0]->GetRef<int32>(Op.Arg.GetRegisterIndex())--;
 				InstructionIndex++;
 				break;
 			}
@@ -129,7 +140,7 @@ bool URigVM::Execute(FRigVMStorage** Storage, TArrayView<void*> AdditionalArgs)
 			case ERigVMOpCode::Execute:
 			{
 				const FRigVMExecuteOp& Op = ByteCode.GetOpAt<FRigVMExecuteOp>(Instructions[InstructionIndex]);
-				TArrayView<FRigVMArgument> Args = ByteCode.GetArgumentsForExecuteOp(Instructions[InstructionIndex]);
+				FRigVMArgumentArray Args = ByteCode.GetArgumentsForExecuteOp(Instructions[InstructionIndex]);
 				(*Functions[Op.FunctionIndex])(Args, Storage, AdditionalArgs);
 				InstructionIndex++;
 				break;
diff --git a/Engine/Source/Runtime/RigVM/Private/RigVMStorage.cpp b/Engine/Source/Runtime/RigVM/Private/RigVMStorage.cpp
--- a/Engine/Source/Runtime/RigVM/Private/RigVMStorage.cpp
+++ b/Engine/Source/Runtime/RigVM/Private/RigVMStorage.cpp
@@ -3,7 +3,8 @@
 #include "RigVMStorage.h"
 
 FRigVMStorage::FRigVMStorage(bool bInUseNames)
-	:bUseNameMap(bInUseNames)
+	: bUseNameMap(bInUseNames)
+	, bIsLiteralStorage(false)
 {
 }
 
diff --git a/Engine/Source/Runtime/RigVM/Public/RigVM.h b/Engine/Source/Runtime/RigVM/Public/RigVM.h
--- a/Engine/Source/Runtime/RigVM/Public/RigVM.h
+++ b/Engine/Source/Runtime/RigVM/Public/RigVM.h
@@ -15,9 +15,12 @@ class RIGVM_API URigVM : public UObject
 
 public:
 
+	URigVM();
+	virtual ~URigVM();
+
 	void Reset();
 
-	bool Execute(FRigVMStorage** Storage = nullptr, TArrayView<void*> AdditionalArgs = TArrayView<void*>());
+	bool Execute(FRigVMStoragePtrArray Storage = FRigVMStoragePtrArray(), TArrayView<void*> AdditionalArgs = TArrayView<void*>());
 
 	UFUNCTION()
 	int32 AddRigVMFunction(UScriptStruct* InRigVMStruct, const FName& InFunctionName);
diff --git a/Engine/Source/Runtime/RigVM/Public/RigVMRegistry.h b/Engine/Source/Runtime/RigVM/Public/RigVMRegistry.h
--- a/Engine/Source/Runtime/RigVM/Public/RigVMRegistry.h
+++ b/Engine/Source/Runtime/RigVM/Public/RigVMRegistry.h
@@ -5,7 +5,8 @@
 #include "CoreMinimal.h"
 #include "RigVMStorage.h"
 
-typedef void (*FRigVMFunctionPtr)(const TArrayView<FRigVMArgument>&, FRigVMStorage**, const TArrayView<void*>&);
+typedef TArrayView<void*> FRigVMUserDataArray;
+typedef void (*FRigVMFunctionPtr)(const FRigVMArgumentArray&, FRigVMStoragePtrArray&, const FRigVMUserDataArray&);
 
 struct RIGVM_API FRigVMFunction
 {
diff --git a/Engine/Source/Runtime/RigVM/Public/RigVMStorage.h b/Engine/Source/Runtime/RigVM/Public/RigVMStorage.h
--- a/Engine/Source/Runtime/RigVM/Public/RigVMStorage.h
+++ b/Engine/Source/Runtime/RigVM/Public/RigVMStorage.h
@@ -11,36 +11,33 @@ struct RIGVM_API FRigVMArgument
 public:
 
 	FRigVMArgument()
-		:Address(INDEX_NONE)
-	{
-	}
-
-	FRigVMArgument(int32 InAddress)
-		:Address(InAddress)
-	{
-	}
-
-	static FRigVMArgument MakeArgument(int32 InAddress)
+		: bIsLiteral(false)
+		, RegisterIndex(INDEX_NONE)
+		, ByteIndex(INDEX_NONE)
 	{
-		return FRigVMArgument(InAddress);
 	}
 
-	static FRigVMArgument MakeLiteral(int32 InAddress)
+	FRigVMArgument(bool InIsLiteral, int32 InRegisterIndex, int32 InByteIndex)
+		: bIsLiteral(InIsLiteral)
+		, RegisterIndex(InRegisterIndex)
+		, ByteIndex(InByteIndex)
 	{
-		return FRigVMArgument(-1 - InAddress);
 	}
 
-	FORCEINLINE bool operator ==(int32 InOther) const { return Address == InOther; }
-
-	FORCEINLINE bool IsLiteral() const { return Address < 0; }
-	FORCEINLINE int32 StorageType() const { return Address < 0 ? 1 : 0; }
-	FORCEINLINE int32 Index() const { return Address < 0 ? -(Address + 1) : Address; }
-	FORCEINLINE int32 NativeAddress() const { return Address; }
+	FORCEINLINE bool IsLiteral() const { return bIsLiteral; }
+	FORCEINLINE int32 StorageType() const { return bIsLiteral ? 1 : 0; }
+	FORCEINLINE int32 GetRegisterIndex() const { return RegisterIndex; }
+	FORCEINLINE int32 GetByteIndex() const { return ByteIndex; }
 
 private:
-	int32 Address;
+
+	bool bIsLiteral;
+	int32 RegisterIndex;
+	int32 ByteIndex;
 };
 
+typedef TArrayView<FRigVMArgument> FRigVMArgumentArray;
+
 UENUM()
 enum class ERigVMAddressType : uint8
 {
@@ -130,6 +127,8 @@ struct RIGVM_API FRigVMAddress
 	}
 };
 
+typedef TArrayView<FRigVMAddress> FRigVMAddressArray;
+
 USTRUCT()
 struct RIGVM_API FRigVMStorage
 {
@@ -143,12 +142,14 @@ public:
 
 	FRigVMStorage& operator= (const FRigVMStorage &InOther);
 
+	FORCEINLINE bool IsLiteralStorage() const { return bIsLiteralStorage;  }
+	FORCEINLINE void SetLiteralStorage(bool InIsLiteralStorage = true) { bIsLiteralStorage = InIsLiteralStorage; }
 	FORCEINLINE bool SupportsNames() const { return bUseNameMap;  }
 	FORCEINLINE int32 Num() const { return Addresses.Num(); }
 	FORCEINLINE const FRigVMAddress& operator[](int32 InIndex) const { return Addresses[InIndex]; }
 	FORCEINLINE FRigVMAddress& operator[](int32 InIndex) { return Addresses[InIndex]; }
-	FORCEINLINE const FRigVMAddress& operator[](const FRigVMArgument& InArg) const { return Addresses[InArg.Index()]; }
-	FORCEINLINE FRigVMAddress& operator[](const FRigVMArgument& InArg) { return Addresses[InArg.Index()]; }
+	FORCEINLINE const FRigVMAddress& operator[](const FRigVMArgument& InArg) const { return Addresses[InArg.GetRegisterIndex()]; }
+	FORCEINLINE FRigVMAddress& operator[](const FRigVMArgument& InArg) { return Addresses[InArg.GetRegisterIndex()]; }
 	FORCEINLINE const FRigVMAddress& operator[](const FName& InName) const { return Addresses[GetIndex(InName)]; }
 	FORCEINLINE FRigVMAddress& operator[](const FName& InName) { return Addresses[GetIndex(InName)]; }
 
@@ -157,6 +158,12 @@ public:
 	FORCEINLINE TArray<FRigVMAddress>::RangedForIteratorType      end() { return Addresses.end(); }
 	FORCEINLINE TArray<FRigVMAddress>::RangedForConstIteratorType end() const { return Addresses.end(); }
 
+	FORCEINLINE FRigVMArgument GetArgument(int32 InAddressIndex) const
+	{
+		ensure(Addresses.IsValidIndex(InAddressIndex));
+		return FRigVMArgument(IsLiteralStorage(), InAddressIndex, Addresses[InAddressIndex].ByteIndex);
+	}
+
 	FORCEINLINE const void* GetData(int32 InAddressIndex) const
 	{
 		ensure(Addresses.IsValidIndex(InAddressIndex));
@@ -224,7 +231,7 @@ public:
 	template<class T>
 	FORCEINLINE T& GetRef(const FRigVMArgument& InArgument)
 	{
-		return GetRef<T>(InArgument.Index());
+		return GetRef<T>(InArgument.GetRegisterIndex());
 	}
 
 	template<class T>
@@ -239,7 +246,7 @@ public:
 	template<class T>
 	FORCEINLINE TArrayView<T> GetArray(const FRigVMArgument& InArgument)
 	{
-		return GetArray<T>(InArgument.Index());
+		return GetArray<T>(InArgument.GetRegisterIndex());
 	}
 
 	FORCEINLINE UScriptStruct* GetScriptStruct(int32 InAddressIndex) const
@@ -533,6 +540,9 @@ private:
 	UPROPERTY()
 	bool bUseNameMap;
 
+	UPROPERTY()
+	bool bIsLiteralStorage;
+
 	UPROPERTY()
 	TArray<FRigVMAddress> Addresses;
 
@@ -545,3 +555,6 @@ private:
 	UPROPERTY(transient)
 	TMap<FName, int32> NameMap;
 };
+
+typedef FRigVMStorage* FRigVMStoragePtr;
+typedef TArrayView<FRigVMStoragePtr> FRigVMStoragePtrArray;
\ No newline at end of file
