RigVM: UHT refactoring - moving parsing to headerparser and more helper data structures

diff --git a/Engine/Source/Programs/UnrealHeaderTool/Private/CodeGenerator.cpp b/Engine/Source/Programs/UnrealHeaderTool/Private/CodeGenerator.cpp
--- a/Engine/Source/Programs/UnrealHeaderTool/Private/CodeGenerator.cpp
+++ b/Engine/Source/Programs/UnrealHeaderTool/Private/CodeGenerator.cpp
@@ -3301,117 +3301,56 @@ void FNativeClassHeaderGenerator::ExportGeneratedStructBodyMacros(FOutputDevice&
 
 	FString RigVMMacroPrefix(StructNameCPP);
 	RigVMMacroPrefix = FString::Printf(TEXT("UE_%s"), *RigVMMacroPrefix.Mid(1));
+	TArray<FString> RigVMVirtualFuncProlog, RigVMStubProlog;
 
-	const TCHAR* Comma = TEXT(", ");
-	TArray<FString> MemberNames, MemberDeclarations, MemberCasts, MemberCastProlog, MemberMulticastProlog;
-	const TArray<FRigVMMethodInfo>* StructRigVMMethods = FHeaderParser::StructRigVMMethods.Find(Struct);
-	if(StructRigVMMethods)
-	{
-		const TCHAR* FText = TEXT("F");
-		const TCHAR* TText = TEXT("T");
-		const TCHAR* EText = TEXT("E");
-		const TCHAR* InputText = TEXT("Input");
-		const TCHAR* OutputText = TEXT("Output");
-		const TCHAR* ConstantText = TEXT("Constant");
-		const TCHAR* MaxArraySizeText = TEXT("MaxArraySize");
-		const TCHAR* TArrayText = TEXT("TArray");
-		const TCHAR* TArrayViewText = TEXT("TArrayView");
-
-		int32 MemberCastSuffix = 0;
-
-		int32 MemberIndex = 0;
-		for (UProperty* Prop : TFieldRange<UProperty>(Struct))
-		{
-			FString Name = Prop->GetName();
-			FString AddressName = FString::Printf(TEXT("%s_Arg"), *Name);
-			MemberMulticastProlog.Add(FString::Printf(TEXT("const FRigVMArgument& %s = RigVMArguments[%d];"), *AddressName, MemberIndex++));
-		}
-		MemberMulticastProlog.Insert(FString::Printf(TEXT("ensure(RigVMArguments.Num() == %d);"), MemberMulticastProlog.Num()), 0);
-		MemberMulticastProlog.Add(FString());
-
-		for (UProperty* Prop : TFieldRange<UProperty>(Struct))
-		{
-			++MemberCastSuffix;
-
-			bool bIsInput = Prop->HasMetaData(InputText);
-			bool bIsOutput = Prop->HasMetaData(OutputText);
-			bool bIsConstant = Prop->HasMetaData(ConstantText);
-			bIsConstant = bIsConstant || (bIsInput && !bIsOutput);
-
-			FString Name = Prop->GetName();
-			MemberNames.Add(Name);
-
-			FString MemberCPPType;
-			FString ExtendedCPPType;
-			MemberCPPType = Prop->GetCPPType(&ExtendedCPPType);
-			MemberCPPType += ExtendedCPPType;
-			
-			FString ParameterCPPType = MemberCPPType;
-			FString MultiCastCPPType = MemberCPPType;
-			FString MulticastGetter = TEXT("GetRef");
-			FString ExtractedMulticastCPPType;
-			FString MemberCast = Name;
-
-			if (MemberCPPType.StartsWith(TArrayText))
-			{
-				if (bIsConstant || Prop->HasMetaData(MaxArraySizeText))
-				{
-					ParameterCPPType = FString::Printf(TEXT("%s%s"), TArrayViewText, *ExtendedCPPType);
-					if (!bIsConstant && Prop->HasMetaData(MaxArraySizeText))
-					{
-						MemberCast = FString::Printf(TEXT("%s_%d_View"), *Name, MemberCastSuffix);
-						MemberCastProlog.Add(FString::Printf(TEXT("%s.SetNumUninitialized( %s );"), *Name, *Prop->GetMetaData(MaxArraySizeText)));
-						MemberCastProlog.Add(FString::Printf(TEXT("%s %s(%s);"), *ParameterCPPType, *MemberCast, *Name));
-						MultiCastCPPType = ExtendedCPPType.LeftChop(1).RightChop(1);
-						MulticastGetter = TEXT("GetArray");
-						ExtractedMulticastCPPType = FString::Printf(TEXT("TArrayView<%s>"), *MultiCastCPPType);
-					}
-				}
-			}
+	// for RigVM methods we need to generated a macro used for implementing the static method
+	// and prepare two prologs: one for the virtual function implementation, and one for the stub
+	// invoking the static method.
+	const FRigVMStructInfo* StructRigVMInfo = FHeaderParser::StructRigVMMap.Find(Struct);
+	if(StructRigVMInfo)
+	{
+		RigVMStubProlog.Add(FString::Printf(TEXT("ensure(RigVMArguments.Num() == %d);"), StructRigVMInfo->Members.Num()));
+		for (int32 ParameterIndex = 0; ParameterIndex < StructRigVMInfo->Members.Num(); ParameterIndex++)
+		{
+			const FRigVMParameter& Parameter = StructRigVMInfo->Members[ParameterIndex];
+			RigVMStubProlog.Add(FString::Printf(TEXT("const FRigVMArgument& %s_Arg = RigVMArguments[%d];"), *Parameter.Name, ParameterIndex));
+		}
+		RigVMStubProlog.Add(FString());
 
-			if (ParameterCPPType.StartsWith(FText, ESearchCase::CaseSensitive) ||
-				ParameterCPPType.StartsWith(TText, ESearchCase::CaseSensitive) ||
-				ParameterCPPType.StartsWith(EText, ESearchCase::CaseSensitive))
+		for (int32 ParameterIndex = 0; ParameterIndex < StructRigVMInfo->Members.Num(); ParameterIndex++)
+		{
+			const FRigVMParameter& Parameter = StructRigVMInfo->Members[ParameterIndex];
+			if(Parameter.RequiresCast())
 			{
-				if (bIsConstant)
-				{
-					ParameterCPPType = FString::Printf(TEXT("const %s&"), *ParameterCPPType);
-				}
-				else
+				if (Parameter.IsArray() && !Parameter.IsConst() && !Parameter.MaxArraySize.IsEmpty())
 				{
-					ParameterCPPType = FString::Printf(TEXT("%s&"), *ParameterCPPType);
+					RigVMVirtualFuncProlog.Add(FString::Printf(TEXT("%s.SetNumUninitialized( %s );"), *Parameter.Name, *Parameter.MaxArraySize));
 				}
+				RigVMVirtualFuncProlog.Add(FString::Printf(TEXT("%s %s(%s);"), *Parameter.CastType, *Parameter.CastName, *Parameter.Name));
 			}
-			else if (bIsConstant)
-			{
-				ParameterCPPType = FString::Printf(TEXT("const %s"), *ParameterCPPType);
-			}
-			else
-			{
-				ParameterCPPType = FString::Printf(TEXT("%s&"), *ParameterCPPType);
-			}
-
-			MemberCasts.Add(MemberCast);
 
-			if (ExtractedMulticastCPPType.IsEmpty())
+			FString VariableType = Parameter.TypeVariableRef(true);
+			FString ExtractedType = Parameter.TypeOriginal();
+			if (Parameter.TypeOriginal(true).StartsWith(TEXT("TArrayView")))
 			{
-				ExtractedMulticastCPPType = ParameterCPPType;
+				ExtractedType = Parameter.ExtendedType().LeftChop(1).RightChop(1);
+				VariableType = Parameter.TypeOriginal(true);
 			}
 
-			MemberDeclarations.Add(FString::Printf(TEXT("%s %s"), *ParameterCPPType, *Name));
-			FString AddressName = FString::Printf(TEXT("%s_Arg"), *Name);
-			MemberMulticastProlog.Add(FString::Printf(TEXT("%s %s = RigVMStorage[%s.GetStorageIndex()]->%s<%s>(%s.GetRegisterIndex(), true);"), *ExtractedMulticastCPPType, *Name, *AddressName, *MulticastGetter, *MultiCastCPPType, *AddressName));
+			RigVMStubProlog.Add(FString::Printf(TEXT("%s %s = RigVMStorage[%s_Arg.GetStorageIndex()]->%s<%s>(%s_Arg.GetRegisterIndex(), true);"),
+				*VariableType,
+				*Parameter.NameOriginal(false),
+				*Parameter.NameOriginal(false),
+				*Parameter.Getter,
+				*ExtractedType,
+				*Parameter.NameOriginal(false)));
 		}
 
-		// we have a hard limit on the number of parameters for our function
-		if (MemberNames.Num() > 64)
-		{
-			UE_LOG_ERROR_UHT(TEXT("RIGVM_METHOD: Struct '%s' has more than 64 members (64 is the limit)."), StructNameCPP);
-		}
+		FString StructMembers = StructRigVMInfo->Members.Declarations(false, TEXT(", \\\r\n\t\t"), true, false);
 
 		OutGeneratedHeaderText.Log(TEXT("\n"));
-		OutGeneratedHeaderText.Logf(TEXT("#define %s_IMPLEMENT_RIGVM(ReturnType, FunctionName, ...) \\\n"), *RigVMMacroPrefix);
-		OutGeneratedHeaderText.Logf(TEXT("\tReturnType %s::Static##FunctionName(\\\r\n\t\t%s, ##__VA_ARGS__)\n"), StructNameCPP, *FString::Join(MemberDeclarations, TEXT(",\\\r\n\t\t")));
+		OutGeneratedHeaderText.Logf(TEXT("#define %s_IMPLEMENT_RIGVM(ReturnType, FunctionName, ...) \\\r\n"), *RigVMMacroPrefix);
+		OutGeneratedHeaderText.Logf(TEXT("\tReturnType %s::Static##FunctionName(\\\r\n\t\t%s,\\\r\n\t\t##__VA_ARGS__ \\\r\n\t)\n"), StructNameCPP, *StructMembers);
 		OutGeneratedHeaderText.Log(TEXT("\n"));
 	}
 
@@ -3430,14 +3369,17 @@ void FNativeClassHeaderGenerator::ExportGeneratedStructBodyMacros(FOutputDevice&
 		const FString StaticClassLine = FString::Printf(TEXT("\t%sstatic class UScriptStruct* StaticStruct();\r\n"), *RequiredAPI);
 		const FString PrivatePropertiesOffset = PrivatePropertiesOffsetGetters(Struct, StructNameCPP);
 		
+		// if we have RigVM methods on this struct we need to 
+		// declare the static method as well as the stub method
 		FString RigVMMethodsDeclarations;
-		if (StructRigVMMethods)
+		if (StructRigVMInfo)
 		{
-			for (const FRigVMMethodInfo& Info : (*StructRigVMMethods))
+			FString StructMembers = StructRigVMInfo->Members.Declarations(false, TEXT(",\r\n\t\t"), true, false);
+			for (const FRigVMMethodInfo& MethodInfo : StructRigVMInfo->Methods)
 			{
-				FString ParamsSuffix = Info.Params.IsEmpty() ? TEXT("") : FString::Printf(TEXT(", %s"), *Info.Params);
-				RigVMMethodsDeclarations += FString::Printf(TEXT("\tstatic %s Static%s(\r\n\t\t%s%s);\r\n"), *Info.ReturnType, *Info.Name, *FString::Join(MemberDeclarations, TEXT(",\r\n\t\t")), *ParamsSuffix);
-				RigVMMethodsDeclarations += FString::Printf(TEXT("\tstatic %s RigVM%s(const FRigVMArgumentArray& RigVMArguments, FRigVMStoragePtrArray& RigVMStorage, const FRigVMUserDataArray& RigVMUserData);\r\n"), *Info.ReturnType, *Info.Name);
+				FString ParameterSuffix = MethodInfo.Parameters.Declarations(true, TEXT(",\r\n\t\t"));
+				RigVMMethodsDeclarations += FString::Printf(TEXT("\tstatic %s Static%s(\r\n\t\t%s%s\r\n\t);\r\n"), *MethodInfo.ReturnType, *MethodInfo.Name, *StructMembers, *ParameterSuffix);
+				RigVMMethodsDeclarations += FString::Printf(TEXT("\tstatic %s RigVM%s(\r\n\t\tconst FRigVMArgumentArray& RigVMArguments,\r\n\t\tFRigVMStoragePtrArray& RigVMStorage,\r\n\t\tconst FRigVMUserDataArray& RigVMUserData\r\n\t);\r\n"), *MethodInfo.ReturnType, *MethodInfo.Name);
 			}
 		}
 
@@ -3475,12 +3417,14 @@ void FNativeClassHeaderGenerator::ExportGeneratedStructBodyMacros(FOutputDevice&
 		Out.Logf(TEXT("\t\tSingleton = GetStaticStruct(%s, %s, TEXT(\"%s\"), sizeof(%s), %s());\r\n"),
 			*SingletonName.LeftChop(2), *OuterName, *ActualStructName, StructNameCPP, *GetHashName);
 
-		if (StructRigVMMethods)
+		// if this struct has RigVM methods - we need to register the method to our central
+		// registry on construction of the static struct
+		if (StructRigVMInfo)
 		{
-			for (const FRigVMMethodInfo& Info : (*StructRigVMMethods))
+			for (const FRigVMMethodInfo& MethodInfo : StructRigVMInfo->Methods)
 			{
 				Out.Logf(TEXT("\t\tFRigVMRegistry::Get().Register(TEXT(\"%s::%s\"), &%s::RigVM%s);\r\n"),
-					StructNameCPP, *Info.Name, StructNameCPP, *Info.Name);
+					StructNameCPP, *MethodInfo.Name, StructNameCPP, *MethodInfo.Name);
 			}
 		}
 
@@ -3656,107 +3600,65 @@ void FNativeClassHeaderGenerator::ExportGeneratedStructBodyMacros(FOutputDevice&
 	Out.Log(GeneratedStructRegisterFunctionText);
 	Out.Logf(TEXT("\tuint32 %s() { return %uU; }\r\n"), *HashFuncName, StructHash);
 
-	if (StructRigVMMethods)
+	// if this struct has RigVM methods we need to implement both the 
+	// virtual function as well as the stub method here.
+	// The static method is implemented by the user using the IMPLEMENT_RIGVM macro.
+	if (StructRigVMInfo)
 	{
-		for (const FRigVMMethodInfo& Info : (*StructRigVMMethods))
+		FString StructMembersForVirtualFunc = StructRigVMInfo->Members.Names(false, TEXT(",\r\n\t\t"), true);
+		FString StructMembersForStub = StructRigVMInfo->Members.Names(false, TEXT(",\r\n\t\t"), false);
+
+		for (const FRigVMMethodInfo& MethodInfo : StructRigVMInfo->Methods)
 		{
 			Out.Log(TEXT("\r\n"));
 
+			FString ParameterDeclaration = MethodInfo.Parameters.Declarations(false, TEXT(",\r\n\t\t"));
+			FString ParameterSuffix = MethodInfo.Parameters.Names(true, TEXT(",\r\n\t\t"));
+
 			// implement the virtual function body.
-			Out.Logf(TEXT("%s %s::%s(%s)\r\n"), *Info.ReturnType, StructNameCPP, *Info.Name, *Info.Params);
+			Out.Logf(TEXT("%s %s::%s(%s)\r\n"), *MethodInfo.ReturnType, StructNameCPP, *MethodInfo.Name, *ParameterDeclaration);
 			Out.Log(TEXT("{\r\n"));
-			for (const FString& MemberCastPrologLine : MemberCastProlog)
-			{
-				Out.Logf(TEXT("\t%s\r\n"), *MemberCastPrologLine);
-			}
 
-			FString ReturnPrefix = (Info.ReturnType.IsEmpty() || (Info.ReturnType == TEXT("void"))) ? TEXT("") : TEXT("return ");
-			FString ParamSuffix;
-
-			TArray<FString> ParameterExtractionProlog;
-			if (!Info.Params.IsEmpty())
+			if(RigVMVirtualFuncProlog.Num() > 0)
 			{
-				TArray<FString> Params;
-				FString ParamPrev, ParamLeft, ParamRight;
-				ParamPrev = Info.Params;
-				while (ParamPrev.Contains(TEXT(",")))
-				{
-					ParamPrev.Split(TEXT(","), &ParamLeft, &ParamRight);
-					Params.Add(ParamLeft.TrimStartAndEnd());
-					ParamPrev = ParamRight;
-				} 
-				Params.Add(ParamPrev.TrimStartAndEnd());
-
-				ParameterExtractionProlog.Add(FString::Printf(TEXT("ensure(RigVMUserData.Num() == %d);"), Params.Num()));
-
-				for (int32 ParamIndex=0;ParamIndex<Params.Num(); ParamIndex++)
+				for (const FString& RigVMVirtualFuncPrologLine : RigVMVirtualFuncProlog)
 				{
-					FString FullParam = Params[ParamIndex];
-					int32 LastEqual = INDEX_NONE;
-					if (FullParam.FindLastChar(TCHAR('='), LastEqual))
-					{
-						FullParam = FullParam.Mid(0, LastEqual);
-					}
-
-					FullParam.TrimStartAndEndInline();
-
-					FString ParamType = FullParam;
-					FString ParamName = FullParam;
-
-					int32 LastSpace = INDEX_NONE;
-					if (ParamName.FindLastChar(TCHAR(' '), LastSpace))
-					{
-						ParamType = ParamType.Mid(0, LastSpace);
-						ParamName = ParamName.Mid(LastSpace + 1);
-					}
-
-					ParamType.TrimStartAndEndInline();
-					ParamName.TrimStartAndEndInline();
-
-					Params[ParamIndex] = ParamName;
-
-					FString ParamPointerType = ParamType;
-					if (ParamPointerType.EndsWith(TEXT("&")))
-					{
-						ParamPointerType = ParamPointerType.LeftChop(1);
-						ParamPointerType.TrimStartAndEndInline();
-					}
-					ParameterExtractionProlog.Add(FString::Printf(TEXT("%s = *(%s*)RigVMUserData[%d];"), *FullParam, *ParamPointerType, ParamIndex));
+					Out.Logf(TEXT("\t%s\r\n"), *RigVMVirtualFuncPrologLine);
 				}
-				ParamSuffix = FString::Printf(TEXT(", %s"), *FString::Join(Params, Comma));
+				Out.Log(TEXT("\t\r\n"));
 			}
 
-			Out.Logf(TEXT("    %sStatic%s(%s%s);\n"), *ReturnPrefix, *Info.Name, *FString::Join(MemberCasts, Comma), *ParamSuffix);
+			Out.Logf(TEXT("    %sStatic%s(\r\n\t\t%s%s\r\n\t);\n"), *MethodInfo.ReturnPrefix(), *MethodInfo.Name, *StructMembersForVirtualFunc, *ParameterSuffix);
 			Out.Log(TEXT("}\r\n"));
 
 			Out.Log(TEXT("\r\n"));
 
-			FString ParamDeclarationSuffix;
-			if (!Info.Params.IsEmpty())
-			{
-				ParamDeclarationSuffix = FString::Printf(TEXT(", %s"), *Info.Params);
-			}
-
-			Out.Logf(TEXT("%s %s::RigVM%s(const FRigVMArgumentArray& RigVMArguments, FRigVMStoragePtrArray& RigVMStorage, const FRigVMUserDataArray& RigVMUserData)\r\n"), *Info.ReturnType, StructNameCPP, *Info.Name);
+			// implement stub method body
+			Out.Logf(TEXT("%s %s::RigVM%s(const FRigVMArgumentArray& RigVMArguments, FRigVMStoragePtrArray& RigVMStorage, const FRigVMUserDataArray& RigVMUserData)\r\n"), *MethodInfo.ReturnType, StructNameCPP, *MethodInfo.Name);
 			Out.Log(TEXT("{\r\n"));
 
-			if (ParameterExtractionProlog.Num() > 0)
+			if (MethodInfo.Parameters.Num() > 0)
 			{
-				for (const FString& ParameterExtractionPrologLine : ParameterExtractionProlog)
+				Out.Logf(TEXT("\tensure(RigVMUserData.Num() == %d);\r\n"), MethodInfo.Parameters.Num());
+				for (int32 ParameterIndex = 0; ParameterIndex < MethodInfo.Parameters.Num(); ParameterIndex++)
 				{
-					Out.Logf(TEXT("\t%s\r\n"), *ParameterExtractionPrologLine);
+					const FRigVMParameter& Parameter = MethodInfo.Parameters[ParameterIndex];
+					Out.Logf(TEXT("\t%s = *(%s*)RigVMUserData[%d];\r\n"), *Parameter.Declaration(), *Parameter.TypeNoRef(), ParameterIndex);
 				}
 				Out.Log(TEXT("\t\r\n"));
 			}
 
-			for (const FString& MemberMulticastPrologLine : MemberMulticastProlog)
+			if(RigVMStubProlog.Num() > 0)
 			{
-				Out.Logf(TEXT("\t%s\r\n"), *MemberMulticastPrologLine);
+				for (const FString& RigVMStubPrologLine : RigVMStubProlog)
+				{
+					Out.Logf(TEXT("\t%s\r\n"), *RigVMStubPrologLine);
+				}
+				Out.Log(TEXT("\t\r\n"));
 			}
 
-			Out.Logf(TEXT("\t%sStatic%s(%s%s);\r\n"), *ReturnPrefix, *Info.Name, *FString::Join(MemberNames, Comma), *ParamSuffix);
+			Out.Logf(TEXT("\t%sStatic%s(\r\n\t\t%s%s\r\n\t);\r\n"), *MethodInfo.ReturnPrefix(), *MethodInfo.Name, *StructMembersForStub, *ParameterSuffix);
 			Out.Log(TEXT("}\r\n"));
-
 		}
 
 		Out.Log(TEXT("\r\n"));
diff --git a/Engine/Source/Programs/UnrealHeaderTool/Private/HeaderParser.cpp b/Engine/Source/Programs/UnrealHeaderTool/Private/HeaderParser.cpp
--- a/Engine/Source/Programs/UnrealHeaderTool/Private/HeaderParser.cpp
+++ b/Engine/Source/Programs/UnrealHeaderTool/Private/HeaderParser.cpp
@@ -72,7 +72,7 @@ static const FName NAME_DocumentationPolicy(TEXT("DocumentationPolicy"));
 EGeneratedCodeVersion FHeaderParser::DefaultGeneratedCodeVersion = EGeneratedCodeVersion::V1;
 TArray<FString> FHeaderParser::StructsWithNoPrefix;
 TArray<FString> FHeaderParser::StructsWithTPrefix;
-TMap<UStruct*, TArray<FRigVMMethodInfo>> FHeaderParser::StructRigVMMethods;
+FRigVMStructMap FHeaderParser::StructRigVMMap;
 TArray<FString> FHeaderParser::DelegateParameterCountStrings;
 TMap<FString, FString> FHeaderParser::TypeRedirectMap;
 TArray<FString> FHeaderParser::PropertyCPPTypesRequiringUIRanges = { TEXT("float"), TEXT("double") };
@@ -3007,47 +3007,7 @@ void FHeaderParser::FixupDelegateProperties( FClasses& AllClasses, UStruct* Stru
 	MetaData.Add(NAME_ToolTip, Struct->GetMetaData(NAME_ToolTip));
 	CheckDocumentationPolicyForStruct(Struct, MetaData);
 
-	// check if the struct has multi plex methods, and if so perform some additional checks
-	TArray<FRigVMMethodInfo>* RigVMMethods = StructRigVMMethods.Find(Struct);
-	if (RigVMMethods)
-	{
-		const TCHAR* InputText = TEXT("Input");
-		const TCHAR* OutputText = TEXT("Output");
-		const TCHAR* ConstantText = TEXT("Constant");
-		const TCHAR* MaxArraySizeText = TEXT("MaxArraySize");
-		const TCHAR* TArrayText = TEXT("TArray");
-
-		// validate the property types for this struct
-		for (TFieldIterator<UProperty> It(Struct); It; ++It)
-		{
-			UProperty const* const Prop = *It;
-			FString PropName = Prop->GetName();
-			FString MemberCPPType;
-			FString ExtendedCPPType;
-			MemberCPPType = Prop->GetCPPType(&ExtendedCPPType);
-
-			if (!ExtendedCPPType.IsEmpty())
-			{
-				// we only support arrays - no maps or similar data structures
-				if (MemberCPPType != TArrayText)
-				{
-					UE_LOG_ERROR_UHT(TEXT("RigVM Struct '%s' - Member '%s' type '%s' not supported by RigVM."), *Struct->GetName(), *PropName, *MemberCPPType);
-				}
-				else
-				{
-					bool bIsInput = Prop->HasMetaData(InputText);
-					bool bIsOutput = Prop->HasMetaData(OutputText);
-					bool bIsConstant = Prop->HasMetaData(ConstantText);
-					bIsConstant = bIsConstant || (bIsInput && !bIsOutput);
-
-					if (!bIsConstant && !Prop->HasMetaData(MaxArraySizeText))
-					{
-						UE_LOG_ERROR_UHT(TEXT("RigVM Struct '%s' - Member '%s' requires the 'MaxArraySize' meta tag."), *Struct->GetName(), *PropName);
-					}
-				}
-			}
-		}
-	}
+	ParseRigVMMethodParameters(Struct);
 }
 
 void FHeaderParser::VerifyBlueprintPropertyGetter(UProperty* Prop, UFunction* TargetFunc)
@@ -6271,13 +6231,131 @@ void FHeaderParser::CompileRigVMMethodDeclaration(FClasses& AllClasses, UStruct*
 		}
 	}
 
-	FRigVMMethodInfo Info;
-	Info.ReturnType = ReturnTypeToken.Identifier;
-	Info.Name = NameToken.Identifier;
-	Info.Params = FString::Join(ParamsContent, TEXT(" "));
+	FRigVMMethodInfo MethodInfo;
+	MethodInfo.ReturnType = ReturnTypeToken.Identifier;
+	MethodInfo.Name = NameToken.Identifier;
+	
+	FString ParamString = FString::Join(ParamsContent, TEXT(" "));
+	if (!ParamString.IsEmpty())
+	{
+		FString ParamPrev, ParamLeft, ParamRight;
+		ParamPrev = ParamString;
+		while (ParamPrev.Contains(TEXT(",")))
+		{
+			ParamPrev.Split(TEXT(","), &ParamLeft, &ParamRight);
+			FRigVMParameter Parameter;
+			Parameter.Name = ParamLeft.TrimStartAndEnd();
+			MethodInfo.Parameters.Add(Parameter);
+			ParamPrev = ParamRight;
+		}
 
-	TArray<FRigVMMethodInfo>& Infos = StructRigVMMethods.FindOrAdd(Struct);
-	Infos.Add(Info);
+		ParamPrev = ParamPrev.TrimStartAndEnd();
+		if (!ParamPrev.IsEmpty())
+		{
+			FRigVMParameter Parameter;
+			Parameter.Name = ParamPrev.TrimStartAndEnd();
+			MethodInfo.Parameters.Add(Parameter);
+		}
+	}
+
+	for (FRigVMParameter& Parameter : MethodInfo.Parameters)
+	{
+		FString FullParameter = Parameter.Name;
+
+		int32 LastEqual = INDEX_NONE;
+		if (FullParameter.FindLastChar(TCHAR('='), LastEqual))
+		{
+			FullParameter = FullParameter.Mid(0, LastEqual);
+		}
+
+		FullParameter.TrimStartAndEndInline();
+
+		FString ParameterType = FullParameter;
+		FString ParameterName = FullParameter;
+
+		int32 LastSpace = INDEX_NONE;
+		if (FullParameter.FindLastChar(TCHAR(' '), LastSpace))
+		{
+			Parameter.Type = FullParameter.Mid(0, LastSpace);
+			Parameter.Name = FullParameter.Mid(LastSpace + 1);
+			Parameter.Type.TrimStartAndEndInline();
+			Parameter.Name.TrimStartAndEndInline();
+		}
+	}
+
+	FRigVMStructInfo& StructRigVMInfo = StructRigVMMap.FindOrAdd(Struct);
+	StructRigVMInfo.Name = Struct->GetName();
+	StructRigVMInfo.Methods.Add(MethodInfo);
+}
+
+void FHeaderParser::ParseRigVMMethodParameters(UStruct* Struct)
+{
+	FRigVMStructInfo* StructRigVMInfo = StructRigVMMap.Find(Struct);
+	if (StructRigVMInfo == nullptr)
+	{
+		return;
+	}
+
+	const TCHAR* InputText = TEXT("Input");
+	const TCHAR* OutputText = TEXT("Output");
+	const TCHAR* ConstantText = TEXT("Constant");
+	const TCHAR* MaxArraySizeText = TEXT("MaxArraySize");
+	const TCHAR* TArrayText = TEXT("TArray");
+	const TCHAR* TArrayViewText = TEXT("TArrayView");
+	const TCHAR* GetRefText = TEXT("GetRef");
+	const TCHAR* GetArrayText = TEXT("GetArray");
+
+	// validate the property types for this struct
+	for (TFieldIterator<UProperty> It(Struct); It; ++It)
+	{
+		UProperty const* const Prop = *It;
+		FString PropName = Prop->GetName();
+		FString MemberCPPType;
+		FString ExtendedCPPType;
+		MemberCPPType = Prop->GetCPPType(&ExtendedCPPType);
+
+		FRigVMParameter Parameter;
+		Parameter.Name = Prop->GetName();
+		Parameter.Type = MemberCPPType + ExtendedCPPType;
+		Parameter.bConstant = Prop->HasMetaData(ConstantText);
+		Parameter.bInput = Prop->HasMetaData(InputText);
+		Parameter.bOutput = Prop->HasMetaData(OutputText);
+		Parameter.MaxArraySize = Prop->GetMetaData(MaxArraySizeText);
+		Parameter.Getter = GetRefText;
+
+		if (!ExtendedCPPType.IsEmpty())
+		{
+			// we only support arrays - no maps or similar data structures
+			if (MemberCPPType != TArrayText)
+			{
+				UE_LOG_ERROR_UHT(TEXT("RigVM Struct '%s' - Member '%s' type '%s' not supported by RigVM."), *Struct->GetName(), *Parameter.Name, *MemberCPPType);
+				continue;
+			}
+
+			if (!Parameter.IsConst() && Parameter.MaxArraySize.IsEmpty())
+			{
+				UE_LOG_ERROR_UHT(TEXT("RigVM Struct '%s' - Member '%s' requires the 'MaxArraySize' meta tag."), *Struct->GetName(), *Parameter.Name);
+				continue;
+			}
+		}
+
+		if (MemberCPPType.StartsWith(TArrayText))
+		{
+			if (Parameter.IsConst() || !Parameter.MaxArraySize.IsEmpty())
+			{
+				Parameter.CastName = FString::Printf(TEXT("%s_%d_View"), *Parameter.Name, StructRigVMInfo->Members.Num());
+				Parameter.CastType = FString::Printf(TEXT("%s%s"), TArrayViewText, *ExtendedCPPType);
+				Parameter.Getter = GetArrayText;
+			}
+		}
+
+		StructRigVMInfo->Members.Add(Parameter);
+	}
+
+	if (StructRigVMInfo->Members.Num() > 64)
+	{
+		UE_LOG_ERROR_UHT(TEXT("RigVM Struct '%s' - has %d members (64 is the limit)."), *Struct->GetName(), StructRigVMInfo->Members.Num());
+	}
 }
 
 // Returns true if the token is a dynamic delegate declaration
diff --git a/Engine/Source/Programs/UnrealHeaderTool/Private/HeaderParser.h b/Engine/Source/Programs/UnrealHeaderTool/Private/HeaderParser.h
--- a/Engine/Source/Programs/UnrealHeaderTool/Private/HeaderParser.h
+++ b/Engine/Source/Programs/UnrealHeaderTool/Private/HeaderParser.h
@@ -117,13 +117,196 @@ struct FIndexRange
 	int32 Count;
 };
 
+struct FRigVMParameter
+{
+	FRigVMParameter()
+		: Name()
+		, Type()
+		, bConstant(false)
+		, bInput(false)
+		, bOutput(false)
+		, MaxArraySize()
+	{
+	}
+
+	FString Name;
+	FString Type;
+	bool bConstant;
+	bool bInput;
+	bool bOutput;
+	FString MaxArraySize;
+	FString Getter;
+	FString CastName;
+	FString CastType;
+
+	const FString& NameOriginal(bool bCastName = false) const
+	{
+		return (bCastName && !CastName.IsEmpty()) ? CastName : Name;
+	}
+
+	const FString& TypeOriginal(bool bCastType = false) const
+	{
+		return (bCastType && !CastType.IsEmpty()) ? CastType : Type;
+	}
+
+	FString Declaration(bool bCastType = false, bool bCastName = false) const
+	{
+		return FString::Printf(TEXT("%s %s"), *TypeOriginal(bCastType), *NameOriginal(bCastName));
+	}
+
+	FString BaseType(bool bCastType = false) const
+	{
+		const FString& String = TypeOriginal(bCastType);
+		int32 LesserPos = 0;
+		if (String.FindChar('<', LesserPos))
+		{
+			return String.Mid(0, LesserPos);
+		}
+		return String;
+	}
+
+	FString ExtendedType(bool bCastType = false) const
+	{
+		const FString& String = TypeOriginal(bCastType);
+		int32 LesserPos = 0;
+		if(String.FindChar('<', LesserPos))
+		{
+			return String.Mid(LesserPos);
+		}
+		return String;
+	}
+
+	FString TypeConstRef(bool bCastType = false) const
+	{
+		const FString& String = TypeNoRef(bCastType);
+		if (String.StartsWith(TEXT("T"), ESearchCase::CaseSensitive) || String.StartsWith(TEXT("F"), ESearchCase::CaseSensitive))
+		{
+			return FString::Printf(TEXT("const %s&"), *String);
+		}
+		return FString::Printf(TEXT("const %s"), *String);
+	}
+
+	FString TypeRef(bool bCastType = false) const
+	{
+		const FString& String = TypeNoRef(bCastType);
+		return FString::Printf(TEXT("%s&"), *String);
+	}
+
+	FString TypeNoRef(bool bCastType = false) const
+	{
+		const FString& String = TypeOriginal(bCastType);
+		if (String.EndsWith(TEXT("&")))
+		{
+			return String.LeftChop(1);
+		}
+		return String;
+	}
+
+	FString TypeVariableRef(bool bCastType = false) const
+	{
+		return IsConst() ? TypeConstRef(bCastType) : TypeRef(bCastType);
+	}
+
+	FString Variable(bool bCastType = false, bool bCastName = false) const
+	{
+		return FString::Printf(TEXT("%s %s"), *TypeVariableRef(bCastType), *NameOriginal(bCastName));
+	}
+
+	bool IsConst() const
+	{
+		return bConstant || (bInput && !bOutput);
+	}
+
+	bool IsArray() const
+	{
+		return BaseType().Equals(TEXT("TArray"));
+	}
+
+	bool RequiresCast() const
+	{
+		return !CastType.IsEmpty() && !CastName.IsEmpty();
+	}
+};
+
+struct FRigVMParameterArray
+{
+public:
+	int32 Num() const { return Parameters.Num(); }
+	const FRigVMParameter& operator[](int32 InIndex) const { return Parameters[InIndex]; }
+	FRigVMParameter& operator[](int32 InIndex) { return Parameters[InIndex]; }
+	TArray<FRigVMParameter>::RangedForConstIteratorType begin() const { return Parameters.begin(); }
+	TArray<FRigVMParameter>::RangedForConstIteratorType end() const { return Parameters.end(); }
+	TArray<FRigVMParameter>::RangedForIteratorType begin() { return Parameters.begin(); }
+	TArray<FRigVMParameter>::RangedForIteratorType end() { return Parameters.end(); }
+
+	int32 Add(const FRigVMParameter& InParameter)
+	{
+		return Parameters.Add(InParameter);
+	}
+
+	FString Names(bool bLeadingSeparator = false, TCHAR* Separator = TEXT(", "), bool bCastType = false) const
+	{
+		if (Parameters.Num() == 0)
+		{
+			return FString();
+		}
+		TArray<FString> NameArray;
+		for (const FRigVMParameter& Parameter : Parameters)
+		{
+			NameArray.Add(Parameter.NameOriginal(bCastType));
+		}
+		FString Joined = FString::Join(NameArray, Separator);
+		if (bLeadingSeparator)
+		{
+			return FString::Printf(TEXT("%s%s"), Separator, *Joined);
+		}
+		return Joined;
+	}
+
+	FString Declarations(bool bLeadingSeparator = false, TCHAR* Separator = TEXT(", "), bool bCastType = false, bool bCastName = false) const
+	{
+		if (Parameters.Num() == 0)
+		{
+			return FString();
+		}
+		TArray<FString> DeclarationArray;
+		for (const FRigVMParameter& Parameter : Parameters)
+		{
+			DeclarationArray.Add(Parameter.Variable(bCastType, bCastName));
+		}
+		FString Joined = FString::Join(DeclarationArray, Separator);
+		if (bLeadingSeparator)
+		{
+			return FString::Printf(TEXT("%s%s"), Separator, *Joined);
+		}
+		return Joined;
+	}
+
+private:
+	TArray<FRigVMParameter> Parameters;
+};
+
 struct FRigVMMethodInfo
 {
 	FString ReturnType;
 	FString Name;
-	FString Params;
+	FRigVMParameterArray Parameters;
+
+	FString ReturnPrefix() const
+	{
+		return (ReturnType.IsEmpty() || (ReturnType == TEXT("void"))) ? TEXT("") : TEXT("return ");
+	}
+};
+
+struct FRigVMStructInfo
+{
+	FString Name;
+	FRigVMParameterArray Members;
+	TArray<FRigVMMethodInfo> Methods;
 };
 
+typedef TMap<UStruct*, FRigVMStructInfo> FRigVMStructMap;
+
 struct ClassDefinitionRange
 {
 	ClassDefinitionRange(const TCHAR* InStart, const TCHAR* InEnd)
@@ -427,7 +610,7 @@ protected:
 	TMap<int32, FString> RPCsNeedingHookup;
 
 	// List of all multiplex methods defined on structs
-	static TMap<UStruct*, TArray<FRigVMMethodInfo>> StructRigVMMethods;
+	static FRigVMStructMap StructRigVMMap;
 
 	// Constructor.
 	explicit FHeaderParser(FFeedbackContext* InWarn, const FManifestModule& InModule);
@@ -573,6 +756,7 @@ protected:
 	void CompileVariableDeclaration (FClasses& AllClasses, UStruct* Struct);
 	void CompileInterfaceDeclaration(FClasses& AllClasses);
 	void CompileRigVMMethodDeclaration(FClasses& AllClasses, UStruct* Struct);
+	void ParseRigVMMethodParameters(UStruct* Struct);
 
 	FClass* ParseInterfaceNameDeclaration(FClasses& AllClasses, FString& DeclaredInterfaceName, FString& RequiredAPIMacroIfPresent);
 	bool TryParseIInterfaceClass(FClasses& AllClasses);
diff --git a/Engine/Source/Programs/UnrealHeaderTool/Resources/UHTDebugging.manifest b/Engine/Source/Programs/UnrealHeaderTool/Resources/UHTDebugging.manifest
--- a/Engine/Source/Programs/UnrealHeaderTool/Resources/UHTDebugging.manifest
+++ b/Engine/Source/Programs/UnrealHeaderTool/Resources/UHTDebugging.manifest
@@ -35,7 +35,7 @@
 				".\\UHTDebugging\\StructOnlyHeader.h",
 				".\\UHTDebugging\\EnumOnlyHeader.h",
 				".\\UHTDebugging\\TestObjectSerializers.h",
-				".\\UHTDebugging\\StaticMethodStructHeader.h"
+				".\\UHTDebugging\\RigVMStructHeader.h"
 			],
 			"PublicHeaders": [
 			],
diff --git a/Engine/Source/Programs/UnrealHeaderTool/Resources/UHTDebugging/RigVMStructHeader.h b/Engine/Source/Programs/UnrealHeaderTool/Resources/UHTDebugging/RigVMStructHeader.h
new file mode 100644
--- /dev/null
+++ b/Engine/Source/Programs/UnrealHeaderTool/Resources/UHTDebugging/RigVMStructHeader.h
@@ -0,0 +1,59 @@
+// Copyright 1998-2019 Epic Games, Inc. All Rights Reserved.
+
+#pragma once
+
+#include "RigVMStructHeader.generated.h"
+
+USTRUCT()
+struct FRigVMStructBase
+{
+	GENERATED_BODY()
+		
+	UPROPERTY(meta = (Input))
+	float Inherited;
+
+	UPROPERTY(meta = (Output))
+	float InheritedOutput;
+};
+
+USTRUCT()
+struct FRigVMMethodStruct : public FRigVMStructBase
+{
+	GENERATED_BODY()
+
+	RIGVM_METHOD()
+	void Clear();
+
+	RIGVM_METHOD()
+	virtual void Execute(bool bAdditionalFlag = false, const FString& InString = TEXT("")) override;
+
+	RIGVM_METHOD()
+	void Compute(float TestFloat);
+
+	UPROPERTY(meta = (Input))
+	float A;
+
+	UPROPERTY(meta = (Output))
+	float B;
+
+	UPROPERTY(meta = (Input))
+	FVector C;
+
+	UPROPERTY(meta = (Output))
+	FVector D;
+
+	UPROPERTY(meta = (Input, Output, MaxArraySize = 8))
+	TArray<FVector> E;
+
+	UPROPERTY(meta = (Input))
+	TArray<FVector> F;
+
+	UPROPERTY(meta = (Output, MaxArraySize = 8))
+	TArray<FVector> G;
+
+	UPROPERTY(meta = (MaxArraySize = 8))
+	TArray<FVector> H;
+
+	UPROPERTY()
+	float Cache;
+};
diff --git a/Engine/Source/Programs/UnrealHeaderTool/Resources/UHTDebugging/StaticMethodStructHeader.h b/Engine/Source/Programs/UnrealHeaderTool/Resources/UHTDebugging/StaticMethodStructHeader.h
deleted file mode 100644
