RigVM: Ensure to forbid WITH_EDITORONLY_DATA for RIGVM_METHOD structs.

diff --git a/Engine/Source/Programs/UnrealHeaderTool/Private/HeaderParser.cpp b/Engine/Source/Programs/UnrealHeaderTool/Private/HeaderParser.cpp
--- a/Engine/Source/Programs/UnrealHeaderTool/Private/HeaderParser.cpp
+++ b/Engine/Source/Programs/UnrealHeaderTool/Private/HeaderParser.cpp
@@ -6322,6 +6322,12 @@ void FHeaderParser::ParseRigVMMethodParameters(UStruct* Struct)
 		Parameter.bOutput = Prop->HasMetaData(OutputText);
 		Parameter.MaxArraySize = Prop->GetMetaData(MaxArraySizeText);
 		Parameter.Getter = GetRefText;
+		Parameter.bEditorOnly = Prop->IsEditorOnlyProperty();
+
+		if (Parameter.bEditorOnly)
+		{
+			UE_LOG_ERROR_UHT(TEXT("RigVM Struct '%s' - Member '%s' is editor only - WITH_EDITORONLY_DATA not allowed on structs with RIGVM_METHOD."), *Struct->GetName(), *Parameter.Name, *MemberCPPType);
+		}
 
 		if (!ExtendedCPPType.IsEmpty())
 		{
diff --git a/Engine/Source/Programs/UnrealHeaderTool/Private/HeaderParser.h b/Engine/Source/Programs/UnrealHeaderTool/Private/HeaderParser.h
--- a/Engine/Source/Programs/UnrealHeaderTool/Private/HeaderParser.h
+++ b/Engine/Source/Programs/UnrealHeaderTool/Private/HeaderParser.h
@@ -126,6 +126,10 @@ struct FRigVMParameter
 		, bInput(false)
 		, bOutput(false)
 		, MaxArraySize()
+		, Getter()
+		, CastName()
+		, CastType()
+		, bEditorOnly(false)
 	{
 	}
 
@@ -138,6 +142,7 @@ struct FRigVMParameter
 	FString Getter;
 	FString CastName;
 	FString CastType;
+	bool bEditorOnly;
 
 	const FString& NameOriginal(bool bCastName = false) const
 	{
@@ -244,7 +249,7 @@ public:
 		return Parameters.Add(InParameter);
 	}
 
-	FString Names(bool bLeadingSeparator = false, TCHAR* Separator = TEXT(", "), bool bCastType = false) const
+	FString Names(bool bLeadingSeparator = false, TCHAR* Separator = TEXT(", "), bool bCastType = false, bool bIncludeEditorOnly = true) const
 	{
 		if (Parameters.Num() == 0)
 		{
@@ -253,8 +258,18 @@ public:
 		TArray<FString> NameArray;
 		for (const FRigVMParameter& Parameter : Parameters)
 		{
+			if (!bIncludeEditorOnly && Parameter.bEditorOnly)
+			{
+				continue;
+			}
 			NameArray.Add(Parameter.NameOriginal(bCastType));
 		}
+
+		if (NameArray.Num() == 0)
+		{
+			return FString();
+		}
+
 		FString Joined = FString::Join(NameArray, Separator);
 		if (bLeadingSeparator)
 		{
@@ -263,7 +278,7 @@ public:
 		return Joined;
 	}
 
-	FString Declarations(bool bLeadingSeparator = false, TCHAR* Separator = TEXT(", "), bool bCastType = false, bool bCastName = false) const
+	FString Declarations(bool bLeadingSeparator = false, TCHAR* Separator = TEXT(", "), bool bCastType = false, bool bCastName = false, bool bIncludeEditorOnly = true) const
 	{
 		if (Parameters.Num() == 0)
 		{
@@ -272,8 +287,18 @@ public:
 		TArray<FString> DeclarationArray;
 		for (const FRigVMParameter& Parameter : Parameters)
 		{
+			if (!bIncludeEditorOnly && Parameter.bEditorOnly)
+			{
+				continue;
+			}
 			DeclarationArray.Add(Parameter.Variable(bCastType, bCastName));
 		}
+
+		if (DeclarationArray.Num() == 0)
+		{
+			return FString();
+		}
+
 		FString Joined = FString::Join(DeclarationArray, Separator);
 		if (bLeadingSeparator)
 		{
