UHT: Fix build break on OSX in HeaderParser.h

diff --git a/Engine/Source/Programs/UnrealHeaderTool/Private/HeaderParser.h b/Engine/Source/Programs/UnrealHeaderTool/Private/HeaderParser.h
--- a/Engine/Source/Programs/UnrealHeaderTool/Private/HeaderParser.h
+++ b/Engine/Source/Programs/UnrealHeaderTool/Private/HeaderParser.h
@@ -249,7 +249,7 @@ public:
 		return Parameters.Add(InParameter);
 	}
 
-	FString Names(bool bLeadingSeparator = false, TCHAR* Separator = TEXT(", "), bool bCastType = false, bool bIncludeEditorOnly = true) const
+	FString Names(bool bLeadingSeparator = false, const TCHAR* Separator = TEXT(", "), bool bCastType = false, bool bIncludeEditorOnly = true) const
 	{
 		if (Parameters.Num() == 0)
 		{
@@ -278,7 +278,7 @@ public:
 		return Joined;
 	}
 
-	FString Declarations(bool bLeadingSeparator = false, TCHAR* Separator = TEXT(", "), bool bCastType = false, bool bCastName = false, bool bIncludeEditorOnly = true) const
+	FString Declarations(bool bLeadingSeparator = false, const TCHAR* Separator = TEXT(", "), bool bCastType = false, bool bCastName = false, bool bIncludeEditorOnly = true) const
 	{
 		if (Parameters.Num() == 0)
 		{
