#rnx

diff --git a/Engine/Source/Programs/UnrealHeaderTool/Private/CodeGenerator.cpp b/Engine/Source/Programs/UnrealHeaderTool/Private/CodeGenerator.cpp
--- a/Engine/Source/Programs/UnrealHeaderTool/Private/CodeGenerator.cpp
+++ b/Engine/Source/Programs/UnrealHeaderTool/Private/CodeGenerator.cpp
@@ -110,6 +110,9 @@ namespace
 
 	const TCHAR RequiredCPPIncludes[] = TEXT("#include \"UObject/GeneratedCppIncludes.h\"") LINE_TERMINATOR;
 
+	const TCHAR EnableDeprecationWarnings[] = TEXT("PRAGMA_ENABLE_DEPRECATION_WARNINGS") LINE_TERMINATOR;
+	const TCHAR DisableDeprecationWarnings[] = TEXT("PRAGMA_DISABLE_DEPRECATION_WARNINGS") LINE_TERMINATOR;
+
 	// A struct which emits #if and #endif blocks as appropriate when invoked.
 	struct FMacroBlockEmitter
 	{
@@ -226,9 +229,7 @@ namespace
 		}
 
 		const FProperty* LastProperty = ClassReps.Last().Property;
-		NetFieldBuilder.Log( LastProperty->ArrayDim > 1 ?
-			FString::Printf(TEXT("\t\tNETFIELD_REP_END=%s_STATIC_ARRAY_END,\r\n"), *LastProperty->GetName()) :
-			FString::Printf(TEXT("\t\tNETFIELD_REP_END=%s,\r\n"), *LastProperty->GetName()));
+		NetFieldBuilder.Logf(TEXT("\t\tNETFIELD_REP_END=%s%s"), *LastProperty->GetName(), LastProperty->ArrayDim > 1 ? TEXT("_STATIC_ARRAY_END") : TEXT(""));
 
 		NetFieldBuilder.Log(TEXT("\t};"));
 
@@ -248,7 +249,7 @@ namespace
 			API);
 	}
 
-	static const FString GetLifetimeReplicatedPropsStr(TEXT("GetLifetimeReplicatedProps"));
+	static const FString STRING_GetLifetimeReplicatedPropsStr(TEXT("GetLifetimeReplicatedProps"));
 
 	static void WriteReplicatedMacroData(
 		const ClassDefinitionRange& ClassRange,
@@ -260,7 +261,7 @@ namespace
 		const FUnrealSourceFile& SourceFile,
 		FNativeClassHeaderGenerator::EExportClassOutFlags& OutFlags)
 	{
-		const bool bHasGetLifetimeReplicatedProps = HasIdentifierExactMatch(ClassRange.Start, ClassRange.End, GetLifetimeReplicatedPropsStr);
+		const bool bHasGetLifetimeReplicatedProps = HasIdentifierExactMatch(ClassRange.Start, ClassRange.End, STRING_GetLifetimeReplicatedPropsStr);
 
 		if (!bHasGetLifetimeReplicatedProps)
 		{
@@ -463,11 +464,6 @@ FString Macroize(const TCHAR* MacroName, FString&& StringToMacroize)
 	return FString::Printf(TEXT("#define %s%s\r\n%s"), MacroName, Result.Len() ? TEXT(" \\") : TEXT(""), *Result);
 }
 
-FString Macroize(const TCHAR* MacroName, const TCHAR* StringToMacroize)
-{
-	return Macroize(MacroName, FString(StringToMacroize));
-}
-
 static void AddGeneratedCodeHash(void* Field, uint32 Hash)
 {
 	FRWScopeLock Lock(GGeneratedCodeHashesLock, SLT_Write);
@@ -1855,21 +1851,17 @@ inline FString GetEventStructParamsName(UObject* Outer, const TCHAR* FunctionNam
 void FNativeClassHeaderGenerator::OutputProperty(FOutputDevice& DeclOut, FOutputDevice& Out, FReferenceGatherers& OutReferenceGatherers, const TCHAR* Scope, TArray<FPropertyNamePointerPair>& PropertyNamesAndPointers, FProperty* Prop, const TCHAR* DeclSpaces, const TCHAR* Spaces) const
 {
 	// Helper to handle the creation of the underlying properties if they're enum properties
-	auto HandleUnderlyingEnumProperty = [this, &PropertyNamesAndPointers, &DeclOut, &Out, &OutReferenceGatherers, DeclSpaces, Spaces](FProperty* LocalProp, FString OuterName)
+	auto HandleUnderlyingEnumProperty = [this, &PropertyNamesAndPointers, &DeclOut, &Out, &OutReferenceGatherers, DeclSpaces, Spaces](FProperty* LocalProp, FString&& InOuterName)
 	{
+		const FString& OuterName = PropertyNamesAndPointers.Emplace_GetRef(MoveTemp(InOuterName), LocalProp).Name;
+
 		if (FEnumProperty* EnumProp = CastField<FEnumProperty>(LocalProp))
 		{
 			FString PropVarName = OuterName + TEXT("_Underlying");
 
-			PropertyNamesAndPointers.Emplace(MoveTemp(OuterName), LocalProp);
-
 			PropertyNew(DeclOut, Out, OutReferenceGatherers, EnumProp->UnderlyingProp, TEXT("0"), *PropVarName, DeclSpaces, Spaces);
 			PropertyNamesAndPointers.Emplace(MoveTemp(PropVarName), EnumProp->UnderlyingProp);
 		}
-		else
-		{
-			PropertyNamesAndPointers.Emplace(MoveTemp(OuterName), LocalProp);
-		}
 	};
 
 	{
@@ -2923,7 +2915,7 @@ FString GetPreservedAccessSpecifierString(FClass* Class)
 
 void WriteMacro(FOutputDevice& Output, const FString& MacroName, FString MacroContent)
 {
-	Output.Log(*Macroize(*MacroName, MoveTemp(MacroContent)));
+	Output.Log(Macroize(*MacroName, MoveTemp(MacroContent)));
 }
 
 static FString PrivatePropertiesOffsetGetters(const UStruct* Struct, const FString& StructCppName)
@@ -3139,8 +3131,6 @@ void FNativeClassHeaderGenerator::ExportClassFromSourceFileInner(
 
 			FString DeprecationWarning = GetGeneratedMacroDeprecationWarning(TEXT("GENERATED_UINTERFACE_BODY"));
 
-			const TCHAR* DeprecationPushString = TEXT("PRAGMA_DISABLE_DEPRECATION_WARNINGS") LINE_TERMINATOR;
-			const TCHAR* DeprecationPopString = TEXT("PRAGMA_ENABLE_DEPRECATION_WARNINGS") LINE_TERMINATOR;
 			const TCHAR* Offset = TEXT("\t");
 
 			OutGeneratedHeaderText.Log(
@@ -3148,10 +3138,10 @@ void FNativeClassHeaderGenerator::ExportClassFromSourceFileInner(
 					*SourceFile.GetGeneratedBodyMacroName(ClassGeneratedBodyLine, true),
 					FString::Printf(TEXT("\t%s\t%s\t%s()") LINE_TERMINATOR TEXT("%s\t%s")
 						, *DeprecationWarning
-						, DeprecationPushString
+						, DisableDeprecationWarnings
 						, *InterfaceMacroName
 						, *StandardUObjectConstructorsMacroCall
-						, DeprecationPopString
+						, EnableDeprecationWarnings
 					)
 				)
 			);
@@ -3160,11 +3150,11 @@ void FNativeClassHeaderGenerator::ExportClassFromSourceFileInner(
 				Macroize(
 					*SourceFile.GetGeneratedBodyMacroName(ClassGeneratedBodyLine),
 					FString::Printf(TEXT("\t%s\t%s()") LINE_TERMINATOR TEXT("%s%s\t%s")
-						, DeprecationPushString
+						, DisableDeprecationWarnings
 						, *InterfaceMacroName
 						, *EnhancedUObjectConstructorsMacroCall
 						, *GetPreservedAccessSpecifierString(Class)
-						, DeprecationPopString
+						, EnableDeprecationWarnings
 					)
 				)
 			);
@@ -3256,8 +3246,6 @@ void FNativeClassHeaderGenerator::ExportClassFromSourceFileInner(
 	}
 
 	{
-		const TCHAR* DeprecationPushString = TEXT("PRAGMA_DISABLE_DEPRECATION_WARNINGS") LINE_TERMINATOR;
-		const TCHAR* DeprecationPopString = TEXT("PRAGMA_ENABLE_DEPRECATION_WARNINGS") LINE_TERMINATOR;
 		const TCHAR* Public = TEXT("public:") LINE_TERMINATOR;
 
 		const bool bIsIInterface = Class->HasAnyClassFlags(CLASS_Interface);
@@ -3285,11 +3273,11 @@ void FNativeClassHeaderGenerator::ExportClassFromSourceFileInner(
 			GeneratedBody = FString::Printf(TEXT("%s%s%s"), *PPOMacroName, *ClassNoPureDeclsMacroCalls, *EnhancedUObjectConstructorsMacroCall);
 		}
 
-		FString WrappedLegacyGeneratedBody = FString::Printf(TEXT("%s%s%s%s%s%s"), *DeprecationWarning, DeprecationPushString, Public, *LegacyGeneratedBody, Public, DeprecationPopString);
-		FString WrappedGeneratedBody = FString::Printf(TEXT("%s%s%s%s%s"), DeprecationPushString, Public, *GeneratedBody, *GetPreservedAccessSpecifierString(Class), DeprecationPopString);
+		FString WrappedLegacyGeneratedBody = FString::Printf(TEXT("%s%s%s%s%s%s"), *DeprecationWarning, DisableDeprecationWarnings, Public, *LegacyGeneratedBody, Public, EnableDeprecationWarnings);
+		FString WrappedGeneratedBody = FString::Printf(TEXT("%s%s%s%s%s"), DisableDeprecationWarnings, Public, *GeneratedBody, *GetPreservedAccessSpecifierString(Class), EnableDeprecationWarnings);
 
-		OutGeneratedHeaderText.Log(*Macroize(*SourceFile.GetGeneratedBodyMacroName(GeneratedBodyLine, true), MoveTemp(WrappedLegacyGeneratedBody)));
-		OutGeneratedHeaderText.Log(*Macroize(*SourceFile.GetGeneratedBodyMacroName(GeneratedBodyLine, false), MoveTemp(WrappedGeneratedBody)));
+		OutGeneratedHeaderText.Log(Macroize(*SourceFile.GetGeneratedBodyMacroName(GeneratedBodyLine, true), MoveTemp(WrappedLegacyGeneratedBody)));
+		OutGeneratedHeaderText.Log(Macroize(*SourceFile.GetGeneratedBodyMacroName(GeneratedBodyLine, false), MoveTemp(WrappedGeneratedBody)));
 	}
 
 	// Forward declare the StaticClass specialisation in the header
@@ -3478,8 +3466,8 @@ void FNativeClassHeaderGenerator::ExportConstructorsMacros(FOutputDevice& OutGen
 		Out.Logf(TEXT("\tDEFINE_VTABLE_PTR_HELPER_CTOR(%s);" LINE_TERMINATOR), *ClassCPPName);
 	}
 
-	OutGeneratedHeaderText.Log(*Macroize(*StdMacroName, *StdMacro));
-	OutGeneratedHeaderText.Log(*Macroize(*EnhMacroName, *EnhMacro));
+	OutGeneratedHeaderText.Log(Macroize(*StdMacroName, *StdMacro));
+	OutGeneratedHeaderText.Log(Macroize(*EnhMacroName, *EnhMacro));
 
 	StandardUObjectConstructorsMacroCall.Logf(TEXT("\t%s\r\n"), *StdMacroName);
 	EnhancedUObjectConstructorsMacroCall.Logf(TEXT("\t%s\r\n"), *EnhMacroName);
@@ -3498,7 +3486,7 @@ bool FNativeClassHeaderGenerator::WriteHeader(const TCHAR* Path, const FString&
 	}
 
 	GeneratedHeaderTextWithCopyright.Log(LINE_TERMINATOR);
-	GeneratedHeaderTextWithCopyright.Log(TEXT("PRAGMA_DISABLE_DEPRECATION_WARNINGS") LINE_TERMINATOR);
+	GeneratedHeaderTextWithCopyright.Log(DisableDeprecationWarnings);
 
 	for (const FString& FWDecl : InOutReferenceGatherers.ForwardDeclarations)
 	{
@@ -3508,8 +3496,8 @@ bool FNativeClassHeaderGenerator::WriteHeader(const TCHAR* Path, const FString&
 		}
 	}
 
-	GeneratedHeaderTextWithCopyright.Log(*InBodyText);
-	GeneratedHeaderTextWithCopyright.Log(TEXT("PRAGMA_ENABLE_DEPRECATION_WARNINGS") LINE_TERMINATOR);
+	GeneratedHeaderTextWithCopyright.Log(InBodyText);
+	GeneratedHeaderTextWithCopyright.Log(EnableDeprecationWarnings);
 
 	bool bHasChanged = SaveHeaderIfChanged(InOutReferenceGatherers, Path, *GeneratedHeaderTextWithCopyright);
 	return bHasChanged;
@@ -4651,8 +4639,7 @@ void FNativeClassHeaderGenerator::CheckRPCFunctions(FReferenceGatherers& OutRefe
 	//
 	// Get string with function specifiers, listing why we need _Implementation or _Validate functions.
 	//
-	TArray<const TCHAR*> FunctionSpecifiers;
-	FunctionSpecifiers.Reserve(4);
+	TArray<const TCHAR*, TInlineAllocator<4>> FunctionSpecifiers;
 	if (bIsNative)			{ FunctionSpecifiers.Add(TEXT("Native"));			}
 	if (bIsNet)				{ FunctionSpecifiers.Add(TEXT("Net"));				}
 	if (bIsBlueprintEvent)	{ FunctionSpecifiers.Add(TEXT("BlueprintEvent"));	}
@@ -6275,8 +6262,6 @@ void FNativeClassHeaderGenerator::ExportUpdatedHeaders(const FString& PackageNam
  */
 void FNativeClassHeaderGenerator::ExportGeneratedCPP(FOutputDevice& Out, const TSet<FString>& InCrossModuleReferences, const TCHAR* EmptyLinkFunctionPostfix, const TCHAR* Body, const TCHAR* OtherIncludes)
 {
-	static const TCHAR EnableDeprecationWarnings [] = TEXT("PRAGMA_ENABLE_DEPRECATION_WARNINGS") LINE_TERMINATOR;
-	static const TCHAR DisableDeprecationWarnings[] = TEXT("PRAGMA_DISABLE_DEPRECATION_WARNINGS") LINE_TERMINATOR;
 	static const TCHAR DisableWarning4883        [] = TEXT("#ifdef _MSC_VER") LINE_TERMINATOR TEXT("#pragma warning (push)") LINE_TERMINATOR TEXT("#pragma warning (disable : 4883)") LINE_TERMINATOR TEXT("#endif") LINE_TERMINATOR;
 	static const TCHAR EnableWarning4883         [] = TEXT("#ifdef _MSC_VER") LINE_TERMINATOR TEXT("#pragma warning (pop)") LINE_TERMINATOR TEXT("#endif") LINE_TERMINATOR;
 
diff --git a/Engine/Source/Programs/UnrealHeaderTool/Private/ParserHelper.cpp b/Engine/Source/Programs/UnrealHeaderTool/Private/ParserHelper.cpp
--- a/Engine/Source/Programs/UnrealHeaderTool/Private/ParserHelper.cpp
+++ b/Engine/Source/Programs/UnrealHeaderTool/Private/ParserHelper.cpp
@@ -150,6 +150,7 @@ void FToken::Clone( const FToken& Other )
 
 	TokenType = Other.TokenType;
 	TokenName = Other.TokenName;
+	bTokenNameInitialized = Other.bTokenNameInitialized;
 	StartPos = Other.StartPos;
 	StartLine = Other.StartLine;
 	TokenProperty = Other.TokenProperty;
diff --git a/Engine/Source/Programs/UnrealHeaderTool/Private/ParserHelper.h b/Engine/Source/Programs/UnrealHeaderTool/Private/ParserHelper.h
--- a/Engine/Source/Programs/UnrealHeaderTool/Private/ParserHelper.h
+++ b/Engine/Source/Programs/UnrealHeaderTool/Private/ParserHelper.h
@@ -734,7 +734,7 @@ public:
 
 private:
 	/** Whether TokenName has been looked up */
-	bool bTokenNameInitialized;
+	mutable bool bTokenNameInitialized;
 	/** Name of token, lazily initialized */
 	mutable FName TokenName;
 
@@ -860,6 +860,7 @@ public:
 		if (!bTokenNameInitialized)
 		{
 			TokenName = FName(Identifier, FNAME_Find);
+			bTokenNameInitialized = true;
 		}
 		return TokenName;
 	}
