Fix off by 1 issue with UHT Generation.

#rnx

diff --git a/Engine/Source/Programs/UnrealHeaderTool/Private/CodeGenerator.cpp b/Engine/Source/Programs/UnrealHeaderTool/Private/CodeGenerator.cpp
--- a/Engine/Source/Programs/UnrealHeaderTool/Private/CodeGenerator.cpp
+++ b/Engine/Source/Programs/UnrealHeaderTool/Private/CodeGenerator.cpp
@@ -185,7 +185,7 @@ namespace
 		NetFieldBuilder.Logf(TEXT(""
 		"\tenum class ENetFields_Private : uint16\r\n"
 		"\t{\r\n"
-		"\t\tNETFIELD_REP_START=(uint16)Super::ENetFields_Private::NETFIELD_REP_END + (uint16)1,\r\n"));
+		"\t\tNETFIELD_REP_START=(uint16)((int32)Super::ENetFields_Private::NETFIELD_REP_END + (int32)1),\r\n"));
 
 		FUHTStringBuilder ArrayDimBuilder;
 
diff --git a/Engine/Source/Runtime/CoreUObject/Public/UObject/Object.h b/Engine/Source/Runtime/CoreUObject/Public/UObject/Object.h
--- a/Engine/Source/Runtime/CoreUObject/Public/UObject/Object.h
+++ b/Engine/Source/Runtime/CoreUObject/Public/UObject/Object.h
@@ -1457,10 +1457,10 @@ private:
 
 public:
 
-	enum class ENetFields_Private : uint16
+	enum class ENetFields_Private
 	{
 		NETFIELD_REP_START = 0,
-		NETFIELD_REP_END = 0
+		NETFIELD_REP_END = -1
 	};
 
 	virtual void ValidateGeneratedRepEnums(const TArray<struct FRepRecord>& ClassReps) const {}	
