Prevent using Hidden enum entries as default values for UFunctions

These enum entries aren't exposed to Blueprints and Python, so cause UI issues or script syntax errors

diff --git a/Engine/Source/Programs/UnrealHeaderTool/Private/HeaderParser.cpp b/Engine/Source/Programs/UnrealHeaderTool/Private/HeaderParser.cpp
--- a/Engine/Source/Programs/UnrealHeaderTool/Private/HeaderParser.cpp
+++ b/Engine/Source/Programs/UnrealHeaderTool/Private/HeaderParser.cpp
@@ -9285,6 +9285,20 @@ bool FHeaderParser::DefaultValueStringCppFormatToInnerFormat(const UProperty* Pr
 		return bIsNull; // always return as null is the only the processing we can do for object defaults
 	}
 
+	auto ValidateEnumEntry = [Property, &CppForm](const UEnum* Enum, const FString& EnumValue)
+	{
+		const int32 EnumEntryIndex = Enum->GetIndexByName(*EnumValue);
+		if (EnumEntryIndex == INDEX_NONE)
+		{
+			return false;
+		}
+		if (Enum->HasMetaData(TEXT("Hidden"), EnumEntryIndex))
+		{
+			FError::Throwf(TEXT("Hidden enum entries cannot be used as default values: %s \"%s\" "), *Property->GetName(), *CppForm);
+		}
+		return true;
+	};
+
 	if( !Property->IsA(UStructProperty::StaticClass()) )
 	{
 		if( Property->IsA(UIntProperty::StaticClass()) )
@@ -9309,7 +9323,7 @@ bool FHeaderParser::DefaultValueStringCppFormatToInnerFormat(const UProperty* Pr
 			if( NULL != Enum )
 			{
 				OutForm = FDefaultValueHelper::GetUnqualifiedEnumValue(FDefaultValueHelper::RemoveWhitespaces(CppForm));
-				return ( INDEX_NONE != Enum->GetIndexByName(*OutForm) );
+				return ValidateEnumEntry(Enum, OutForm);
 			}
 			int32 Value;
 			if( FDefaultValueHelper::ParseInt( CppForm, Value) )
@@ -9324,7 +9338,7 @@ bool FHeaderParser::DefaultValueStringCppFormatToInnerFormat(const UProperty* Pr
 			if (const UEnum* Enum = CastChecked<UEnumProperty>(Property)->GetEnum())
 			{
 				OutForm = FDefaultValueHelper::GetUnqualifiedEnumValue(FDefaultValueHelper::RemoveWhitespaces(CppForm));
-				return Enum->GetIndexByName(*OutForm) != INDEX_NONE;
+				return ValidateEnumEntry(Enum, OutForm);
 			}
 
 			int64 Value;
@@ -9400,7 +9414,6 @@ bool FHeaderParser::DefaultValueStringCppFormatToInnerFormat(const UProperty* Pr
 					if (ParsedTextNamespace.GetValue().Equals(UHTDummyNamespace))
 					{
 						FError::Throwf(TEXT("LOCTEXT default parameter values are not supported; use NSLOCTEXT instead: %s \"%s\" "), *Property->GetName(), *CppForm);
-						return false;
 					}
 				}
 			}
