Expose the new 'UncookedOnly' UBT module type to UHT in order to exclude uncooked-only script code modules from nativized Blueprint dependencies.

diff --git a/Engine/Source/Developer/BlueprintNativeCodeGen/Private/BlueprintNativeCodeGenManifest.cpp b/Engine/Source/Developer/BlueprintNativeCodeGen/Private/BlueprintNativeCodeGenManifest.cpp
--- a/Engine/Source/Developer/BlueprintNativeCodeGen/Private/BlueprintNativeCodeGenManifest.cpp
+++ b/Engine/Source/Developer/BlueprintNativeCodeGen/Private/BlueprintNativeCodeGenManifest.cpp
@@ -235,8 +235,8 @@ static bool BlueprintNativeCodeGenManifestImpl::GatherModuleDependencies(const U
 				continue;
 			}
 
-			// we want only native packages, ones that are not editor-only
-			if ((DependentPackage->HasAnyPackageFlags(PKG_CompiledIn)) && !DependentPackage->HasAnyPackageFlags(PKG_EditorOnly | PKG_Developer) && !DependentPackage->IsLoadedByEditorPropertiesOnly())
+			// we want only native packages, ones that are not editor-only and only those that will be cooked
+			if (DependentPackage->HasAllPackagesFlags(PKG_CompiledIn) && !DependentPackage->HasAnyPackageFlags(PKG_EditorOnly | PKG_Developer | PKG_UncookedOnly))
 			{
 				DependenciesOut.AddUnique(DependentPackage);// PkgImport.ObjectName.ToString());
 			}
diff --git a/Engine/Source/Developer/BlueprintNativeCodeGen/Private/BlueprintNativeCodeGenModule.cpp b/Engine/Source/Developer/BlueprintNativeCodeGen/Private/BlueprintNativeCodeGenModule.cpp
--- a/Engine/Source/Developer/BlueprintNativeCodeGen/Private/BlueprintNativeCodeGenModule.cpp
+++ b/Engine/Source/Developer/BlueprintNativeCodeGen/Private/BlueprintNativeCodeGenModule.cpp
@@ -840,25 +840,25 @@ EReplacementResult FBlueprintNativeCodeGenModule::IsTargetedForReplacement(const
 
 	auto ObjectIsNotReplacedAtAll = [&]() -> bool
 	{
-		// EDITOR ON DEVELOPMENT OBJECT
+		// EDITOR OR UNCOOKED OBJECT
 		{
-			auto IsDeveloperObject = [](const UObject* Obj) -> bool
+			auto IsUncookedOnlyObject = [](const UObject* Obj) -> bool
 			{
-				auto IsObjectFromDeveloperPackage = [](const UObject* InObj) -> bool
+				auto IsObjectFromUncookedPackage = [](const UObject* InObj) -> bool
 				{
-					return InObj && InObj->GetOutermost()->HasAllPackagesFlags(PKG_Developer);
+					return InObj && InObj->GetOutermost()->HasAnyPackageFlags(PKG_Developer | PKG_UncookedOnly);
 				};
 
 				if (Obj)
 				{
-					if (IsObjectFromDeveloperPackage(Obj))
+					if (IsObjectFromUncookedPackage(Obj))
 					{
 						return true;
 					}
 					const UStruct* StructToTest = Obj->IsA<UStruct>() ? CastChecked<const UStruct>(Obj) : Obj->GetClass();
 					for (; StructToTest; StructToTest = StructToTest->GetSuperStruct())
 					{
-						if (IsObjectFromDeveloperPackage(StructToTest))
+						if (IsObjectFromUncookedPackage(StructToTest))
 						{
 							return true;
 						}
@@ -866,7 +866,7 @@ EReplacementResult FBlueprintNativeCodeGenModule::IsTargetedForReplacement(const
 				}
 				return false;
 			};
-			if (Object && (IsEditorOnlyObject(Object) || IsDeveloperObject(Object)))
+			if (Object && (IsEditorOnlyObject(Object) || IsUncookedOnlyObject(Object)))
 			{
 				UE_LOG(LogBlueprintCodeGen, Verbose, TEXT("Object %s depends on Editor or Development stuff. It shouldn't be cooked."), *GetPathNameSafe(Object));
 				return true;
diff --git a/Engine/Source/Programs/UnrealBuildTool/System/ExternalExecution.cs b/Engine/Source/Programs/UnrealBuildTool/System/ExternalExecution.cs
--- a/Engine/Source/Programs/UnrealBuildTool/System/ExternalExecution.cs
+++ b/Engine/Source/Programs/UnrealBuildTool/System/ExternalExecution.cs
@@ -88,10 +88,12 @@ namespace UnrealBuildTool
 	{
 		Program,
 		EngineRuntime,
+		EngineUncooked,
 		EngineDeveloper,
 		EngineEditor,
 		EngineThirdParty,
 		GameRuntime,
+		GameUncooked,
 		GameDeveloper,
 		GameEditor,
 		GameThirdParty,
@@ -120,7 +122,6 @@ namespace UnrealBuildTool
 				case ModuleHostType.RuntimeNoCommandlet:
 				case ModuleHostType.RuntimeAndProgram:
 				case ModuleHostType.CookedOnly:
-				case ModuleHostType.UncookedOnly:
 				case ModuleHostType.ServerOnly:
 				case ModuleHostType.ClientOnly:
 				case ModuleHostType.ClientOnlyNoCommandlet:
@@ -132,6 +133,8 @@ namespace UnrealBuildTool
 				case ModuleHostType.EditorNoCommandlet:
 				case ModuleHostType.EditorAndProgram:
 					return UHTModuleType.EngineEditor;
+				case ModuleHostType.UncookedOnly:
+					return UHTModuleType.EngineUncooked;
 				default:
 					return null;
 			}
@@ -146,7 +149,6 @@ namespace UnrealBuildTool
 				case ModuleHostType.RuntimeNoCommandlet:
 				case ModuleHostType.RuntimeAndProgram:
 				case ModuleHostType.CookedOnly:
-				case ModuleHostType.UncookedOnly:
 				case ModuleHostType.ServerOnly:
 				case ModuleHostType.ClientOnly:
 				case ModuleHostType.ClientOnlyNoCommandlet:
@@ -158,6 +160,8 @@ namespace UnrealBuildTool
 				case ModuleHostType.EditorNoCommandlet:
 				case ModuleHostType.EditorAndProgram:
 					return UHTModuleType.GameEditor;
+				case ModuleHostType.UncookedOnly:
+					return UHTModuleType.GameUncooked;
 				default:
 					return null;
 			}
diff --git a/Engine/Source/Programs/UnrealHeaderTool/Private/CodeGenerator.cpp b/Engine/Source/Programs/UnrealHeaderTool/Private/CodeGenerator.cpp
--- a/Engine/Source/Programs/UnrealHeaderTool/Private/CodeGenerator.cpp
+++ b/Engine/Source/Programs/UnrealHeaderTool/Private/CodeGenerator.cpp
@@ -2115,7 +2115,7 @@ void FNativeClassHeaderGenerator::ExportGeneratedPackageInitCode(FOutputDevice&
 	Out.Logf(TEXT("\t\t\t\t%s,\r\n"), *CreateUTF8LiteralString(InPackage->GetName()));
 	Out.Logf(TEXT("\t\t\t\t%s,\r\n"), SingletonArray);
 	Out.Logf(TEXT("\t\t\t\t%s,\r\n"), SingletonCount);
-	Out.Logf(TEXT("\t\t\t\tPKG_CompiledIn | 0x%08X,\r\n"), InPackage->GetPackageFlags() & (PKG_ClientOptional | PKG_ServerSideOnly | PKG_EditorOnly | PKG_Developer));
+	Out.Logf(TEXT("\t\t\t\tPKG_CompiledIn | 0x%08X,\r\n"), InPackage->GetPackageFlags() & (PKG_ClientOptional | PKG_ServerSideOnly | PKG_EditorOnly | PKG_Developer | PKG_UncookedOnly));
 	Out.Logf(TEXT("\t\t\t\t0x%08X,\r\n"), Hash);
 	Out.Logf(TEXT("\t\t\t\t0x%08X,\r\n"), GenerateTextHash(InDeclarations));
 	Out.Logf(TEXT("\t\t\t\t%s\r\n"), *MetaDataParams);
@@ -6350,14 +6350,22 @@ ECompilationResult::Type PreparseModules(const FString& ModuleInfoPath, int32& N
 		//       want to make sure our flags get set
 		Package->SetPackageFlags(PKG_ContainsScript | PKG_Compiling);
 		Package->ClearPackageFlags(PKG_ClientOptional | PKG_ServerSideOnly);
-		if (Module.ModuleType == EBuildModuleType::GameEditor || Module.ModuleType == EBuildModuleType::EngineEditor)
+		switch (Module.ModuleType)
 		{
+		case EBuildModuleType::GameEditor:
+		case EBuildModuleType::EngineEditor:
 			Package->SetPackageFlags(PKG_EditorOnly);
-		}
+			break;
 
-		if (Module.ModuleType == EBuildModuleType::GameDeveloper || Module.ModuleType == EBuildModuleType::EngineDeveloper)
-		{
-			Package->SetPackageFlags(Package->GetPackageFlags() | PKG_Developer);
+		case EBuildModuleType::GameDeveloper:
+		case EBuildModuleType::EngineDeveloper:
+			Package->SetPackageFlags(PKG_Developer);
+			break;
+
+		case EBuildModuleType::GameUncooked:
+		case EBuildModuleType::EngineUncooked:
+			Package->SetPackageFlags(PKG_UncookedOnly);
+			break;
 		}
 
 		// Add new module or overwrite whatever we had loaded, that data is obsolete.
diff --git a/Engine/Source/Programs/UnrealHeaderTool/Private/HeaderParser.cpp b/Engine/Source/Programs/UnrealHeaderTool/Private/HeaderParser.cpp
--- a/Engine/Source/Programs/UnrealHeaderTool/Private/HeaderParser.cpp
+++ b/Engine/Source/Programs/UnrealHeaderTool/Private/HeaderParser.cpp
@@ -8478,12 +8478,14 @@ FHeaderParser::FHeaderParser(FFeedbackContext* InWarn, const FManifestModule& In
 		}
 		break;
 	case EBuildModuleType::EngineRuntime:
+	case EBuildModuleType::EngineUncooked:
 	case EBuildModuleType::EngineDeveloper:
 	case EBuildModuleType::EngineEditor:
 	case EBuildModuleType::EngineThirdParty:
 		bIsCurrentModulePartOfEngine = true;
 		break;
 	case EBuildModuleType::GameRuntime:
+	case EBuildModuleType::GameUncooked:
 	case EBuildModuleType::GameDeveloper:
 	case EBuildModuleType::GameEditor:
 	case EBuildModuleType::GameThirdParty:
diff --git a/Engine/Source/Programs/UnrealHeaderTool/Private/IScriptGeneratorPluginInterface.cpp b/Engine/Source/Programs/UnrealHeaderTool/Private/IScriptGeneratorPluginInterface.cpp
--- a/Engine/Source/Programs/UnrealHeaderTool/Private/IScriptGeneratorPluginInterface.cpp
+++ b/Engine/Source/Programs/UnrealHeaderTool/Private/IScriptGeneratorPluginInterface.cpp
@@ -12,10 +12,12 @@ EBuildModuleType::Type EBuildModuleType::Parse(const TCHAR* Value)
 		TEXT("EngineEditor"),
 		TEXT("EngineRuntime"),
 		TEXT("EngineThirdParty"),
+		TEXT("EngineUncooked"),
 		TEXT("GameDeveloper"),
 		TEXT("GameEditor"),
 		TEXT("GameRuntime"),
 		TEXT("GameThirdParty"),
+		TEXT("GameUncooked"),
 		TEXT("Program")
 	};
 
@@ -30,10 +32,12 @@ EBuildModuleType::Type EBuildModuleType::Parse(const TCHAR* Value)
 		EngineEditor,
 		EngineRuntime,
 		EngineThirdParty,
+		EngineUncooked,
 		GameDeveloper,
 		GameEditor,
 		GameRuntime,
 		GameThirdParty,
+		GameUncooked,
 		Program
 	};
 
diff --git a/Engine/Source/Programs/UnrealHeaderTool/Public/IScriptGeneratorPluginInterface.h b/Engine/Source/Programs/UnrealHeaderTool/Public/IScriptGeneratorPluginInterface.h
--- a/Engine/Source/Programs/UnrealHeaderTool/Public/IScriptGeneratorPluginInterface.h
+++ b/Engine/Source/Programs/UnrealHeaderTool/Public/IScriptGeneratorPluginInterface.h
@@ -16,10 +16,12 @@ struct EBuildModuleType
 	{
 		Program,
 		EngineRuntime,
+		EngineUncooked,
 		EngineDeveloper,
 		EngineEditor,
 		EngineThirdParty,
 		GameRuntime,
+		GameUncooked,
 		GameDeveloper,
 		GameEditor,
 		GameThirdParty,
diff --git a/Engine/Source/Runtime/CoreUObject/Public/UObject/ObjectMacros.h b/Engine/Source/Runtime/CoreUObject/Public/UObject/ObjectMacros.h
--- a/Engine/Source/Runtime/CoreUObject/Public/UObject/ObjectMacros.h
+++ b/Engine/Source/Runtime/CoreUObject/Public/UObject/ObjectMacros.h
@@ -109,7 +109,7 @@ enum EPackageFlags
 	PKG_ForDiffing					= 0x00000020,	///< This package was loaded just for the purposes of diffing
 	PKG_EditorOnly					= 0x00000040,	///< This is editor-only package (for example: editor module script package)
 	PKG_Developer					= 0x00000080,	///< Developer module
-//	PKG_Unused						= 0x00000100,
+	PKG_UncookedOnly				= 0x00000100,	///< Loaded only in uncooked builds (i.e. runtime in editor)
 //	PKG_Unused						= 0x00000200,
 //	PKG_Unused						= 0x00000400,
 //	PKG_Unused						= 0x00000800,
