New non-member FStructuredArchiveField type are now recognised by UnrealHeaderTool.
Refactoring to avoid the use of metadata which ought not to exist for passing data between the parser and the generator.
Test case added.

diff --git a/Engine/Source/Programs/UnrealHeaderTool/Private/ClassMaps.cpp b/Engine/Source/Programs/UnrealHeaderTool/Private/ClassMaps.cpp
--- a/Engine/Source/Programs/UnrealHeaderTool/Private/ClassMaps.cpp
+++ b/Engine/Source/Programs/UnrealHeaderTool/Private/ClassMaps.cpp
@@ -18,3 +18,4 @@ TMap<FName, TSharedRef<FClassDeclarationMetaData> > GClassDeclarations;
 TSet<UProperty*> GUnsizedProperties;
 TSet<UField*> GEditorOnlyDataTypes;
 TMap<UStruct*, TTuple<TSharedRef<FUnrealSourceFile>, int32>> GStructToSourceLine;
+TMap<UClass*, FArchiveTypeDefinePair> GClassSerializerMap;
diff --git a/Engine/Source/Programs/UnrealHeaderTool/Private/ClassMaps.h b/Engine/Source/Programs/UnrealHeaderTool/Private/ClassMaps.h
--- a/Engine/Source/Programs/UnrealHeaderTool/Private/ClassMaps.h
+++ b/Engine/Source/Programs/UnrealHeaderTool/Private/ClassMaps.h
@@ -19,6 +19,19 @@ struct FManifestModule;
 class FUnrealSourceFile;
 class FUnrealTypeDefinitionInfo;
 
+enum class ESerializerArchiveType
+{
+	None,
+	Archive,
+	StructuredArchiveRecord
+};
+
+struct FArchiveTypeDefinePair
+{
+	ESerializerArchiveType ArchiveType;
+	FString EnclosingDefine;
+};
+
 extern TMap<FString, TSharedRef<FUnrealSourceFile> > GUnrealSourceFilesMap;
 extern TMap<UField*, TSharedRef<FUnrealTypeDefinitionInfo> > GTypeDefinitionInfoMap;
 extern TMap<const UPackage*, TArray<UField*>> GPackageSingletons;
@@ -33,6 +46,7 @@ extern TMap<FName, TSharedRef<FClassDeclarationMetaData> > GClassDeclarations;
 extern TSet<UProperty*> GUnsizedProperties;
 extern TSet<UField*> GEditorOnlyDataTypes;
 extern TMap<UStruct*, TTuple<TSharedRef<FUnrealSourceFile>, int32>> GStructToSourceLine;
+extern TMap<UClass*, FArchiveTypeDefinePair> GClassSerializerMap;
 
 /** Types access specifiers. */
 enum EAccessSpecifier
diff --git a/Engine/Source/Programs/UnrealHeaderTool/Private/CodeGenerator.cpp b/Engine/Source/Programs/UnrealHeaderTool/Private/CodeGenerator.cpp
--- a/Engine/Source/Programs/UnrealHeaderTool/Private/CodeGenerator.cpp
+++ b/Engine/Source/Programs/UnrealHeaderTool/Private/CodeGenerator.cpp
@@ -2729,28 +2729,24 @@ void FNativeClassHeaderGenerator::ExportClassFromSourceFileInner(
 		ClassRange.Validate();
 	}
 
-	bool bHasSerializeToFArchive = Class->HasMetaData(TEXT("SerializeToFArchive"));
-	bool bHasSerializeToFStructuredArchive = Class->HasMetaData(TEXT("SerializeToFStructuredArchive"));
-
 	FString GeneratedSerializeFunctionCPP;
 	FString GeneratedSerializeFunctionHeaderMacroName;
 
-	if (bHasSerializeToFArchive != bHasSerializeToFStructuredArchive)
+	if (const FArchiveTypeDefinePair* ArchiveTypeDefinePair = GClassSerializerMap.Find(Class))
 	{
 		FString EnclosingDefines;
 		FUHTStringBuilder Boilerplate, BoilerPlateCPP;
 		FString MacroNameHeader, MacroNameCPP;
 		GeneratedSerializeFunctionHeaderMacroName = SourceFile.GetGeneratedMacroName(ClassData, TEXT("_ARCHIVESERIALIZER"));
 
-		if (!bHasSerializeToFArchive)
+		EnclosingDefines = ArchiveTypeDefinePair->EnclosingDefine;
+		if (ArchiveTypeDefinePair->ArchiveType == ESerializerArchiveType::StructuredArchiveRecord)
 		{
-			EnclosingDefines = Class->GetMetaData(TEXT("SerializeToFStructuredArchive"));
 			MacroNameHeader = TEXT("DECLARE_FARCHIVE_SERIALIZER");
 			MacroNameCPP = TEXT("IMPLEMENT_FARCHIVE_SERIALIZER");
 		}
 		else
 		{
-			EnclosingDefines = Class->GetMetaData(TEXT("SerializeToFArchive"));
 			MacroNameHeader = TEXT("DECLARE_FSTRUCTUREDARCHIVE_SERIALIZER");
 			MacroNameCPP = TEXT("IMPLEMENT_FSTRUCTUREDARCHIVE_SERIALIZER");
 		}
@@ -5228,18 +5224,15 @@ FNativeClassHeaderGenerator::FNativeClassHeaderGenerator(
 				}
 			}
 
-			bool bHasSerializeToFArchive = Class->HasMetaData(TEXT("SerializeToFArchive"));
-			bool bHasSerializeToFStructuredArchive = Class->HasMetaData(TEXT("SerializeToFStructuredArchive"));
-
-			if (bHasSerializeToFArchive != bHasSerializeToFStructuredArchive)
+			if (const FArchiveTypeDefinePair* ArchiveTypeDefinePair = GClassSerializerMap.Find(Class))
 			{
-				if (!bAddedStructuredArchiveFromArchiveHeader && !bHasSerializeToFArchive)
+				if (!bAddedStructuredArchiveFromArchiveHeader && ArchiveTypeDefinePair->ArchiveType == ESerializerArchiveType::StructuredArchiveRecord)
 				{
 					RelativeIncludes.AddUnique(TEXT("Serialization/StructuredArchive.h"));
 					bAddedStructuredArchiveFromArchiveHeader = true;
 				}
 
-				if (!bAddedArchiveUObjectFromStructuredArchiveHeader && !bHasSerializeToFStructuredArchive)
+				if (!bAddedArchiveUObjectFromStructuredArchiveHeader && ArchiveTypeDefinePair->ArchiveType == ESerializerArchiveType::Archive)
 				{
 					RelativeIncludes.AddUnique(TEXT("Serialization/ArchiveUObjectFromStructuredArchive.h"));
 					bAddedArchiveUObjectFromStructuredArchiveHeader = true;
diff --git a/Engine/Source/Programs/UnrealHeaderTool/Private/HeaderParser.cpp b/Engine/Source/Programs/UnrealHeaderTool/Private/HeaderParser.cpp
--- a/Engine/Source/Programs/UnrealHeaderTool/Private/HeaderParser.cpp
+++ b/Engine/Source/Programs/UnrealHeaderTool/Private/HeaderParser.cpp
@@ -5442,6 +5442,7 @@ bool FHeaderParser::CompileDeclaration(FClasses& AllClasses, TArray<UDelegateFun
 		static const FName NAME_CloseBracket(TEXT(")"));
 		static const FName NAME_FArchive(TEXT("FArchive"));
 		static const FName NAME_FStructuredArchive(TEXT("FStructuredArchive"));
+		static const FName NAME_FStructuredArchiveRecord(TEXT("FStructuredArchiveRecord"));
 		static const FName NAME_Reference(TEXT("&"));
 		static const FName NAME_ClassMember(TEXT("::"));
 		static const FName NAME_FRecord(TEXT("FRecord"));
@@ -5460,72 +5461,81 @@ bool FHeaderParser::CompileDeclaration(FClasses& AllClasses, TArray<UDelegateFun
 				if (Token.Identifier == NAME_OpenBracket)
 				{
 					GetToken(Token);
-					bool bMatchedSerializeToFArchive = Token.Identifier == NAME_FArchive;
-					bool bMatchedSerializeToFStructuredArchive = Token.Identifier == NAME_FStructuredArchive;
 
-					if (bMatchedSerializeToFArchive || bMatchedSerializeToFStructuredArchive)
+					ESerializerArchiveType ArchiveType = ESerializerArchiveType::None;
+					if (Token.Identifier == NAME_FArchive)
 					{
-						bool bMatchingFunctionSignature = false;
 						GetToken(Token);
-
-						if (bMatchedSerializeToFArchive)
+						if (Token.Identifier == NAME_Reference)
 						{
-							if (Token.Identifier == NAME_Reference)
+							GetToken(Token);
+
+							// Allow the declaration to not define a name for the archive parameter
+							if (Token.Identifier != NAME_CloseBracket)
 							{
 								GetToken(Token);
+							}
 
-								// Allow the declaration to not define a name for the archive parameter
-								if (Token.Identifier != NAME_CloseBracket)
-								{
-									GetToken(Token);
-								}
-
-								bMatchingFunctionSignature = Token.Identifier == NAME_CloseBracket;
+							if (Token.Identifier == NAME_CloseBracket)
+							{
+								ArchiveType = ESerializerArchiveType::Archive;
 							}
 						}
-						else
+					}
+					else if (Token.Identifier == NAME_FStructuredArchive)
+					{
+						GetToken(Token);
+						if (Token.Identifier == NAME_ClassMember)
 						{
-							if (Token.Identifier == NAME_ClassMember)
+							GetToken(Token);
+
+							if (Token.Identifier == NAME_FRecord)
 							{
 								GetToken(Token);
 
-								if (Token.Identifier == NAME_FRecord)
+								// Allow the declaration to not define a name for the slot parameter
+								if (Token.Identifier != NAME_CloseBracket)
 								{
 									GetToken(Token);
+								}
 
-									// Allow the declaration to not define a name for the slot parameter
-									if (Token.Identifier != NAME_CloseBracket)
-									{
-										GetToken(Token);
-									}
-
-									bMatchingFunctionSignature = Token.Identifier == NAME_CloseBracket;
+								if (Token.Identifier == NAME_CloseBracket)
+								{
+									ArchiveType = ESerializerArchiveType::StructuredArchiveRecord;
 								}
 							}
 						}
+					}
+					else if (Token.Identifier == NAME_FStructuredArchiveRecord)
+					{
+						GetToken(Token);
 
-						if (bMatchingFunctionSignature)
+						// Allow the declaration to not define a name for the slot parameter
+						if (Token.Identifier != NAME_CloseBracket)
 						{
-							// Found what we want!
-							if (CompilerDirectiveStack.Num() == 0 || (CompilerDirectiveStack.Num() == 1 && CompilerDirectiveStack[0] == ECompilerDirective::WithEditorOnlyData))
-							{
-								FString EnclosingDefine = CompilerDirectiveStack.Num() > 0 ? TEXT("WITH_EDITORONLY_DATA") : TEXT("");
+							GetToken(Token);
+						}
 
-								UClass* CurrentClass = GetCurrentClass();
-								
-								if (bMatchedSerializeToFArchive)
-								{
-									CurrentClass->SetMetaData(TEXT("SerializeToFArchive"), *EnclosingDefine);
-								}
-								else
-								{
-									CurrentClass->SetMetaData(TEXT("SerializeToFStructuredArchive"), *EnclosingDefine);
-								}
-							}
-							else
-							{
-								FError::Throwf(TEXT("Serialize functions must not be inside preprocessor blocks, except for WITH_EDITORONLY_DATA"));
-							}
+						if (Token.Identifier == NAME_CloseBracket)
+						{
+							ArchiveType = ESerializerArchiveType::StructuredArchiveRecord;
+						}
+					}
+
+					if (ArchiveType != ESerializerArchiveType::None)
+					{
+						// Found what we want!
+						if (CompilerDirectiveStack.Num() == 0 || (CompilerDirectiveStack.Num() == 1 && CompilerDirectiveStack[0] == ECompilerDirective::WithEditorOnlyData))
+						{
+							FString EnclosingDefine = CompilerDirectiveStack.Num() > 0 ? TEXT("WITH_EDITORONLY_DATA") : TEXT("");
+
+							UClass* CurrentClass = GetCurrentClass();
+
+							GClassSerializerMap.Add(CurrentClass, { ArchiveType, MoveTemp(EnclosingDefine) });
+						}
+						else
+						{
+							FError::Throwf(TEXT("Serialize functions must not be inside preprocessor blocks, except for WITH_EDITORONLY_DATA"));
 						}
 					}
 				}
diff --git a/Engine/Source/Runtime/Engine/Classes/Engine/DataAsset.h b/Engine/Source/Runtime/Engine/Classes/Engine/DataAsset.h
--- a/Engine/Source/Runtime/Engine/Classes/Engine/DataAsset.h
+++ b/Engine/Source/Runtime/Engine/Classes/Engine/DataAsset.h
@@ -19,7 +19,7 @@ class UDataAsset : public UObject
 public:
 	// UObject interface
 #if WITH_EDITORONLY_DATA
-	ENGINE_API virtual void Serialize(FStructuredArchive::FRecord Record) override;
+	ENGINE_API virtual void Serialize(FStructuredArchiveRecord Record) override;
 #endif
 
 private:
