HeaderParser  ARRAY_COUNT -> UE_ARRAY_COUNT

diff --git a/Engine/Source/Programs/UnrealHeaderTool/Private/HeaderParser.cpp b/Engine/Source/Programs/UnrealHeaderTool/Private/HeaderParser.cpp
--- a/Engine/Source/Programs/UnrealHeaderTool/Private/HeaderParser.cpp
+++ b/Engine/Source/Programs/UnrealHeaderTool/Private/HeaderParser.cpp
@@ -2078,7 +2078,7 @@ TMap<FName, FString> FHeaderParser::GetParameterToolTipsFromFunctionComment(cons
 		int32 ParamStart = Input.Find(ParamTag, ESearchCase::CaseSensitive, ESearchDir::FromStart, Offset);
 		if(ParamStart != INDEX_NONE)
 		{
-			ParamStart = ParamStart + ARRAY_COUNT(ParamTag);
+			ParamStart = ParamStart + UE_ARRAY_COUNT(ParamTag);
 			Offset = ParamStart;
 		}
 		else
@@ -2086,7 +2086,7 @@ TMap<FName, FString> FHeaderParser::GetParameterToolTipsFromFunctionComment(cons
 			ParamStart = Input.Find(ReturnTag, ESearchCase::CaseSensitive, ESearchDir::FromStart, Offset);
 			if (ParamStart != INDEX_NONE)
 			{
-				ParamStart = ParamStart + ARRAY_COUNT(ReturnTag);
+				ParamStart = ParamStart + UE_ARRAY_COUNT(ReturnTag);
 				Offset = ParamStart;
 				ParamPrefix = ReturnParamPrefix;
 			}
