Support correctly skipping over PURE_VIRTUAL macros in UHT.

diff --git a/Engine/Source/Programs/UnrealHeaderTool/Private/HeaderParser.cpp b/Engine/Source/Programs/UnrealHeaderTool/Private/HeaderParser.cpp
--- a/Engine/Source/Programs/UnrealHeaderTool/Private/HeaderParser.cpp
+++ b/Engine/Source/Programs/UnrealHeaderTool/Private/HeaderParser.cpp
@@ -5675,6 +5675,12 @@ bool FHeaderParser::SkipDeclaration(FToken& Token)
 			break;
 		}
 
+		if (!bMacroDeclaration && Token.Matches(TEXT("PURE_VIRTUAL")) && NestedScopes == 0)
+		{
+			OpeningBracket = TEXT("(");
+			ClosingBracket = TEXT(")");
+		}
+
 		if (Token.Matches(OpeningBracket))
 		{
 			// This is a function definition or class declaration.
diff --git a/Engine/Source/Programs/UnrealHeaderTool/Resources/UHTDebugging/TestObject.h b/Engine/Source/Programs/UnrealHeaderTool/Resources/UHTDebugging/TestObject.h
--- a/Engine/Source/Programs/UnrealHeaderTool/Resources/UHTDebugging/TestObject.h
+++ b/Engine/Source/Programs/UnrealHeaderTool/Resources/UHTDebugging/TestObject.h
@@ -116,6 +116,8 @@ public:
 		return FString("Hello").Len();
 	}
 
+	virtual void PureVirtualImplementedFunction() PURE_VIRTUAL(, )
+
 	UFUNCTION()
 	FORCENOINLINE int32 NoInlineFunc()
 	{
