Fixing UHT compile error on the Mac

diff --git a/Engine/Source/Programs/UnrealHeaderTool/Private/NativeClassExporter.h b/Engine/Source/Programs/UnrealHeaderTool/Private/NativeClassExporter.h
--- a/Engine/Source/Programs/UnrealHeaderTool/Private/NativeClassExporter.h
+++ b/Engine/Source/Programs/UnrealHeaderTool/Private/NativeClassExporter.h
@@ -75,6 +75,8 @@ struct FPropertyNamePointerPair
 	FProperty* Prop;
 };
 
+FString CreateUTF8LiteralString(const FString& Str);
+
 struct FNativeClassHeaderGenerator
 {
 private:
