Renaming FField::GetOwner() to FField::GetOwnerVariant() to work around incorrect static analysis warnings

diff --git a/Engine/Plugins/Experimental/DataPrepEditor/Source/DataPrepCore/Private/Parameterization/DataprepParameterization.cpp b/Engine/Plugins/Experimental/DataPrepEditor/Source/DataPrepCore/Private/Parameterization/DataprepParameterization.cpp
--- a/Engine/Plugins/Experimental/DataPrepEditor/Source/DataPrepCore/Private/Parameterization/DataprepParameterization.cpp
+++ b/Engine/Plugins/Experimental/DataPrepEditor/Source/DataPrepCore/Private/Parameterization/DataprepParameterization.cpp
@@ -326,7 +326,7 @@ namespace DataprepParameterization
 
 					// We look if the outer of the property is the right one as a heuristic for the validity of the current property
 					bool bPropertySeamsValid = false;
-					if ( FFieldVariant CurrentPropertyOuter = CurentProperty->GetOwner() )
+					if ( FFieldVariant CurrentPropertyOuter = CurentProperty->GetOwnerVariant() )
 					{
 						if ( CurrentPropertyOuter == CurrentOuter )
 						{
@@ -335,7 +335,7 @@ namespace DataprepParameterization
 						else if ( UStruct* CurrentOuterAsStruct = CurrentOuter.Get<UStruct>() )
 						{
 							// Walk the struct hierarchy in case we have a sub struct of the expected struct
-							FFieldVariant PropertyOuter = CurentProperty->GetOwner();
+							FFieldVariant PropertyOuter = CurentProperty->GetOwnerVariant();
 							CurrentOuterAsStruct = CurrentOuterAsStruct->GetSuperStruct();
 							while ( !bPropertySeamsValid && CurrentOuterAsStruct )
 							{
diff --git a/Engine/Plugins/Experimental/DataPrepEditor/Source/DataPrepEditor/Private/Widgets/DataprepWidgets.cpp b/Engine/Plugins/Experimental/DataPrepEditor/Source/DataPrepEditor/Private/Widgets/DataprepWidgets.cpp
--- a/Engine/Plugins/Experimental/DataPrepEditor/Source/DataPrepEditor/Private/Widgets/DataprepWidgets.cpp
+++ b/Engine/Plugins/Experimental/DataPrepEditor/Source/DataPrepEditor/Private/Widgets/DataprepWidgets.cpp
@@ -710,7 +710,7 @@ void SDataprepDetailsView::AddWidgets( const TArray< TSharedRef< IDetailTreeNode
 
 			if ( Property )
 			{
-				if ( FFieldVariant Outer = Property->GetOwner() )
+				if ( FFieldVariant Outer = Property->GetOwnerVariant() )
 				{
 					// if the outer is a container property (array,set or map) it's editable even without the proper flags.
 					//UClass* OuterClass = Outer->GetClass();
diff --git a/Engine/Plugins/Experimental/PythonScriptPlugin/Source/PythonScriptPlugin/Private/PyGenUtil.cpp b/Engine/Plugins/Experimental/PythonScriptPlugin/Source/PythonScriptPlugin/Private/PyGenUtil.cpp
--- a/Engine/Plugins/Experimental/PythonScriptPlugin/Source/PythonScriptPlugin/Private/PyGenUtil.cpp
+++ b/Engine/Plugins/Experimental/PythonScriptPlugin/Source/PythonScriptPlugin/Private/PyGenUtil.cpp
@@ -597,10 +597,10 @@ void FGeneratedWrappedFieldTracker::RegisterPythonFieldName(const FString& InPyt
 		{
 			// Note: We don't use GetOwnerStruct here, as UFunctions are UStructs so it
 			// doesn't work correctly for them as it includes 'this' in the look-up chain
-			FFieldVariant OwnerStruct = InField.GetOwner();
+			FFieldVariant OwnerStruct = InField.GetOwnerVariant();
 			while (OwnerStruct.IsValid() && !OwnerStruct.IsA<UStruct>())
 			{
-				OwnerStruct = OwnerStruct.GetOwner();
+				OwnerStruct = OwnerStruct.GetOwnerVariant();
 			}
 			return OwnerStruct.IsValid() ? FString::Printf(TEXT("%s.%s"), *OwnerStruct.GetName(), *InField.GetName()) : InField.GetName();
 		};
diff --git a/Engine/Plugins/MovieScene/MatineeToLevelSequence/Source/MatineeToLevelSequence/Private/MatineeToLevelSequenceModule.cpp b/Engine/Plugins/MovieScene/MatineeToLevelSequence/Source/MatineeToLevelSequence/Private/MatineeToLevelSequenceModule.cpp
--- a/Engine/Plugins/MovieScene/MatineeToLevelSequence/Source/MatineeToLevelSequence/Private/MatineeToLevelSequenceModule.cpp
+++ b/Engine/Plugins/MovieScene/MatineeToLevelSequence/Source/MatineeToLevelSequence/Private/MatineeToLevelSequenceModule.cpp
@@ -416,11 +416,11 @@ protected:
 			}
 
 			TArray<FFieldVariant> PropertyArray;
-			FFieldVariant Outer = Property->GetOwner();
+			FFieldVariant Outer = Property->GetOwnerVariant();
 			while (Outer.IsA(FProperty::StaticClass()) || Outer.IsA(UScriptStruct::StaticClass()))
 			{
 				PropertyArray.Insert(Outer, 0);
-				Outer = Outer.GetOwner();
+				Outer = Outer.GetOwnerVariant();
 			}
 
 			FString PropertyPath;
diff --git a/Engine/Plugins/Runtime/GameplayAbilities/Source/GameplayAbilities/Private/AttributeSet.cpp b/Engine/Plugins/Runtime/GameplayAbilities/Source/GameplayAbilities/Private/AttributeSet.cpp
--- a/Engine/Plugins/Runtime/GameplayAbilities/Source/GameplayAbilities/Private/AttributeSet.cpp
+++ b/Engine/Plugins/Runtime/GameplayAbilities/Source/GameplayAbilities/Private/AttributeSet.cpp
@@ -346,7 +346,7 @@ void UAttributeSet::InitFromMetaDataTable(const UDataTable* DataTable)
 		FNumericProperty *NumericProperty = CastField<FNumericProperty>(Property);
 		if (NumericProperty)
 		{
-			FString RowNameStr = FString::Printf(TEXT("%s.%s"), *Property->GetOwner().GetName(), *Property->GetName());
+			FString RowNameStr = FString::Printf(TEXT("%s.%s"), *Property->GetOwnerVariant().GetName(), *Property->GetName());
 		
 			FAttributeMetaData * MetaData = DataTable->FindRow<FAttributeMetaData>(FName(*RowNameStr), Context, false);
 			if (MetaData)
@@ -357,7 +357,7 @@ void UAttributeSet::InitFromMetaDataTable(const UDataTable* DataTable)
 		}
 		else if (FGameplayAttribute::IsGameplayAttributeDataProperty(Property))
 		{
-			FString RowNameStr = FString::Printf(TEXT("%s.%s"), *Property->GetOwner().GetName(), *Property->GetName());
+			FString RowNameStr = FString::Printf(TEXT("%s.%s"), *Property->GetOwnerVariant().GetName(), *Property->GetName());
 
 			FAttributeMetaData * MetaData = DataTable->FindRow<FAttributeMetaData>(FName(*RowNameStr), Context, false);
 			if (MetaData)
diff --git a/Engine/Plugins/Runtime/GameplayAbilities/Source/GameplayAbilities/Public/AttributeSet.h b/Engine/Plugins/Runtime/GameplayAbilities/Source/GameplayAbilities/Public/AttributeSet.h
--- a/Engine/Plugins/Runtime/GameplayAbilities/Source/GameplayAbilities/Public/AttributeSet.h
+++ b/Engine/Plugins/Runtime/GameplayAbilities/Source/GameplayAbilities/Public/AttributeSet.h
@@ -101,7 +101,7 @@ struct GAMEPLAYABILITIES_API FGameplayAttribute
 	UClass* GetAttributeSetClass() const
 	{
 		check(Attribute.Get());
-		return CastChecked<UClass>(Attribute->GetOwner().ToUObject());
+		return CastChecked<UClass>(Attribute->GetOwner<UObject>());
 	}
 
 	/** Returns true if this is one of the special attributes defined on the bBilitySystemComponent itself */
diff --git a/Engine/Plugins/Runtime/GameplayAbilities/Source/GameplayAbilitiesEditor/Private/AttributeDetails.cpp b/Engine/Plugins/Runtime/GameplayAbilities/Source/GameplayAbilitiesEditor/Private/AttributeDetails.cpp
--- a/Engine/Plugins/Runtime/GameplayAbilities/Source/GameplayAbilitiesEditor/Private/AttributeDetails.cpp
+++ b/Engine/Plugins/Runtime/GameplayAbilities/Source/GameplayAbilitiesEditor/Private/AttributeDetails.cpp
@@ -47,7 +47,7 @@ void FAttributePropertyDetails::CustomizeHeader( TSharedRef<IPropertyHandle> Str
 
 	for ( auto* Property : PropertiesToAdd )
 	{
-		PropertyOptions.Add(MakeShareable(new FString(FString::Printf(TEXT("%s.%s"), *Property->GetOwner().GetName(), *Property->GetName()))));
+		PropertyOptions.Add(MakeShareable(new FString(FString::Printf(TEXT("%s.%s"), *Property->GetOwnerVariant().GetName(), *Property->GetName()))));
 	}
 
 	FProperty* PropertyValue = nullptr;
@@ -97,7 +97,7 @@ TSharedPtr<FString> FAttributePropertyDetails::GetPropertyType() const
 		PropertyValue = ObjPtr;
 		if (PropertyValue)
 		{
-			FString FullString = PropertyValue->GetOwner().GetName() + TEXT(".") + PropertyValue->GetName();
+			FString FullString = PropertyValue->GetOwnerVariant().GetName() + TEXT(".") + PropertyValue->GetName();
 			for (int32 i=0; i < PropertyOptions.Num(); ++i)
 			{
 				if (PropertyOptions[i].IsValid() && PropertyOptions[i].Get()->Equals(FullString))
diff --git a/Engine/Plugins/Runtime/GameplayAbilities/Source/GameplayAbilitiesEditor/Private/GameplayEffectExecutionScopedModifierInfoDetails.cpp b/Engine/Plugins/Runtime/GameplayAbilities/Source/GameplayAbilitiesEditor/Private/GameplayEffectExecutionScopedModifierInfoDetails.cpp
--- a/Engine/Plugins/Runtime/GameplayAbilities/Source/GameplayAbilitiesEditor/Private/GameplayEffectExecutionScopedModifierInfoDetails.cpp
+++ b/Engine/Plugins/Runtime/GameplayAbilities/Source/GameplayAbilitiesEditor/Private/GameplayEffectExecutionScopedModifierInfoDetails.cpp
@@ -262,7 +262,7 @@ void FGameplayEffectExecutionScopedModifierInfoDetails::CustomizeChildren(TShare
 	ScopedModifierStructPropertyHandle = StructPropertyHandle;
 
 	TSharedPtr<IPropertyHandle> ParentArrayHandle = StructPropertyHandle->GetParentHandle();
-	const bool bIsExecutionDefAttribute = (ParentArrayHandle.IsValid() && ParentArrayHandle->GetProperty()->GetOwner().ToUObject() == FGameplayEffectExecutionDefinition::StaticStruct());
+	const bool bIsExecutionDefAttribute = (ParentArrayHandle.IsValid() && ParentArrayHandle->GetProperty()->GetOwner<UObject>() == FGameplayEffectExecutionDefinition::StaticStruct());
 	
 	if (bIsExecutionDefAttribute)
 	{
diff --git a/Engine/Source/Developer/BlueprintCompilerCppBackend/Private/BlueprintCompilerCppBackendValueHelper.cpp b/Engine/Source/Developer/BlueprintCompilerCppBackend/Private/BlueprintCompilerCppBackendValueHelper.cpp
--- a/Engine/Source/Developer/BlueprintCompilerCppBackend/Private/BlueprintCompilerCppBackendValueHelper.cpp
+++ b/Engine/Source/Developer/BlueprintCompilerCppBackend/Private/BlueprintCompilerCppBackendValueHelper.cpp
@@ -1195,7 +1195,7 @@ struct FFakeImportTableHelper
 					// Should we try to fix it at this stage?
 
 					const bool bIsParam = (0 != (OwnerProperty->PropertyFlags & CPF_Parm)) && OwnerProperty->IsIn(InStruct);
-					const bool bIsMemberVariable = (OwnerProperty->GetOwner().ToUObject() == InStruct);
+					const bool bIsMemberVariable = (OwnerProperty->GetOwner<UObject>() == InStruct);
 					if (bIsParam || bIsMemberVariable) // Affects the class signature. It is necessary while ZCOnstructor/linking.
 					{
 						TArray<UObject*> LocalPreloadDependencies;
diff --git a/Engine/Source/Editor/BlueprintGraph/Private/BlueprintVariableNodeSpawner.cpp b/Engine/Source/Editor/BlueprintGraph/Private/BlueprintVariableNodeSpawner.cpp
--- a/Engine/Source/Editor/BlueprintGraph/Private/BlueprintVariableNodeSpawner.cpp
+++ b/Engine/Source/Editor/BlueprintGraph/Private/BlueprintVariableNodeSpawner.cpp
@@ -255,7 +255,7 @@ FFieldVariant UBlueprintVariableNodeSpawner::GetVarOuter() const
 	}
 	else if (FProperty const* MemberVariable = GetVarProperty())
 	{
-		VarOuter = MemberVariable->GetOwner();
+		VarOuter = MemberVariable->GetOwnerVariant();
 	}	
 	return VarOuter;
 }
diff --git a/Engine/Source/Editor/DetailCustomizations/Private/ConfigEditorPropertyDetails.cpp b/Engine/Source/Editor/DetailCustomizations/Private/ConfigEditorPropertyDetails.cpp
--- a/Engine/Source/Editor/DetailCustomizations/Private/ConfigEditorPropertyDetails.cpp
+++ b/Engine/Source/Editor/DetailCustomizations/Private/ConfigEditorPropertyDetails.cpp
@@ -33,7 +33,7 @@ void FConfigPropertyHelperDetails::CustomizeDetails(IDetailLayoutBuilder& Detail
 	OriginalProperty = CastFieldChecked<FProperty>(PropValue);
 
 	// Create a unique runtime UClass with the provided property as the only member. We will use this in the details view for the config hierarchy.
-	FName TempClassName = *(FString(TEXT("TempConfigEditorUClass_")) + OriginalProperty->GetOwner().GetName() + OriginalProperty->GetName());
+	FName TempClassName = *(FString(TEXT("TempConfigEditorUClass_")) + OriginalProperty->GetOwnerVariant().GetName() + OriginalProperty->GetName());
 	ConfigEditorPropertyViewClass = NewObject<UClass>(GetTransientPackage(), TempClassName, RF_Standalone);
 
 	// Keep a record of the FProperty we are looking to update
diff --git a/Engine/Source/Editor/DetailCustomizations/Private/FbxImportUIDetails.cpp b/Engine/Source/Editor/DetailCustomizations/Private/FbxImportUIDetails.cpp
--- a/Engine/Source/Editor/DetailCustomizations/Private/FbxImportUIDetails.cpp
+++ b/Engine/Source/Editor/DetailCustomizations/Private/FbxImportUIDetails.cpp
@@ -560,7 +560,7 @@ void FFbxImportUIDetails::CustomizeDetails( IDetailLayoutBuilder& DetailBuilder
 		for (TSharedPtr<IPropertyHandle> Handle : ExtraProperties)
 		{
 			// We ignore base import data for this window.
-			if (Handle->GetProperty()->GetOwner().ToUObject() == UFbxTextureImportData::StaticClass())
+			if (Handle->GetProperty()->GetOwner<UObject>() == UFbxTextureImportData::StaticClass())
 			{
 				if (Handle->GetProperty()->GetFName() == GET_MEMBER_NAME_CHECKED(UFbxTextureImportData, MaterialSearchLocation))
 				{
diff --git a/Engine/Source/Editor/Kismet/Private/BlueprintCompilationManager.cpp b/Engine/Source/Editor/Kismet/Private/BlueprintCompilationManager.cpp
--- a/Engine/Source/Editor/Kismet/Private/BlueprintCompilationManager.cpp
+++ b/Engine/Source/Editor/Kismet/Private/BlueprintCompilationManager.cpp
@@ -2874,7 +2874,7 @@ FArchive& FBlueprintCompilationManagerImpl::FFixupBytecodeReferences::operator<<
 			{
 				// drill into new class finding equivalent object:
 				TArray<FFieldVariant> OwnerChain;
-				for (FFieldVariant Iter = Field; Iter.IsValid() && Iter.Get<UObject>(); Iter = Iter.GetOwner())
+				for (FFieldVariant Iter = Field; Iter.IsValid() && Iter.Get<UObject>(); Iter = Iter.GetOwnerVariant())
 				{
 					OwnerChain.Add(Iter);
 				}
diff --git a/Engine/Source/Editor/Kismet/Private/BlueprintDetailsCustomization.cpp b/Engine/Source/Editor/Kismet/Private/BlueprintDetailsCustomization.cpp
--- a/Engine/Source/Editor/Kismet/Private/BlueprintDetailsCustomization.cpp
+++ b/Engine/Source/Editor/Kismet/Private/BlueprintDetailsCustomization.cpp
@@ -4095,7 +4095,7 @@ bool FBlueprintDelegateActionDetails::IsBlueprintProperty() const
 	const UBlueprint* Blueprint = GetBlueprintObj();
 	if(Property && Blueprint)
 	{
-		return (Property->GetOwner().ToUObject() == Blueprint->SkeletonGeneratedClass);
+		return (Property->GetOwner<UObject>() == Blueprint->SkeletonGeneratedClass);
 	}
 
 	return false;
diff --git a/Engine/Source/Editor/Kismet/Private/SKismetInspector.cpp b/Engine/Source/Editor/Kismet/Private/SKismetInspector.cpp
--- a/Engine/Source/Editor/Kismet/Private/SKismetInspector.cpp
+++ b/Engine/Source/Editor/Kismet/Private/SKismetInspector.cpp
@@ -784,7 +784,7 @@ bool SKismetInspector::IsAnyParentContainerSelected(const FPropertyAndParent& Pr
 {
 	for (const FProperty* CurrentProperty : PropertyAndParent.ParentProperties)
 	{
-		const FProperty* CurrentOuter = CurrentProperty->GetOwner().Get<FProperty>();
+		const FProperty* CurrentOuter = CurrentProperty->GetOwner<FProperty>();
 
 		if (CurrentOuter != nullptr && SelectedObjectProperties.Find(const_cast<FProperty*>(CurrentOuter)))
 		{
diff --git a/Engine/Source/Editor/Persona/Private/AnimNotifyDetails.cpp b/Engine/Source/Editor/Persona/Private/AnimNotifyDetails.cpp
--- a/Engine/Source/Editor/Persona/Private/AnimNotifyDetails.cpp
+++ b/Engine/Source/Editor/Persona/Private/AnimNotifyDetails.cpp
@@ -268,7 +268,7 @@ void FAnimNotifyDetails::CustomizeLinkProperties(IDetailLayoutBuilder& Builder,
 		for(uint32 ChildIdx = 0 ; ChildIdx < NumChildProperties ; ++ChildIdx)
 		{
 			TSharedPtr<IPropertyHandle> ChildHandle = NotifyProperty->GetChildHandle(ChildIdx);
-			FString OuterFieldType = ChildHandle->GetProperty()->GetOwner().GetName();
+			FString OuterFieldType = ChildHandle->GetProperty()->GetOwnerVariant().GetName();
 
 			if(ChildHandle->GetProperty()->GetName() == GET_MEMBER_NAME_CHECKED(FAnimNotifyEvent, EndLink).ToString()
 			   || OuterFieldType == FString(TEXT("AnimLinkableElement")))
@@ -314,7 +314,7 @@ void FAnimNotifyDetails::HideLinkProperties(IDetailLayoutBuilder& Builder, TShar
 		for(uint32 ChildIdx = 0 ; ChildIdx < NumChildProperties ; ++ChildIdx)
 		{
 			TSharedPtr<IPropertyHandle> ChildHandle = NotifyProperty->GetChildHandle(ChildIdx);
-			FString OuterFieldType = ChildHandle->GetProperty()->GetOwner().GetName();
+			FString OuterFieldType = ChildHandle->GetProperty()->GetOwnerVariant().GetName();
 			if(ChildHandle->GetProperty()->GetName() == GET_MEMBER_NAME_CHECKED(FAnimNotifyEvent, EndLink).ToString()
 			   || OuterFieldType == FString(TEXT("AnimLinkableElement")))
 			{
diff --git a/Engine/Source/Editor/PropertyEditor/Private/EditConditionContext.cpp b/Engine/Source/Editor/PropertyEditor/Private/EditConditionContext.cpp
--- a/Engine/Source/Editor/PropertyEditor/Private/EditConditionContext.cpp
+++ b/Engine/Source/Editor/PropertyEditor/Private/EditConditionContext.cpp
@@ -91,7 +91,7 @@ T* FindTypedField(const TWeakPtr<FPropertyNode>& PropertyNode, const FString& Pr
 static FPropertyNode* GetEditConditionParentNode(const TSharedPtr<FPropertyNode>& PropertyNode)
 {
 	FPropertyNode* ParentNode = PropertyNode->GetParentNode();
-	FFieldVariant PropertyOuter = PropertyNode->GetProperty()->GetOwner();
+	FFieldVariant PropertyOuter = PropertyNode->GetProperty()->GetOwnerVariant();
 
 	if (PropertyOuter.Get<FArrayProperty>() != nullptr ||
 		PropertyOuter.Get<FSetProperty>() != nullptr ||
diff --git a/Engine/Source/Editor/PropertyEditor/Private/UserInterface/PropertyEditor/SPropertyEditorAsset.cpp b/Engine/Source/Editor/PropertyEditor/Private/UserInterface/PropertyEditor/SPropertyEditorAsset.cpp
--- a/Engine/Source/Editor/PropertyEditor/Private/UserInterface/PropertyEditor/SPropertyEditorAsset.cpp
+++ b/Engine/Source/Editor/PropertyEditor/Private/UserInterface/PropertyEditor/SPropertyEditorAsset.cpp
@@ -38,7 +38,7 @@ DECLARE_DELEGATE( FOnPaste );
 // Helper to retrieve the correct property that has the applicable metadata.
 static const FProperty* GetActualMetadataProperty(const FProperty* Property)
 {
-	if (FProperty* OuterProperty = CastField<FProperty>(Property->GetOwner().ToField()))
+	if (FProperty* OuterProperty = Property->GetOwner<FProperty>())
 	{
 		if (OuterProperty->IsA<FArrayProperty>()
 			|| OuterProperty->IsA<FSetProperty>()
diff --git a/Engine/Source/Editor/UMGEditor/Private/Customizations/UMGDetailCustomizations.cpp b/Engine/Source/Editor/UMGEditor/Private/Customizations/UMGDetailCustomizations.cpp
--- a/Engine/Source/Editor/UMGEditor/Private/Customizations/UMGDetailCustomizations.cpp
+++ b/Engine/Source/Editor/UMGEditor/Private/Customizations/UMGDetailCustomizations.cpp
@@ -362,7 +362,7 @@ void FBlueprintWidgetCustomization::CustomizeAccessibilityProperty(IDetailLayout
 
 	TSharedRef<IPropertyHandle> AccessibleTextPropertyHandle = DetailLayout.GetProperty(TextPropertyName);
 	const FName DelegateName(*(TextPropertyName.ToString() + "Delegate"));
-	FDelegateProperty* AccessibleTextDelegateProperty = FindFieldChecked<FDelegateProperty>(CastChecked<UClass>(AccessibleTextPropertyHandle->GetProperty()->GetOwner().ToUObject()), DelegateName);
+	FDelegateProperty* AccessibleTextDelegateProperty = FindFieldChecked<FDelegateProperty>(CastChecked<UClass>(AccessibleTextPropertyHandle->GetProperty()->GetOwner<UObject>()), DelegateName);
 	// Make sure the old AccessibleText properties are hidden so we don't get duplicate widgets
 	DetailLayout.HideProperty(AccessibleTextPropertyHandle);
 
diff --git a/Engine/Source/Editor/UnrealEd/Private/Kismet2/KismetDebugUtilities.cpp b/Engine/Source/Editor/UnrealEd/Private/Kismet2/KismetDebugUtilities.cpp
--- a/Engine/Source/Editor/UnrealEd/Private/Kismet2/KismetDebugUtilities.cpp
+++ b/Engine/Source/Editor/UnrealEd/Private/Kismet2/KismetDebugUtilities.cpp
@@ -1214,7 +1214,7 @@ FKismetDebugUtilities::EWatchTextResult FKismetDebugUtilities::FindDebuggingData
 			}
 
 			// Try at member scope if it wasn't part of a current function scope
-			UClass* PropertyClass = Cast<UClass>(Property->GetOwner().ToUObject());
+			UClass* PropertyClass = Property->GetOwner<UClass>();
 			if (!PropertyBase && PropertyClass)
 			{
 				if (ActiveObject->GetClass()->IsChildOf(PropertyClass))
@@ -1236,7 +1236,7 @@ FKismetDebugUtilities::EWatchTextResult FKismetDebugUtilities::FindDebuggingData
 			}
 #if USE_UBER_GRAPH_PERSISTENT_FRAME
 			// Try find the propertybase in the persistent ubergraph frame
-			UFunction* OuterFunction = Cast<UFunction>(Property->GetOwner().ToUObject());
+			UFunction* OuterFunction = Property->GetOwner<UFunction>();
 			if (!PropertyBase && OuterFunction)
 			{
 				UBlueprintGeneratedClass* BPGC = Cast<UBlueprintGeneratedClass>(Blueprint->GeneratedClass);
diff --git a/Engine/Source/Editor/UnrealEd/Private/Kismet2/KismetReinstanceUtilities.cpp b/Engine/Source/Editor/UnrealEd/Private/Kismet2/KismetReinstanceUtilities.cpp
--- a/Engine/Source/Editor/UnrealEd/Private/Kismet2/KismetReinstanceUtilities.cpp
+++ b/Engine/Source/Editor/UnrealEd/Private/Kismet2/KismetReinstanceUtilities.cpp
@@ -376,7 +376,7 @@ void FBlueprintCompileReinstancer::SaveClassFieldMapping(UClass* InClassToReinst
 {
 	check(InClassToReinstance);
 
-	for (FProperty* Prop = InClassToReinstance->PropertyLink; Prop && (Prop->GetOwner().ToUObject() == InClassToReinstance); Prop = Prop->PropertyLinkNext)
+	for (FProperty* Prop = InClassToReinstance->PropertyLink; Prop && (Prop->GetOwner<UObject>() == InClassToReinstance); Prop = Prop->PropertyLinkNext)
 	{
 		PropertyMap.Add(Prop->GetFName(), Prop);
 	}
diff --git a/Engine/Source/Programs/UnrealHeaderTool/Private/CodeGenerator.cpp b/Engine/Source/Programs/UnrealHeaderTool/Private/CodeGenerator.cpp
--- a/Engine/Source/Programs/UnrealHeaderTool/Private/CodeGenerator.cpp
+++ b/Engine/Source/Programs/UnrealHeaderTool/Private/CodeGenerator.cpp
@@ -61,6 +61,7 @@
 #include "UObject/FieldPath.h"
 
 #include "UObject/WeakFieldPtr.h"
+#include "Templates/SubclassOf.h"
 
 /////////////////////////////////////////////////////
 // Globals
@@ -1276,7 +1277,7 @@ void FNativeClassHeaderGenerator::PropertyNew(FOutputDevice& DeclOut, FOutputDev
 	{
 		FString OuterSize;
 		FString Setter;
-		if (!Prop->GetOwner().IsUObject())
+		if (!Prop->GetOwner<UObject>())
 		{
 			OuterSize = TEXT("0");
 			Setter    = TEXT("nullptr");
@@ -1722,7 +1723,7 @@ bool IsEditorOnlyDataProperty(FProperty* Prop)
 			return true;
 		}
 
-		Prop = CastField<FProperty>(Prop->GetOwner().ToField());
+		Prop = Prop->GetOwner<FProperty>();
 	}
 
 	return false;
@@ -1836,7 +1837,7 @@ void FNativeClassHeaderGenerator::OutputProperty(FOutputDevice& DeclOut, FOutput
 
 	{
 		FString SourceStruct;
-		if (UFunction* Function = Cast<UFunction>(Prop->GetOwner().ToUObject()))
+		if (UFunction* Function = Prop->GetOwner<UFunction>())
 		{
 			while (Function->GetSuperFunction())
 			{
@@ -1852,7 +1853,7 @@ void FNativeClassHeaderGenerator::OutputProperty(FOutputDevice& DeclOut, FOutput
 		}
 		else
 		{
-			SourceStruct = NameLookupCPP.GetNameCPP(CastChecked<UStruct>(Prop->GetOwner().ToUObject()));
+			SourceStruct = NameLookupCPP.GetNameCPP(CastChecked<UStruct>(Prop->GetOwner<UObject>()));
 		}
 
 		FString PropNameDep = PropName;
@@ -6358,12 +6359,78 @@ ECompilationResult::Type PreparseModules(const FString& ModuleInfoPath, int32& N
 	return Result;
 }
 
+struct FTestCollector
+{
+	void HandleObjectReference(UObject*& Obj)
+	{
+
+	}
+
+	template <typename T>
+	typename TEnableIf<TIsDerivedFrom<T, FField>::IsDerived, void>::Type AddReferencedObject(T*& Field)
+	{
+		if (Field)
+		{
+			UObject* Owner = Field->GetOwnerUObject();
+			UObject* OldOwnerPtr = nullptr;
+			HandleObjectReference(Owner);
+			if (Owner == nullptr && OldOwnerPtr != nullptr)
+			{
+				Field = nullptr;
+			}
+		}
+	}
+};
+
 ECompilationResult::Type UnrealHeaderTool_Main(const FString& ModuleInfoFilename)
 {
 	double MainTime = 0.0;
 	FDurationTimer MainTimer(MainTime);
 	MainTimer.Start();
 
+	//FFieldVariant Variant;
+	//TFieldPath<FProperty> PropPath;
+	//FProperty* RawProp = nullptr;
+
+	//Variant.IsA<FProperty>();
+	//Variant.IsA<UObject>();
+	//Variant.Get<FProperty>();
+	//Variant.Get<UObject>();
+	//RawProp->GetOwner<FProperty>();
+	//RawProp->GetOwner<UObject>();
+	//RawProp->GetOwnerChecked<FProperty>();
+	//RawProp->GetOwnerChecked<UObject>();
+	//RawProp->GetTypedOwner<FProperty>();
+	//RawProp->GetTypedOwner<UObject>();
+	//{
+	//	UStruct* Struct = nullptr;
+	//	for (TFieldIterator<FProperty> It(Struct); It; ++It)
+	//	{
+	//		check(It);
+	//	}
+	//	for (TFieldIterator<UFunction> It(Struct); It; ++It)
+	//	{
+	//		check(It);
+	//	}
+	//}
+	//{
+	//	FArchive* Ar = nullptr;
+	//	TWeakFieldPtr<FProperty> WeakPtr;
+	//	(*Ar) << WeakPtr;
+	//}
+	//{
+	//	FProperty* Prop = nullptr;
+	//	UObject* Obj = nullptr;
+	//	FTestCollector Collector;
+	//	Collector.AddReferencedObject(Prop);
+	//	Collector.AddReferencedObject(Obj);
+	//}
+	//{
+	//	TSubclassOf<FProperty> SubProp;
+	//	SubProp.GetDefaultObject();
+	//	TSubclassOf<UFunction> SubFn;
+	//	SubFn.GetDefaultObject();
+	//}
 	check(GIsUCCMakeStandaloneHeaderGenerator);
 	ECompilationResult::Type Result = ECompilationResult::Succeeded;
 
diff --git a/Engine/Source/Programs/UnrealHeaderTool/Private/HeaderParser.cpp b/Engine/Source/Programs/UnrealHeaderTool/Private/HeaderParser.cpp
--- a/Engine/Source/Programs/UnrealHeaderTool/Private/HeaderParser.cpp
+++ b/Engine/Source/Programs/UnrealHeaderTool/Private/HeaderParser.cpp
@@ -5078,7 +5078,7 @@ FProperty* FHeaderParser::GetVarNameAndDim
 				HintText,
 				VarProperty.Identifier,
 				*Scope->GetName(),
-				ExistingField ? *ExistingField->GetOuter()->GetName() : *ExistingProperty->GetOwner().GetFullName());
+				ExistingField ? *ExistingField->GetOuter()->GetName() : *ExistingProperty->GetOwnerVariant().GetFullName());
 		}
 	}
 
diff --git a/Engine/Source/Programs/UnrealHeaderTool/Private/ParserClass.cpp b/Engine/Source/Programs/UnrealHeaderTool/Private/ParserClass.cpp
--- a/Engine/Source/Programs/UnrealHeaderTool/Private/ParserClass.cpp
+++ b/Engine/Source/Programs/UnrealHeaderTool/Private/ParserClass.cpp
@@ -111,7 +111,7 @@ bool FClass::IsOwnedByDynamicType(const UField* Field)
 
 bool FClass::IsOwnedByDynamicType(const FField* Field)
 {
-	for (FFieldVariant Owner = Field->GetOwner(); Owner.IsValid(); Owner = Owner.GetOwner())
+	for (FFieldVariant Owner = Field->GetOwnerVariant(); Owner.IsValid(); Owner = Owner.GetOwnerVariant())
 	{
 		if (Owner.IsUObject())
 		{
diff --git a/Engine/Source/Programs/UnrealHeaderTool/Private/ParserHelper.cpp b/Engine/Source/Programs/UnrealHeaderTool/Private/ParserHelper.cpp
--- a/Engine/Source/Programs/UnrealHeaderTool/Private/ParserHelper.cpp
+++ b/Engine/Source/Programs/UnrealHeaderTool/Private/ParserHelper.cpp
@@ -23,7 +23,7 @@ FTokenData* FClassMetaData::FindTokenData( FProperty* Prop )
 	check(Prop);
 
 	FTokenData* Result = nullptr;
-	UObject* Outer = Prop->GetOwner().ToUObject();
+	UObject* Outer = Prop->GetOwner<UObject>();
 	check(Outer);
 	UClass* OuterClass = nullptr;
 	if (Outer->IsA<UStruct>())
diff --git a/Engine/Source/Programs/UnrealHeaderTool/Private/ParserHelper.h b/Engine/Source/Programs/UnrealHeaderTool/Private/ParserHelper.h
--- a/Engine/Source/Programs/UnrealHeaderTool/Private/ParserHelper.h
+++ b/Engine/Source/Programs/UnrealHeaderTool/Private/ParserHelper.h
@@ -1597,7 +1597,7 @@ public:
 		FProperty* Prop = PropertyToken.TokenProperty;
 		check(Prop);
 
-		UObject* Outer = Prop->GetOwner().ToUObject();
+		UObject* Outer = Prop->GetOwner<UObject>();
 		check(Outer);
 		UStruct* OuterClass = Cast<UStruct>(Outer);
 		if ( OuterClass != NULL )
diff --git a/Engine/Source/Runtime/AIModule/Private/BehaviorTree/BlueprintNodeHelpers.cpp b/Engine/Source/Runtime/AIModule/Private/BehaviorTree/BlueprintNodeHelpers.cpp
--- a/Engine/Source/Runtime/AIModule/Private/BehaviorTree/BlueprintNodeHelpers.cpp
+++ b/Engine/Source/Runtime/AIModule/Private/BehaviorTree/BlueprintNodeHelpers.cpp
@@ -67,7 +67,7 @@ namespace BlueprintNodeHelpers
 		for (FProperty* TestProperty = Ob->GetClass()->PropertyLink; TestProperty; TestProperty = TestProperty->PropertyLinkNext)
 		{
 			// stop when reaching base class
-			if (TestProperty->GetOwner().ToUObject() == StopAtClass)
+			if (TestProperty->GetOwner<UObject>() == StopAtClass)
 			{
 				break;
 			}
@@ -132,7 +132,7 @@ namespace BlueprintNodeHelpers
 		for (FProperty* TestProperty = Ob->GetClass()->PropertyLink; TestProperty; TestProperty = TestProperty->PropertyLinkNext)
 		{
 			// stop when reaching base class
-			if (TestProperty->GetOwner().ToUObject() == StopAtClass)
+			if (TestProperty->GetOwner<UObject>() == StopAtClass)
 			{
 				break;
 			}
@@ -158,7 +158,7 @@ namespace BlueprintNodeHelpers
 		for (FProperty* TestProperty = Ob.GetClass()->PropertyLink; TestProperty; TestProperty = TestProperty->PropertyLinkNext)
 		{
 			// stop when reaching base class
-			if (TestProperty->GetOwner().ToUObject() == &StopAtClass)
+			if (TestProperty->GetOwner<UObject>() == &StopAtClass)
 			{
 				break;
 			}
@@ -179,7 +179,7 @@ namespace BlueprintNodeHelpers
 		for (FProperty* TestProperty = Ob->GetClass()->PropertyLink; TestProperty; TestProperty = TestProperty->PropertyLinkNext)
 		{
 			// stop when reaching base class
-			if (TestProperty->GetOwner().ToUObject() == StopAtClass)
+			if (TestProperty->GetOwner<UObject>() == StopAtClass)
 			{
 				break;
 			}
@@ -210,7 +210,7 @@ namespace BlueprintNodeHelpers
 		for (FProperty* TestProperty = Ob->GetClass()->PropertyLink; TestProperty; TestProperty = TestProperty->PropertyLinkNext)
 		{
 			// stop when reaching base class
-			if (TestProperty->GetOwner().ToUObject() == StopAtClass)
+			if (TestProperty->GetOwner<UObject>() == StopAtClass)
 			{
 				break;
 			}
diff --git a/Engine/Source/Runtime/CoreUObject/Private/UObject/Class.cpp b/Engine/Source/Runtime/CoreUObject/Private/UObject/Class.cpp
--- a/Engine/Source/Runtime/CoreUObject/Private/UObject/Class.cpp
+++ b/Engine/Source/Runtime/CoreUObject/Private/UObject/Class.cpp
@@ -752,7 +752,7 @@ void UStruct::Link(FArchive& Ar, bool bRelinkExistingProperties)
 
 			for (FField* Field = ChildProperties; Field; Field = Field->Next)
 			{
-				if (Field->GetOwner().ToUObject() != this)
+				if (Field->GetOwner<UObject>() != this)
 				{
 					break;
 				}
@@ -763,8 +763,8 @@ void UStruct::Link(FArchive& Ar, bool bRelinkExistingProperties)
 					// If we don't have the editor, make sure we aren't trying to link properties that are editor only.
 					check(!Property->IsEditorOnlyProperty());
 #endif // WITH_EDITORONLY_DATA
-					ensureMsgf(Property->GetOwner().ToUObject() == this, TEXT("Linking '%s'. Property '%s' has outer '%s'"),
-						*GetFullName(), *Property->GetName(), *Property->GetOwner().GetFullName());
+					ensureMsgf(Property->GetOwner<UObject>() == this, TEXT("Linking '%s'. Property '%s' has outer '%s'"),
+						*GetFullName(), *Property->GetName(), *Property->GetOwnerVariant().GetFullName());
 
 					// Linking a property can cause a recompilation of the struct. 
 					// When the property was changed, the struct should be relinked again, to be sure, the PropertiesSize is actual.
@@ -821,7 +821,7 @@ void UStruct::Link(FArchive& Ar, bool bRelinkExistingProperties)
 	}
 	else
 	{
-		for (FField* Field = ChildProperties; (Field != NULL) && (Field->GetOwner().ToUObject() == this); Field = Field->Next)
+		for (FField* Field = ChildProperties; (Field != NULL) && (Field->GetOwner<UObject>() == this); Field = Field->Next)
 		{
 			if (FProperty* Property = CastField<FProperty>(Field))
 			{
@@ -3983,7 +3983,7 @@ void UClass::SetUpRuntimeReplicationData()
 				{
 					NetProperties.Add(Prop);
 
-					if (Prop->GetOwner().ToUObject() == this)
+					if (Prop->GetOwner<UObject>() == this)
 					{
 						NetPropsForClassRep.Add(Prop);
 					}
@@ -4686,13 +4686,13 @@ bool UClass::IsFunctionImplementedInScript(FName InFunctionName) const
 
 bool UClass::HasProperty(FProperty* InProperty) const
 {
-	if (InProperty->GetOwner().IsUObject())
+	if (InProperty->GetOwner<UObject>())
 	{
-		UClass* PropertiesClass = Cast<UClass>(InProperty->GetOwner().ToUObject());
+		UClass* PropertiesClass = InProperty->GetOwner<UClass>();
 		if (PropertiesClass)
-	{
-		return IsChildOf(PropertiesClass);
-	}
+		{
+			return IsChildOf(PropertiesClass);
+		}
 	}
 
 	return false;
diff --git a/Engine/Source/Runtime/CoreUObject/Private/UObject/EditPropertyChain.cpp b/Engine/Source/Runtime/CoreUObject/Private/UObject/EditPropertyChain.cpp
--- a/Engine/Source/Runtime/CoreUObject/Private/UObject/EditPropertyChain.cpp
+++ b/Engine/Source/Runtime/CoreUObject/Private/UObject/EditPropertyChain.cpp
@@ -108,7 +108,7 @@ void FEditPropertyChain::SetListSize( int32 NewListSize )
 			for ( TIterator It(PropertyNode); It; --It )
 			{
 				// if we've found the member property, we can stop here
-				if ( Cast<UClass>(It->GetOwner().ToUObject()) != NULL )
+				if (It->GetOwner<UClass>() != nullptr)
 				{
 					PropertyNode = It.GetNode();
 					break;
diff --git a/Engine/Source/Runtime/CoreUObject/Private/UObject/EnumProperty.cpp b/Engine/Source/Runtime/CoreUObject/Private/UObject/EnumProperty.cpp
--- a/Engine/Source/Runtime/CoreUObject/Private/UObject/EnumProperty.cpp
+++ b/Engine/Source/Runtime/CoreUObject/Private/UObject/EnumProperty.cpp
@@ -400,7 +400,7 @@ EConvertFromTypeResult FEnumProperty::ConvertFromType(const FPropertyTag& Tag, F
 		if (Tag.EnumName == NAME_None)
 		{
 			// If we're a nested property the EnumName tag got lost. Handle this case for backward compatibility reasons
-			FProperty* const PropertyOwner = CastField<FProperty>(GetOwner().ToField());
+			FProperty* const PropertyOwner = GetOwner<FProperty>();
 
 			if (PropertyOwner)
 			{
diff --git a/Engine/Source/Runtime/CoreUObject/Private/UObject/Field.cpp b/Engine/Source/Runtime/CoreUObject/Private/UObject/Field.cpp
--- a/Engine/Source/Runtime/CoreUObject/Private/UObject/Field.cpp
+++ b/Engine/Source/Runtime/CoreUObject/Private/UObject/Field.cpp
@@ -112,7 +112,7 @@ FArchive& operator << (FArchive& Ar, FFieldClass*& InOutFieldClass)
 	return Ar;
 }
 
-FFieldVariant FFieldVariant::GetOwner() const
+FFieldVariant FFieldVariant::GetOwnerVariant() const
 {
 	if (bIsUObject)
 	{
@@ -120,7 +120,7 @@ FFieldVariant FFieldVariant::GetOwner() const
 	}
 	else
 	{
-		return Container.Field->GetOwner();
+		return Container.Field->GetOwnerVariant();
 	}
 }
 
@@ -487,7 +487,7 @@ bool FField::IsIn(UObject* InOwner) const
 
 bool FField::IsIn(FField* InOwner) const
 {
-	for (FField* OwnerField = GetOwner().ToField(); OwnerField; OwnerField = OwnerField->GetOwner().ToField())
+	for (FField* OwnerField = GetOwner<FField>(); OwnerField; OwnerField = OwnerField->GetOwner<FField>())
 	{
 		if (OwnerField == InOwner)
 		{
@@ -512,7 +512,7 @@ void FField::AddCppProperty(FProperty* Property)
 FString FField::GetPathName(const UObject* StopOuter /*= nullptr*/) const
 {
 	FString PathName;
-	for (FFieldVariant TempOwner = Owner; TempOwner.IsValid(); TempOwner = TempOwner.GetOwner())
+	for (FFieldVariant TempOwner = Owner; TempOwner.IsValid(); TempOwner = TempOwner.GetOwnerVariant())
 	{		
 		if (!TempOwner.IsUObject())
 		{
diff --git a/Engine/Source/Runtime/CoreUObject/Private/UObject/FieldPath.cpp b/Engine/Source/Runtime/CoreUObject/Private/UObject/FieldPath.cpp
--- a/Engine/Source/Runtime/CoreUObject/Private/UObject/FieldPath.cpp
+++ b/Engine/Source/Runtime/CoreUObject/Private/UObject/FieldPath.cpp
@@ -28,7 +28,7 @@ void FFieldPath::Generate(FField* InField)
 	if (InField)
 	{
 		// Add names from the innermost to the outermost
-		for (FFieldVariant Iter(InField); Iter.IsValid(); Iter = Iter.GetOwner())
+		for (FFieldVariant Iter(InField); Iter.IsValid(); Iter = Iter.GetOwnerVariant())
 		{
 			Path.Add(Iter.GetFName());
 		}
diff --git a/Engine/Source/Runtime/CoreUObject/Private/UObject/LinkerPlaceholderBase.cpp b/Engine/Source/Runtime/CoreUObject/Private/UObject/LinkerPlaceholderBase.cpp
--- a/Engine/Source/Runtime/CoreUObject/Private/UObject/LinkerPlaceholderBase.cpp
+++ b/Engine/Source/Runtime/CoreUObject/Private/UObject/LinkerPlaceholderBase.cpp
@@ -316,7 +316,7 @@ FLinkerPlaceholderBase::FPlaceholderValuePropertyPath::FPlaceholderValueProperty
 	check(ReferencingProperty.IsA<FProperty>() || ReferencingProperty.IsA<UProperty>());
 
 	PropertyChain.Add(ReferencingProperty);
-	FFieldVariant PropertyOuter = ReferencingProperty.GetOwner();
+	FFieldVariant PropertyOuter = ReferencingProperty.GetOwnerVariant();
 	
 	const TArray<FFieldVariant>& StructPropertyStack = FPlaceholderContainerTracker::Get().IntermediatePropertyStack;
 	int32 StructStackIndex = StructPropertyStack.Num() - 1; // "top" of the array is the last element
@@ -375,7 +375,7 @@ FLinkerPlaceholderBase::FPlaceholderValuePropertyPath::FPlaceholderValueProperty
 				break;
 			}
 		}
-		PropertyOuter = PropertyOuter.GetOwner();
+		PropertyOuter = PropertyOuter.GetOwnerVariant();
 	}
 
 #if USE_DEFERRED_DEPENDENCY_CHECK_VERIFICATION_TESTS
diff --git a/Engine/Source/Runtime/CoreUObject/Private/UObject/Obj.cpp b/Engine/Source/Runtime/CoreUObject/Private/UObject/Obj.cpp
--- a/Engine/Source/Runtime/CoreUObject/Private/UObject/Obj.cpp
+++ b/Engine/Source/Runtime/CoreUObject/Private/UObject/Obj.cpp
@@ -2951,7 +2951,7 @@ void UObject::ParseParms( const TCHAR* Parms )
 	}
 	for( TFieldIterator<FProperty> It(GetClass()); It; ++It )
 	{
-		if (It->GetOwner().ToUObject() != UObject::StaticClass())
+		if (It->GetOwner<UObject>() != UObject::StaticClass())
 		{
 			FString Value;
 			if( FParse::Value(Parms,*(FString(*It->GetName())+TEXT("=")),Value) )
diff --git a/Engine/Source/Runtime/CoreUObject/Private/UObject/Property.cpp b/Engine/Source/Runtime/CoreUObject/Private/UObject/Property.cpp
--- a/Engine/Source/Runtime/CoreUObject/Private/UObject/Property.cpp
+++ b/Engine/Source/Runtime/CoreUObject/Private/UObject/Property.cpp
@@ -471,15 +471,14 @@ void FProperty::Init()
 	checkSlow(GetOwnerUField()->HasAllFlags(RF_Transient));
 	checkSlow(HasAllFlags(RF_Transient));
 
-	if (GetOwner().IsUObject())
+	if (GetOwner<UObject>())
 	{
-		UField* OwnerField = CastChecked<UField>(GetOwner().ToUObject());
+		UField* OwnerField = GetOwnerChecked<UField>();
 		OwnerField->AddCppProperty(this);
 	}
 	else
 	{
-		FField* OwnerField = GetOwner().ToField();
-		check(OwnerField);
+		FField* OwnerField = GetOwnerChecked<FField>();
 		OwnerField->AddCppProperty(this);
 	}
 }
@@ -661,7 +660,7 @@ void FProperty::ExportCppDeclaration(FOutputDevice& Out, EExportedDeclaration::T
 			TypeText = FString::Printf(TEXT("const %s"), *TypeText);
 		}
 
-		const UClass* const MyPotentialConstClass = (DeclarationType == EExportedDeclaration::Member) ? Cast<UClass>(GetOwner().ToUObject()) : nullptr;
+		const UClass* const MyPotentialConstClass = (DeclarationType == EExportedDeclaration::Member) ? GetOwner<UClass>() : nullptr;
 		const bool bFromConstClass = MyPotentialConstClass && MyPotentialConstClass->HasAnyClassFlags(CLASS_Const);
 		const bool bConstAtTheEnd = bFromConstClass || (bIsConstParam && bShouldHaveRef);
 		if (bConstAtTheEnd)
@@ -915,7 +914,7 @@ EConvertFromTypeResult FProperty::ConvertFromType(const FPropertyTag& Tag, FStru
 
 int32 FProperty::SetupOffset()
 {
-	UObject* OwnerUObject = GetOwner().ToUObject();
+	UObject* OwnerUObject = GetOwner<UObject>();
 	if (OwnerUObject && (OwnerUObject->GetClass()->ClassCastFlags & CASTCLASS_UStruct))
 	{
 		UStruct* OwnerStruct = (UStruct*)OwnerUObject;
@@ -1363,7 +1362,7 @@ const TCHAR* FProperty::ImportSingleProperty( const TCHAR* Str, void* DestData,
 
 			// disallow importing of an object's name from here
 			// not done above with ShouldPort() check because this is intentionally exported so we don't want it to cause errors on import
-			if (Property->GetFName() != NAME_Name || !Property->GetOwner().IsUObject() || Property->GetOwner().ToUObject()->GetFName() != NAME_Object)
+			if (Property->GetFName() != NAME_Name || !Property->GetOwnerVariant().IsUObject() || Property->GetOwner<UObject>()->GetFName() != NAME_Object)
 			{
 				if (Index > -1 && ArrayProperty != NULL) //set single dynamic array element
 				{
diff --git a/Engine/Source/Runtime/CoreUObject/Private/UObject/ScriptCore.cpp b/Engine/Source/Runtime/CoreUObject/Private/UObject/ScriptCore.cpp
--- a/Engine/Source/Runtime/CoreUObject/Private/UObject/ScriptCore.cpp
+++ b/Engine/Source/Runtime/CoreUObject/Private/UObject/ScriptCore.cpp
@@ -1977,7 +1977,7 @@ DEFINE_FUNCTION(UObject::execInstanceVariable)
 	FProperty* VarProperty = (FProperty*)Stack.ReadObject();
 	Stack.MostRecentProperty = VarProperty;
 
-	if (VarProperty == nullptr || !P_THIS->IsA((UClass*)VarProperty->GetOwner().ToUObject()))
+	if (VarProperty == nullptr || !P_THIS->IsA(VarProperty->GetOwner<UClass>()))
 	{
 		FBlueprintExceptionInfo ExceptionInfo(EBlueprintExceptionType::AccessViolation, FText::Format(LOCTEXT("MissingProperty", "Attempted to access missing property '{0}'. If this is a packaged/cooked build, are you attempting to use an editor-only property?"), FText::FromString(GetNameSafe(VarProperty))));
 		FBlueprintCoreDelegates::ThrowScriptException(P_THIS, Stack, ExceptionInfo);
@@ -2043,7 +2043,7 @@ DEFINE_FUNCTION(UObject::execDefaultVariable)
 		// @todo - allow access to archetype properties through object references?
 	}
 
-	if (VarProperty == nullptr || (DefaultObject && !DefaultObject->IsA((UClass*)VarProperty->GetOwner().ToUObject())))
+	if (VarProperty == nullptr || (DefaultObject && !DefaultObject->IsA(VarProperty->GetOwner<UClass>())))
 	{
 		FBlueprintExceptionInfo ExceptionInfo(EBlueprintExceptionType::AccessViolation, LOCTEXT("MissingPropertyDefaultObject", "Attempted to access a missing property on a CDO. If this is a packaged/cooked build, are you attempting to use an editor-only property?"));
 		FBlueprintCoreDelegates::ThrowScriptException(P_THIS, Stack, ExceptionInfo);
diff --git a/Engine/Source/Runtime/CoreUObject/Private/UObject/UObjectGlobals.cpp b/Engine/Source/Runtime/CoreUObject/Private/UObject/UObjectGlobals.cpp
--- a/Engine/Source/Runtime/CoreUObject/Private/UObject/UObjectGlobals.cpp
+++ b/Engine/Source/Runtime/CoreUObject/Private/UObject/UObjectGlobals.cpp
@@ -1245,7 +1245,7 @@ UPackage* LoadPackageInternal(UPackage* InOuter, const TCHAR* InLongPackageNameO
 					bIsEditorOnly = true;
 					break;
 				}
-				SerializingProperty = CastField<FProperty>(SerializingProperty->GetOwner().ToField());
+				SerializingProperty = SerializingProperty->GetOwner<FProperty>();
 			}
 
 			if (!bIsEditorOnly)
diff --git a/Engine/Source/Runtime/CoreUObject/Private/UObject/UnrealType.cpp b/Engine/Source/Runtime/CoreUObject/Private/UObject/UnrealType.cpp
--- a/Engine/Source/Runtime/CoreUObject/Private/UObject/UnrealType.cpp
+++ b/Engine/Source/Runtime/CoreUObject/Private/UObject/UnrealType.cpp
@@ -169,7 +169,7 @@ void FPropertyValueIterator::GetPropertyChain(TArray<const FProperty*>& Property
 		{
 			// This handles container property nesting
 			PropertyChain.Add(Property);
-			Property = Property->GetOwner().Get<FProperty>();
+			Property = Property->GetOwner<FProperty>();
 		}
 	}
 }
\ No newline at end of file
diff --git a/Engine/Source/Runtime/CoreUObject/Public/UObject/Field.h b/Engine/Source/Runtime/CoreUObject/Public/UObject/Field.h
--- a/Engine/Source/Runtime/CoreUObject/Public/UObject/Field.h
+++ b/Engine/Source/Runtime/CoreUObject/Public/UObject/Field.h
@@ -296,7 +296,7 @@ public:
 	{
 		return Container.Field;
 	}
-	FFieldVariant GetOwner() const;
+	FFieldVariant GetOwnerVariant() const;
 	UClass* GetOwnerClass() const;
 	FString GetFullName() const;
 	FString GetPathName() const;
@@ -494,7 +494,7 @@ public:
 	}
 
 	/** Gets the owner container for this field */
-	FFieldVariant GetOwner() const
+	FFieldVariant GetOwnerVariant() const
 	{
 		return Owner;
 	}
@@ -537,7 +537,9 @@ public:
 	}
 
 	template <typename T>
+	FUNCTION_NON_NULL_RETURN_START
 	T* GetOwnerChecked() const
+	FUNCTION_NON_NULL_RETURN_END
 	{
 		static_assert(sizeof(T) > 0, "T must not be an incomplete type");
 		T* Result = Owner.Get<T>();
diff --git a/Engine/Source/Runtime/CoreUObject/Public/UObject/UnrealType.h b/Engine/Source/Runtime/CoreUObject/Public/UObject/UnrealType.h
--- a/Engine/Source/Runtime/CoreUObject/Public/UObject/UnrealType.h
+++ b/Engine/Source/Runtime/CoreUObject/Public/UObject/UnrealType.h
@@ -354,7 +354,7 @@ private:
 		if (0)
 		{
 			// in the future, these checks will be tested if the property is NOT relative to a UClass
-			check(!Cast<UClass>(GetOwner().ToUObject())); // Check we are _not_ calling this on a direct child property of a UClass, you should pass in a UObject* in that case
+			check(!GetOwner<UClass>()); // Check we are _not_ calling this on a direct child property of a UClass, you should pass in a UObject* in that case
 		}
 
 		return (uint8*)ContainerPtr + Offset_Internal + ElementSize * ArrayIndex;
@@ -369,19 +369,19 @@ private:
 		// need something for networking, since those are NOT live uobjects, just memory blocks
 		check(((UObject*)ContainerPtr)->IsValidLowLevel()); // Check its a valid UObject that was passed in
 		check(((UObject*)ContainerPtr)->GetClass() != NULL);
-		check(Cast<UClass>(GetOwner().ToUObject())); // Check that the outer of this property is a UClass (not another property)
+		check(GetOwner<UClass>()); // Check that the outer of this property is a UClass (not another property)
 
 		// Check that the object we are accessing is of the class that contains this property
-		checkf(((UObject*)ContainerPtr)->IsA((UClass*)GetOwner().ToUObject()), TEXT("'%s' is of class '%s' however property '%s' belongs to class '%s'")
+		checkf(((UObject*)ContainerPtr)->IsA(GetOwner<UClass>()), TEXT("'%s' is of class '%s' however property '%s' belongs to class '%s'")
 			, *((UObject*)ContainerPtr)->GetName()
 			, *((UObject*)ContainerPtr)->GetClass()->GetName()
 			, *GetName()
-			, *((UClass*)GetOwner().ToUObject())->GetName());
+			, *GetOwner<UClass>()->GetName());
 
 		if (0)
 		{
 			// in the future, these checks will be tested if the property is NOT relative to a UClass
-			check(!GetOwner().ToUObject()->IsA(UClass::StaticClass())); // Check we are _not_ calling this on a direct child property of a UClass, you should pass in a UObject* in that case
+			check(!GetOwner<UClass>()); // Check we are _not_ calling this on a direct child property of a UClass, you should pass in a UObject* in that case
 		}
 
 		return (uint8*)ContainerPtr + Offset_Internal + ElementSize * ArrayIndex;
@@ -761,7 +761,7 @@ public:
 	FProperty* GetOwnerProperty()
 	{
 		FProperty* Result = this;
-		for (FProperty* PropBase = GetOwner().Get<FProperty>(); PropBase; PropBase = PropBase->GetOwner().Get<FProperty>())
+		for (FProperty* PropBase = GetOwner<FProperty>(); PropBase; PropBase = PropBase->GetOwner<FProperty>())
 		{
 			Result = PropBase;
 		}
@@ -771,7 +771,7 @@ public:
 	const FProperty* GetOwnerProperty() const
 	{
 		const FProperty* Result = this;
-		for (const FProperty* PropBase = GetOwner().Get<FProperty>(); PropBase; PropBase = PropBase->GetOwner().Get<FProperty>())
+		for (const FProperty* PropBase = GetOwner<FProperty>(); PropBase; PropBase = PropBase->GetOwner<FProperty>())
 		{
 			Result = PropBase;
 		}
diff --git a/Engine/Source/Runtime/Engine/Classes/Engine/MemberReference.h b/Engine/Source/Runtime/Engine/Classes/Engine/MemberReference.h
--- a/Engine/Source/Runtime/Engine/Classes/Engine/MemberReference.h
+++ b/Engine/Source/Runtime/Engine/Classes/Engine/MemberReference.h
@@ -584,7 +584,7 @@ inline UFunction* FMemberReference::ResolveUFunction() const
 template<>
 inline UObject* FMemberReference::GetFieldOuter(FField* Field) const
 {
-	return Field->GetOwner().ToUObject();
+	return Field->GetOwner<UObject>();
 }
 
 template <>
diff --git a/Engine/Source/Runtime/Engine/Private/ActorConstruction.cpp b/Engine/Source/Runtime/Engine/Private/ActorConstruction.cpp
--- a/Engine/Source/Runtime/Engine/Private/ActorConstruction.cpp
+++ b/Engine/Source/Runtime/Engine/Private/ActorConstruction.cpp
@@ -97,7 +97,7 @@ void AActor::ResetPropertiesForConstruction()
 	{
 		FProperty* Prop = *It;
 		FStructProperty* StructProp = CastField<FStructProperty>(Prop);
-		UClass* PropClass = CastChecked<UClass>(Prop->GetOwner().ToUObject()); // get the class that added this property
+		UClass* PropClass = Prop->GetOwnerChecked<UClass>(); // get the class that added this property
 
 		// First see if it is a random stream, if so reset before running construction script
 		if( (StructProp != nullptr) && (StructProp->Struct != nullptr) && (StructProp->Struct->GetFName() == RandomStreamName) )
diff --git a/Engine/Source/Runtime/Engine/Private/Animation/AnimNodeBase.cpp b/Engine/Source/Runtime/Engine/Private/Animation/AnimNodeBase.cpp
--- a/Engine/Source/Runtime/Engine/Private/Animation/AnimNodeBase.cpp
+++ b/Engine/Source/Runtime/Engine/Private/Animation/AnimNodeBase.cpp
@@ -683,7 +683,7 @@ void FExposedValueHandler::Initialize(UObject* AnimInstanceObject, int32 NodeOff
 			if (CopyRecord.bInstanceIsTarget)
 			{
 				// Re-find our dest property as it (or its class outer) may have changed
-				if (CopyRecord.DestProperty->GetOwner().ToUObject() != AnimInstanceObject->GetClass())
+				if (CopyRecord.DestProperty->GetOwner<UObject>() != AnimInstanceObject->GetClass())
 				{
 					CopyRecord.DestProperty = AnimInstanceObject->GetClass()->FindPropertyByName(CopyRecord.DestProperty->GetFName());
 				}
diff --git a/Engine/Source/Runtime/Engine/Private/ComponentInstanceDataCache.cpp b/Engine/Source/Runtime/Engine/Private/ComponentInstanceDataCache.cpp
--- a/Engine/Source/Runtime/Engine/Private/ComponentInstanceDataCache.cpp
+++ b/Engine/Source/Runtime/Engine/Private/ComponentInstanceDataCache.cpp
@@ -56,7 +56,7 @@ public:
 	virtual bool ShouldSkipProperty(const FProperty* InProperty) const override
 	{
 		// Immutable structs expect to serialize all properties so don't skip regardless of other conditions
-		UScriptStruct* ScriptStruct = Cast<UScriptStruct>(InProperty->GetOwner().ToUObject());
+		UScriptStruct* ScriptStruct = InProperty->GetOwner<UScriptStruct>();
 		const bool bPropertyInImmutableStruct = ScriptStruct && ((ScriptStruct->StructFlags & STRUCT_Immutable) != 0);
 
 		return (!bPropertyInImmutableStruct
diff --git a/Engine/Source/Runtime/Engine/Private/KismetArrayLibrary.cpp b/Engine/Source/Runtime/Engine/Private/KismetArrayLibrary.cpp
--- a/Engine/Source/Runtime/Engine/Private/KismetArrayLibrary.cpp
+++ b/Engine/Source/Runtime/Engine/Private/KismetArrayLibrary.cpp
@@ -1,4 +1,4 @@
-﻿// Copyright 1998-2019 Epic Games, Inc. All Rights Reserved.
+// Copyright 1998-2019 Epic Games, Inc. All Rights Reserved.
 
 #include "Kismet/KismetArrayLibrary.h"
 #include "GameFramework/Actor.h"
@@ -296,7 +296,7 @@ void UKismetArrayLibrary::GenericArray_Get(void* TargetArray, const FArrayProper
 				Index,
 				*ArrayProp->GetName(),
 				ArrayHelper.Num(),
-				ArrayProp->GetOwner().IsUObject() ? *GetPathNameSafe(ArrayProp->GetOwner().ToUObject()) : *GetPathNameSafe(ArrayProp->GetOwner().ToField())),
+				*ArrayProp->GetOwnerVariant().GetPathName()),
 				ELogVerbosity::Warning,
 				GetOutOfBoundsWarning);
 			InnerProp->InitializeValue(Item);
diff --git a/Engine/Source/Runtime/Engine/Private/WorldSettings.cpp b/Engine/Source/Runtime/Engine/Private/WorldSettings.cpp
--- a/Engine/Source/Runtime/Engine/Private/WorldSettings.cpp
+++ b/Engine/Source/Runtime/Engine/Private/WorldSettings.cpp
@@ -525,8 +525,8 @@ bool AWorldSettings::CanEditChange(const FProperty* InProperty) const
 	{
 		FString PropertyName = InProperty->GetName();
 
-		if (InProperty->GetOwner().IsUObject() && InProperty->GetOwner().IsValid() &&
-			  InProperty->GetOwner().ToUObject()->GetName() == TEXT("LightmassWorldInfoSettings"))
+		if (InProperty->GetOwner<UObject>() &&
+			  InProperty->GetOwner<UObject>()->GetName() == TEXT("LightmassWorldInfoSettings"))
 		{
 			if (PropertyName == GET_MEMBER_NAME_STRING_CHECKED(FLightmassWorldInfoSettings, bGenerateAmbientOcclusionMaterialMask)
 				|| PropertyName == GET_MEMBER_NAME_STRING_CHECKED(FLightmassWorldInfoSettings, DirectIlluminationOcclusionFraction)
diff --git a/Engine/Source/Runtime/MovieSceneTracks/Private/Evaluation/MovieSceneEventTemplate.cpp b/Engine/Source/Runtime/MovieSceneTracks/Private/Evaluation/MovieSceneEventTemplate.cpp
--- a/Engine/Source/Runtime/MovieSceneTracks/Private/Evaluation/MovieSceneEventTemplate.cpp
+++ b/Engine/Source/Runtime/MovieSceneTracks/Private/Evaluation/MovieSceneEventTemplate.cpp
@@ -261,7 +261,7 @@ struct FEventTriggerExecutionToken
 
 	void TriggerEventWithParameters(UObject* DirectorInstance, const FMovieSceneEventPtrs& Event, TArrayView<UObject* const> EventContexts, IMovieScenePlayer& Player, FMovieSceneSequenceID SequenceID)
 	{
-		if (!ensureMsgf(!Event.BoundObjectProperty.Get() || (Event.BoundObjectProperty->GetOwner().ToUObject() == Event.Function && Event.BoundObjectProperty->GetOffset_ForUFunction() < Event.Function->ParmsSize), TEXT("Bound object property belongs to the wrong function or has an offset greater than the parameter size! This should never happen and indicates a BP compilation or nativization error.")))
+		if (!ensureMsgf(!Event.BoundObjectProperty.Get() || (Event.BoundObjectProperty->GetOwner<UObject>() == Event.Function && Event.BoundObjectProperty->GetOffset_ForUFunction() < Event.Function->ParmsSize), TEXT("Bound object property belongs to the wrong function or has an offset greater than the parameter size! This should never happen and indicates a BP compilation or nativization error.")))
 		{
 			return;
 		}
diff --git a/Engine/Source/Runtime/MovieSceneTracks/Private/Sections/MovieSceneEventSectionBase.cpp b/Engine/Source/Runtime/MovieSceneTracks/Private/Sections/MovieSceneEventSectionBase.cpp
--- a/Engine/Source/Runtime/MovieSceneTracks/Private/Sections/MovieSceneEventSectionBase.cpp
+++ b/Engine/Source/Runtime/MovieSceneTracks/Private/Sections/MovieSceneEventSectionBase.cpp
@@ -26,7 +26,7 @@ void UMovieSceneEventSectionBase::OnPostCompile(UBlueprint* Blueprint)
 				if (EntryPoint.Ptrs.Function && EntryPoint.BoundObjectPinName != NAME_None)
 				{
 					EntryPoint.Ptrs.BoundObjectProperty = EntryPoint.Ptrs.Function->FindPropertyByName(EntryPoint.BoundObjectPinName);
-					check(!EntryPoint.Ptrs.BoundObjectProperty.Get() || EntryPoint.Ptrs.BoundObjectProperty->GetOwner().ToUObject() == EntryPoint.Ptrs.Function);
+					check(!EntryPoint.Ptrs.BoundObjectProperty.Get() || EntryPoint.Ptrs.BoundObjectProperty->GetOwner<UObject>() == EntryPoint.Ptrs.Function);
 					if (CastField<FObjectProperty>(EntryPoint.Ptrs.BoundObjectProperty.Get()) || CastField<FInterfaceProperty>(EntryPoint.Ptrs.BoundObjectProperty.Get()))
 					{
 					}
diff --git a/Engine/Source/Runtime/Serialization/Private/Backends/CborStructSerializerBackend.cpp b/Engine/Source/Runtime/Serialization/Private/Backends/CborStructSerializerBackend.cpp
--- a/Engine/Source/Runtime/Serialization/Private/Backends/CborStructSerializerBackend.cpp
+++ b/Engine/Source/Runtime/Serialization/Private/Backends/CborStructSerializerBackend.cpp
@@ -21,7 +21,7 @@ FCborStructSerializerBackend::~FCborStructSerializerBackend() = default;
 void FCborStructSerializerBackend::BeginArray(const FStructSerializerState& State)
 {
 	// Array nested in Array/Set
-	if (State.ValueProperty->GetOwner().IsA(FArrayProperty::StaticClass()) || State.ValueProperty->GetOwner().IsA(FSetProperty::StaticClass()))
+	if (State.ValueProperty->GetOwner<FArrayProperty>() || State.ValueProperty->GetOwner<FSetProperty>())
 	{
 		CborWriter.WriteContainerStart(ECborCode::Array, -1);
 	}
@@ -46,7 +46,7 @@ void FCborStructSerializerBackend::BeginStructure(const FStructSerializerState&
 	if (State.ValueProperty != nullptr)
 	{
 		// Object nested in Array/Set
-		if (State.ValueProperty->GetOwner().IsA(FArrayProperty::StaticClass()) || State.ValueProperty->GetOwner().IsA(FSetProperty::StaticClass()))
+		if (State.ValueProperty->GetOwner<FArrayProperty>() || State.ValueProperty->GetOwner<FSetProperty>())
 		{
 			CborWriter.WriteContainerStart(ECborCode::Map, -1/*Indefinite*/);
 		}
@@ -96,7 +96,7 @@ namespace CborStructSerializerBackend
 		// Value nested in Array/Set or as root
 		if ((State.ValueProperty == nullptr) ||
 			(State.ValueProperty->ArrayDim > 1) ||
-			(State.ValueProperty->GetOwner().IsA(FArrayProperty::StaticClass()) || State.ValueProperty->GetOwner().IsA(FSetProperty::StaticClass())))
+			(State.ValueProperty->GetOwner<FArrayProperty>() || State.ValueProperty->GetOwner<FSetProperty>()))
 		{
 			CborWriter.WriteValue(Value);
 		}
@@ -121,7 +121,7 @@ namespace CborStructSerializerBackend
 	{
 		if ((State.ValueProperty == nullptr) ||
 			(State.ValueProperty->ArrayDim > 1) ||
-			(State.ValueProperty->GetOwner().IsA(FArrayProperty::StaticClass()) || State.ValueProperty->GetOwner().IsA(FSetProperty::StaticClass())))
+			(State.ValueProperty->GetOwner<FArrayProperty>() || State.ValueProperty->GetOwner<FSetProperty>()))
 		{
 			CborWriter.WriteNull();
 		}
diff --git a/Engine/Source/Runtime/Serialization/Private/Backends/JsonStructSerializerBackend.cpp b/Engine/Source/Runtime/Serialization/Private/Backends/JsonStructSerializerBackend.cpp
--- a/Engine/Source/Runtime/Serialization/Private/Backends/JsonStructSerializerBackend.cpp
+++ b/Engine/Source/Runtime/Serialization/Private/Backends/JsonStructSerializerBackend.cpp
@@ -18,8 +18,8 @@ namespace JsonStructSerializerBackend
 	{
 		if ((State.ValueProperty == nullptr) ||
 			(State.ValueProperty->ArrayDim > 1) ||
-			(State.ValueProperty->GetOwner().IsA(FArrayProperty::StaticClass())) ||
-			State.ValueProperty->GetOwner().IsA(FSetProperty::StaticClass()))
+			State.ValueProperty->GetOwner<FArrayProperty>() ||
+			State.ValueProperty->GetOwner<FSetProperty>())
 		{
 			JsonWriter->WriteValue(Value);
 		}
@@ -40,8 +40,8 @@ namespace JsonStructSerializerBackend
 	{
 		if ((State.ValueProperty == nullptr) ||
 			(State.ValueProperty->ArrayDim > 1) ||
-			(State.ValueProperty->GetOwner().IsA(FArrayProperty::StaticClass())) ||
-			(State.ValueProperty->GetOwner().IsA(FSetProperty::StaticClass())))
+			State.ValueProperty->GetOwner<FArrayProperty>() ||
+			State.ValueProperty->GetOwner<FSetProperty>())
 		{
 			JsonWriter->WriteNull();
 		}
@@ -64,9 +64,7 @@ namespace JsonStructSerializerBackend
 
 void FJsonStructSerializerBackend::BeginArray(const FStructSerializerState& State)
 {
-	FFieldVariant Outer = State.ValueProperty->GetOwner();
-
-	if (Outer.ToField() && (Outer.ToField()->GetClass() == FArrayProperty::StaticClass()))
+	if (State.ValueProperty->GetOwner<FArrayProperty>())
 	{
 		JsonWriter->WriteArrayStart();
 	}
@@ -87,9 +85,7 @@ void FJsonStructSerializerBackend::BeginStructure(const FStructSerializerState&
 {
 	if (State.ValueProperty != nullptr)
 	{
-		FFieldVariant Outer = State.ValueProperty->GetOwner();
-
-		if (CastField<FArrayProperty>(Outer.ToField()) || CastField<FSetProperty>(Outer.ToField()))
+		if (State.ValueProperty->GetOwner<FArrayProperty>() || State.ValueProperty->GetOwner<FSetProperty>())
 		{
 			JsonWriter->WriteObjectStart();
 		}
diff --git a/Engine/Source/Runtime/Serialization/Private/StructSerializer.cpp b/Engine/Source/Runtime/Serialization/Private/StructSerializer.cpp
--- a/Engine/Source/Runtime/Serialization/Private/StructSerializer.cpp
+++ b/Engine/Source/Runtime/Serialization/Private/StructSerializer.cpp
@@ -86,7 +86,7 @@ void FStructSerializer::Serialize(const void* Struct, UStruct& TypeInfo, IStruct
 				// write object start
 				if (CurrentState.ValueProperty != nullptr)
 				{
-					FFieldVariant Outer = CurrentState.ValueProperty->GetOwner();
+					FFieldVariant Outer = CurrentState.ValueProperty->GetOwnerVariant();
 
 					if ((Outer.ToField() == nullptr) || (Outer.ToField()->GetClass() != FArrayProperty::StaticClass()))
 					{
diff --git a/Engine/Source/Runtime/WebBrowser/Private/Native/NativeJSStructSerializerBackend.cpp b/Engine/Source/Runtime/WebBrowser/Private/Native/NativeJSStructSerializerBackend.cpp
--- a/Engine/Source/Runtime/WebBrowser/Private/Native/NativeJSStructSerializerBackend.cpp
+++ b/Engine/Source/Runtime/WebBrowser/Private/Native/NativeJSStructSerializerBackend.cpp
@@ -25,7 +25,7 @@ void FNativeJSStructSerializerBackend::WriteUObject(const FStructSerializerState
 {
 	// Note this function uses WriteRawJSONValue to append non-json data to the output stream.
 	FString RawValue = Scripting->ConvertObject(Value);
-	if ((State.ValueProperty == nullptr) || (State.ValueProperty->ArrayDim > 1) || (State.ValueProperty->GetOwner().IsA(FArrayProperty::StaticClass())))
+	if ((State.ValueProperty == nullptr) || (State.ValueProperty->ArrayDim > 1) || State.ValueProperty->GetOwner< FArrayProperty>())
 	{
 		GetWriter()->WriteRawJSONValue(RawValue);
 	}
