Test code removed.

#rnx

diff --git a/Engine/Source/Programs/UnrealHeaderTool/Private/CodeGenerator.cpp b/Engine/Source/Programs/UnrealHeaderTool/Private/CodeGenerator.cpp
--- a/Engine/Source/Programs/UnrealHeaderTool/Private/CodeGenerator.cpp
+++ b/Engine/Source/Programs/UnrealHeaderTool/Private/CodeGenerator.cpp
@@ -6523,29 +6523,6 @@ ECompilationResult::Type PreparseModules(const FString& ModuleInfoPath, int32& N
 	return Result;
 }
 
-struct FTestCollector
-{
-	void HandleObjectReference(UObject*& Obj)
-	{
-
-	}
-
-	template <typename T>
-	typename TEnableIf<TIsDerivedFrom<T, FField>::IsDerived, void>::Type AddReferencedObject(T*& Field)
-	{
-		if (Field)
-		{
-			UObject* Owner = Field->GetOwnerUObject();
-			UObject* OldOwnerPtr = nullptr;
-			HandleObjectReference(Owner);
-			if (Owner == nullptr && OldOwnerPtr != nullptr)
-			{
-				Field = nullptr;
-			}
-		}
-	}
-};
-
 ECompilationResult::Type UnrealHeaderTool_Main(const FString& ModuleInfoFilename)
 {
 	double MainTime = 0.0;
