More informative assert message when trying to export a header which has no class data.

#rnx

diff --git a/Engine/Source/Programs/UnrealHeaderTool/Private/CodeGenerator.cpp b/Engine/Source/Programs/UnrealHeaderTool/Private/CodeGenerator.cpp
--- a/Engine/Source/Programs/UnrealHeaderTool/Private/CodeGenerator.cpp
+++ b/Engine/Source/Programs/UnrealHeaderTool/Private/CodeGenerator.cpp
@@ -2970,7 +2970,7 @@ void FNativeClassHeaderGenerator::ExportClassFromSourceFileInner(
 	FUHTStringBuilder EnhancedUObjectConstructorsMacroCall;
 
 	FClassMetaData* ClassData = GScriptHelper.FindClassData(Class);
-	check(ClassData);
+	checkf(ClassData, TEXT("No class data generated for file %s"), *SourceFile.GetFilename());
 
 	// C++ -> VM stubs (native function execs)
 	FUHTStringBuilder ClassMacroCalls;
