Fix for UHT crashing when you have a UInterface but no IInterface.

diff --git a/Engine/Source/Programs/UnrealHeaderTool/Private/HeaderParser.cpp b/Engine/Source/Programs/UnrealHeaderTool/Private/HeaderParser.cpp
--- a/Engine/Source/Programs/UnrealHeaderTool/Private/HeaderParser.cpp
+++ b/Engine/Source/Programs/UnrealHeaderTool/Private/HeaderParser.cpp
@@ -5575,12 +5575,21 @@ bool FHeaderParser::CompileDeclaration(FClasses& AllClasses, TArray<UDelegateFun
 		RequireSymbol(TEXT(')'), Token.Identifier);
 
 		FClassMetaData* ClassData = GetCurrentClassData();
-		if (!ClassData)
+		if (!ClassData || ClassData->ParsedInterface == EParsedInterface::NotAnInterface)
 		{
 			FString CurrentClassName = GetCurrentClass()->GetName();
 			FError::Throwf(TEXT("Could not find the associated 'U%s' class while parsing 'I%s' - it could be missing or malformed"), *CurrentClassName, *CurrentClassName);
 		}
 
+		if (ClassData->ParsedInterface == EParsedInterface::ParsedIInterface)
+		{
+			FString CurrentClassName = GetCurrentClass()->GetName();
+			FError::Throwf(TEXT("Duplicate IInterface definition found while parsing 'I%s'"), *CurrentClassName);
+		}
+
+		check(ClassData->ParsedInterface == EParsedInterface::ParsedUInterface);
+
+		ClassData->ParsedInterface = EParsedInterface::ParsedIInterface;
 		ClassData->GeneratedBodyMacroAccessSpecifier = CurrentAccessSpecifier;
 		ClassData->SetInterfaceGeneratedBodyLine(InputLine);
 
@@ -6550,6 +6559,7 @@ void FHeaderParser::CompileInterfaceDeclaration(FClasses& AllClasses)
 	FClassMetaData* ClassData = GScriptHelper.AddClassData(InterfaceClass, CurrentSrcFile);
 	check(ClassData);
 
+	ClassData->ParsedInterface = EParsedInterface::ParsedUInterface;
 	ClassData->SetPrologLine(PrologFinishLine);
 
 	// Register the metadata
@@ -8458,6 +8468,8 @@ ECompilationResult::Type FHeaderParser::ParseHeader(FClasses& AllClasses, FUnrea
 			const FString ExpectedHeaderName = CurrentSrcFile->GetGeneratedHeaderFilename();
 			FError::Throwf(TEXT("Expected an include at the top of the header: '#include \"%s\"'"), *ExpectedHeaderName);
 		}
+
+		GScriptHelper.CheckForNoIInterfaces();
 	}
 #if !PLATFORM_EXCEPTIONS_DISABLED
 	catch( TCHAR* ErrorMsg )
diff --git a/Engine/Source/Programs/UnrealHeaderTool/Private/ParserHelper.cpp b/Engine/Source/Programs/UnrealHeaderTool/Private/ParserHelper.cpp
--- a/Engine/Source/Programs/UnrealHeaderTool/Private/ParserHelper.cpp
+++ b/Engine/Source/Programs/UnrealHeaderTool/Private/ParserHelper.cpp
@@ -3,6 +3,7 @@
 
 #include "ParserHelper.h"
 #include "UnrealHeaderTool.h"
+#include "Algo/Find.h"
 #include "Misc/DefaultValueHelper.h"
 
 /////////////////////////////////////////////////////
@@ -282,4 +283,18 @@ FTokenData* FPropertyData::Set(FProperty* InKey, const FTokenData& InValue, FUnr
 	}
 
 	return Result;
+}
+
+void FCompilerMetadataManager::CheckForNoIInterfaces() const
+{
+	const TPair<UStruct*, TUniquePtr<FClassMetaData>>* UnparsedIInterface = Algo::FindBy(
+		*(Super*)this,
+		EParsedInterface::ParsedUInterface,
+		[](const TPair<UStruct*, TUniquePtr<FClassMetaData>>& Kvp) { return Kvp.Value->ParsedInterface; }
+	);
+	if (UnparsedIInterface)
+	{
+		FString Name = UnparsedIInterface->Key->GetName();
+		FError::Throwf(TEXT("UInterface 'U%s' parsed without a corresponding 'I%s'"), *Name, *Name);
+	}
 }
\ No newline at end of file
diff --git a/Engine/Source/Programs/UnrealHeaderTool/Private/ParserHelper.h b/Engine/Source/Programs/UnrealHeaderTool/Private/ParserHelper.h
--- a/Engine/Source/Programs/UnrealHeaderTool/Private/ParserHelper.h
+++ b/Engine/Source/Programs/UnrealHeaderTool/Private/ParserHelper.h
@@ -1511,6 +1511,13 @@ struct FMultipleInheritanceBaseClass
 	}
 };
 
+enum class EParsedInterface
+{
+	NotAnInterface,
+	ParsedUInterface,
+	ParsedIInterface
+};
+
 /**
  * Class for storing compiler metadata about a class's properties.
  */
@@ -1800,6 +1807,9 @@ public:
 	// GENERATED_BODY access specifier to preserve.
 	EAccessSpecifier GeneratedBodyMacroAccessSpecifier;
 
+	/** Parsed interface state */
+	EParsedInterface ParsedInterface = EParsedInterface::NotAnInterface;
+
 	friend struct FClassMetaDataArchiveProxy;
 };
 
@@ -1810,6 +1820,8 @@ public:
  */
 class FCompilerMetadataManager : protected TMap<UStruct*, TUniquePtr<FClassMetaData> >
 {
+	using Super = TMap<UStruct*, TUniquePtr<FClassMetaData>>;
+
 public:
 	/**
 	 * Adds a new class to be tracked
@@ -1853,6 +1865,11 @@ public:
 		}
 	}
 
+	/**
+	 * Throws an exception if a UInterface was parsed but not the corresponding IInterface.
+	 */
+	void CheckForNoIInterfaces() const;
+
 	friend struct FCompilerMetadataManagerArchiveProxy;
 };
 
