Update UHT to support int64 properties with ExposeOnSpawn

diff --git a/Engine/Source/Programs/UnrealHeaderTool/Private/HeaderParser.cpp b/Engine/Source/Programs/UnrealHeaderTool/Private/HeaderParser.cpp
--- a/Engine/Source/Programs/UnrealHeaderTool/Private/HeaderParser.cpp
+++ b/Engine/Source/Programs/UnrealHeaderTool/Private/HeaderParser.cpp
@@ -7903,6 +7903,7 @@ struct FExposeOnSpawnValidator
 		switch (Property.Type)
 		{
 		case CPT_Int:
+		case CPT_Int64:
 		case CPT_Byte:
 		case CPT_Float:
 		case CPT_Bool:
