- Hardened UHT against an order of module dependencies issue that can happen with CoreUObject depending on other modules

diff --git a/Engine/Source/Programs/UnrealHeaderTool/Private/CodeGenerator.cpp b/Engine/Source/Programs/UnrealHeaderTool/Private/CodeGenerator.cpp
--- a/Engine/Source/Programs/UnrealHeaderTool/Private/CodeGenerator.cpp
+++ b/Engine/Source/Programs/UnrealHeaderTool/Private/CodeGenerator.cpp
@@ -6346,7 +6346,7 @@ void ResolveSuperClasses(UPackage* Package)
 
 	for (UObject* Object : Objects)
 	{
-		if (!Object->IsA<UClass>())
+		if (!Object->IsA<UClass>() || Object->HasAnyFlags(RF_ClassDefaultObject))
 		{
 			continue;
 		}
