Improved validation for per class sparse data in UHT. Fixed a bug where improperly specified sparse data structures were reported multiple times and with the wrong file location.

diff --git a/Engine/Source/Programs/UnrealHeaderTool/Private/HeaderParser.cpp b/Engine/Source/Programs/UnrealHeaderTool/Private/HeaderParser.cpp
--- a/Engine/Source/Programs/UnrealHeaderTool/Private/HeaderParser.cpp
+++ b/Engine/Source/Programs/UnrealHeaderTool/Private/HeaderParser.cpp
@@ -2779,6 +2779,9 @@ UScriptStruct* FHeaderParser::CompileStructDeclaration(FClasses& AllClasses)
 		FError::Throwf(TEXT("Expected a GENERATED_BODY() at the start of struct"));
 	}
 
+	// Validate sparse class data
+	CheckSparseClassData(Struct);
+
 	// Link the properties within the struct
 	Struct->StaticLink(true);
 
@@ -3159,9 +3162,8 @@ void FHeaderParser::CheckSparseClassData(const UStruct* StructToCheck)
 		// make sure we don't try to have sparse class data inside of a struct instead of a class
 		if (StructToCheck->HasMetaData(FHeaderParserNames::NAME_SparseClassDataTypes))
 		{
-			UE_LOG_ERROR_UHT(TEXT("%s contains sparse class data but is not a class."), *StructToCheck->GetName());
+			FError::Throwf(TEXT("%s contains sparse class data but is not a class."), *StructToCheck->GetName());
 		}
-		return;
 	}
 
 	if (!ClassToCheck->HasMetaData(FHeaderParserNames::NAME_SparseClassDataTypes))
@@ -3175,13 +3177,11 @@ void FHeaderParser::CheckSparseClassData(const UStruct* StructToCheck)
 	// for now we only support one sparse class data structure per class
 	if (SparseClassDataTypes.Num() > 1)
 	{
-		UE_LOG_ERROR_UHT(TEXT("Class %s contains multiple sparse class data types."), *ClassToCheck->GetName());
-		return;
+		FError::Throwf(TEXT("Class %s contains multiple sparse class data types."), *ClassToCheck->GetName());
 	}
 	if (SparseClassDataTypes.Num() == 0)
 	{
-		UE_LOG_ERROR_UHT(TEXT("Class %s has sparse class metadata but does not specify a type."), *ClassToCheck->GetName());
-		return;
+		FError::Throwf(TEXT("Class %s has sparse class metadata but does not specify a type."), *ClassToCheck->GetName());
 	}
 
 	for (const FString& SparseClassDataTypeName : SparseClassDataTypes)
@@ -3191,8 +3191,7 @@ void FHeaderParser::CheckSparseClassData(const UStruct* StructToCheck)
 		// make sure the sparse class data struct actually exists
 		if (!SparseClassDataStruct)
 		{
-			UE_LOG_ERROR_UHT(TEXT("Unable to find sparse data type %s for class %s."), *SparseClassDataTypeName, *ClassToCheck->GetName());
-			return;
+			FError::Throwf(TEXT("Unable to find sparse data type %s for class %s."), *SparseClassDataTypeName, *ClassToCheck->GetName());
 		}
 
 		// check the data struct for invalid properties
@@ -3200,19 +3199,19 @@ void FHeaderParser::CheckSparseClassData(const UStruct* StructToCheck)
 		{
 			if (Property->HasAnyPropertyFlags(CPF_BlueprintAssignable))
 			{
-				UE_LOG_ERROR_UHT(TEXT("Sparse class data types can not contain blueprint assignable delegates. Type '%s' Delegate '%s'"), *SparseClassDataStruct->GetName(), *Property->GetName());
+				FError::Throwf(TEXT("Sparse class data types can not contain blueprint assignable delegates. Type '%s' Delegate '%s'"), *SparseClassDataStruct->GetName(), *Property->GetName());
 			}
 
 			// all sparse properties should have EditDefaultsOnly
 			if (!Property->HasAllPropertyFlags(CPF_Edit | CPF_DisableEditOnInstance))
 			{
-				UE_LOG_ERROR_UHT(TEXT("Sparse class data types must be EditDefaultsOnly. Type '%s' Property '%s'"), *SparseClassDataStruct->GetName(), *Property->GetName());
+				FError::Throwf(TEXT("Sparse class data types must be EditDefaultsOnly. Type '%s' Property '%s'"), *SparseClassDataStruct->GetName(), *Property->GetName());
 			}
 
 			// no sparse properties should have BlueprintReadWrite
 			if (Property->HasAllPropertyFlags(CPF_BlueprintVisible) && !Property->HasAllPropertyFlags(CPF_BlueprintReadOnly))
 			{
-				UE_LOG_ERROR_UHT(TEXT("Sparse class data types must not be BlueprintReadWrite. Type '%s' Property '%s'"), *SparseClassDataStruct->GetName(), *Property->GetName());
+				FError::Throwf(TEXT("Sparse class data types must not be BlueprintReadWrite. Type '%s' Property '%s'"), *SparseClassDataStruct->GetName(), *Property->GetName());
 			}
 		}
 
@@ -3225,7 +3224,7 @@ void FHeaderParser::CheckSparseClassData(const UStruct* StructToCheck)
 			UScriptStruct* ParentSparseClassDataStruct = FindObjectSafe<UScriptStruct>(ANY_PACKAGE, *ParentSparseClassDataTypeName);
 			if (ParentSparseClassDataStruct && !SparseClassDataStruct->IsChildOf(ParentSparseClassDataStruct))
 			{
-				UE_LOG_ERROR_UHT(TEXT("Class %s is a child of %s but its sparse class data struct, %s, does not inherit from %s."), *ClassToCheck->GetName(), *ParentClass->GetName(), *SparseClassDataStruct->GetName(), *ParentSparseClassDataStruct->GetName());
+				FError::Throwf(TEXT("Class %s is a child of %s but its sparse class data struct, %s, does not inherit from %s."), *ClassToCheck->GetName(), *ParentClass->GetName(), *SparseClassDataStruct->GetName(), *ParentSparseClassDataStruct->GetName());
 			}
 		}
 	}
@@ -6385,6 +6384,9 @@ UClass* FHeaderParser::CompileClassDeclaration(FClasses& AllClasses)
 		}
 	}
 
+	// Validate sparse class data
+	CheckSparseClassData(Class);
+
 	return Class;
 }
 
@@ -8877,12 +8879,6 @@ ECompilationResult::Type FHeaderParser::ParseAllHeadersInside(
 
 		if (Result == ECompilationResult::Succeeded)
 		{
-			// Validate the sparse class data for all classes in the current package
-			for (const FClass* Class : ModuleClasses.GetClassesInPackage(CurrentPackage))
-			{
-				CheckSparseClassData(Class);
-			}
-
 			// Export the autogenerated code wrappers
 
 			// At this point all headers have been parsed and the header parser will
