Fixed missing ';' in HeaderParser.cpp

#rnx

diff --git a/Engine/Source/Programs/UnrealHeaderTool/Private/HeaderParser.cpp b/Engine/Source/Programs/UnrealHeaderTool/Private/HeaderParser.cpp
--- a/Engine/Source/Programs/UnrealHeaderTool/Private/HeaderParser.cpp
+++ b/Engine/Source/Programs/UnrealHeaderTool/Private/HeaderParser.cpp
@@ -3184,7 +3184,7 @@ void FHeaderParser::CheckSparseClassData(const UStruct* StructToCheck)
 	if (SparseClassDataTypes.Num() == 0)
 	{
 		FError::Throwf(TEXT("Class %s has sparse class metadata but does not specify a type."), *ClassToCheck->GetName());
-		return
+		return;
 	}
 
 	for (const FString& SparseClassDataTypeName : SparseClassDataTypes)
