Improve header generation time by ~4 seconds
#rnx

diff --git a/Engine/Source/Programs/UnrealHeaderTool/Private/HeaderParser.cpp b/Engine/Source/Programs/UnrealHeaderTool/Private/HeaderParser.cpp
--- a/Engine/Source/Programs/UnrealHeaderTool/Private/HeaderParser.cpp
+++ b/Engine/Source/Programs/UnrealHeaderTool/Private/HeaderParser.cpp
@@ -6563,10 +6563,9 @@ void FHeaderParser::CompileInterfaceDeclaration(FClasses& AllClasses)
 	}
 
 	// Try parsing metadata for the interface
-	FClassMetaData* ClassData = GScriptHelper.AddClassData(InterfaceClass, CurrentSrcFile);
+	FClassMetaData* ClassData = GScriptHelper.AddInterfaceClassData(InterfaceClass, CurrentSrcFile);
 	check(ClassData);
 
-	ClassData->ParsedInterface = EParsedInterface::ParsedUInterface;
 	ClassData->SetPrologLine(PrologFinishLine);
 
 	// Register the metadata
diff --git a/Engine/Source/Programs/UnrealHeaderTool/Private/ParserHelper.cpp b/Engine/Source/Programs/UnrealHeaderTool/Private/ParserHelper.cpp
--- a/Engine/Source/Programs/UnrealHeaderTool/Private/ParserHelper.cpp
+++ b/Engine/Source/Programs/UnrealHeaderTool/Private/ParserHelper.cpp
@@ -266,6 +266,14 @@ FClassMetaData* FCompilerMetadataManager::AddClassData(UStruct* Struct, FUnrealS
 	return pClassData->Get();
 }
 
+FClassMetaData* FCompilerMetadataManager::AddInterfaceClassData(UStruct* Struct, FUnrealSourceFile* UnrealSourceFile)
+{
+	FClassMetaData* ClassData = AddClassData(Struct, UnrealSourceFile);
+	ClassData->ParsedInterface = EParsedInterface::ParsedUInterface;
+	InterfacesToVerify.Emplace(Struct, ClassData);
+	return ClassData;
+}
+
 FTokenData* FPropertyData::Set(FProperty* InKey, const FTokenData& InValue, FUnrealSourceFile* UnrealSourceFile)
 {
 	FTokenData* Result = NULL;
@@ -285,16 +293,15 @@ FTokenData* FPropertyData::Set(FProperty* InKey, const FTokenData& InValue, FUnr
 	return Result;
 }
 
-void FCompilerMetadataManager::CheckForNoIInterfaces() const
+void FCompilerMetadataManager::CheckForNoIInterfaces()
 {
-	const TPair<UStruct*, TUniquePtr<FClassMetaData>>* UnparsedIInterface = Algo::FindBy(
-		*(Super*)this,
-		EParsedInterface::ParsedUInterface,
-		[](const TPair<UStruct*, TUniquePtr<FClassMetaData>>& Kvp) { return Kvp.Value->ParsedInterface; }
-	);
-	if (UnparsedIInterface)
+	for (const TPair<UStruct*, FClassMetaData*>& StructDataPair : InterfacesToVerify)
 	{
-		FString Name = UnparsedIInterface->Key->GetName();
-		FError::Throwf(TEXT("UInterface 'U%s' parsed without a corresponding 'I%s'"), *Name, *Name);
+		if (StructDataPair.Value->ParsedInterface == EParsedInterface::ParsedUInterface)
+		{
+			FString Name = StructDataPair.Key->GetName();
+			FError::Throwf(TEXT("UInterface 'U%s' parsed without a corresponding 'I%s'"), *Name, *Name);
+		}
 	}
+	InterfacesToVerify.Reset();
 }
\ No newline at end of file
diff --git a/Engine/Source/Programs/UnrealHeaderTool/Private/ParserHelper.h b/Engine/Source/Programs/UnrealHeaderTool/Private/ParserHelper.h
--- a/Engine/Source/Programs/UnrealHeaderTool/Private/ParserHelper.h
+++ b/Engine/Source/Programs/UnrealHeaderTool/Private/ParserHelper.h
@@ -1832,6 +1832,8 @@ public:
 	 */
 	FClassMetaData* AddClassData(UStruct* Struct, FUnrealSourceFile* UnrealSourceFile);
 
+	FClassMetaData* AddInterfaceClassData(UStruct* Struct, FUnrealSourceFile* UnrealSourceFile);
+
 	/**
 	 * Find the metadata associated with the class specified
 	 * 
@@ -1868,9 +1870,12 @@ public:
 	/**
 	 * Throws an exception if a UInterface was parsed but not the corresponding IInterface.
 	 */
-	void CheckForNoIInterfaces() const;
+	void CheckForNoIInterfaces();
 
 	friend struct FCompilerMetadataManagerArchiveProxy;
+
+private:
+	TArray<TPair<UStruct*, FClassMetaData*>> InterfacesToVerify;
 };
 
 /*-----------------------------------------------------------------------------
