Move headers in to place and delete stale/unused generated headers asynchronously while moving on to parsing next package
#rnx

diff --git a/Engine/Source/Programs/UnrealHeaderTool/Private/CodeGenerator.cpp b/Engine/Source/Programs/UnrealHeaderTool/Private/CodeGenerator.cpp
--- a/Engine/Source/Programs/UnrealHeaderTool/Private/CodeGenerator.cpp
+++ b/Engine/Source/Programs/UnrealHeaderTool/Private/CodeGenerator.cpp
@@ -53,6 +53,7 @@
 #include "Algo/Copy.h"
 #include "Algo/Sort.h"
 #include "Algo/Reverse.h"
+#include "Async/Async.h"
 #include "Async/ParallelFor.h"
 #include "Misc/ScopeExit.h"
 #include "UnrealTypeDefinitionInfo.h"
@@ -79,6 +80,9 @@ static TSharedRef<FUnrealSourceFile> PerformInitialParseOnHeader(UPackage* InPar
 
 FCompilerMetadataManager GScriptHelper;
 
+// Array of all the temporary header async file tasks so we can ensure they have completed before issuing our timings
+static FGraphEventArray GAsyncFileTasks;
+
 bool HasIdentifierExactMatch(const TCHAR* StringBegin, const TCHAR* StringEnd, const FString& Find);
 
 namespace
@@ -5643,7 +5647,7 @@ FNativeClassHeaderGenerator::FNativeClassHeaderGenerator(
 	, bAllowSaveExportedHeaders  (InAllowSaveExportedHeaders)
 	, bFailIfGeneratedCodeChanges(FParse::Param(FCommandLine::Get(), TEXT("FailIfGeneratedCodeChanges")))
 {
-	const FString PackageName = FPackageName::GetShortName(Package);
+	FString PackageName = FPackageName::GetShortName(Package);
 
 	FManifestModule* PackageManifest = GetPackageManifest(PackageName);
 	if (!PackageManifest)
@@ -6030,50 +6034,55 @@ FNativeClassHeaderGenerator::FNativeClassHeaderGenerator(
 	}
 
 	// Export all changed headers from their temp files to the .h files
-	ExportUpdatedHeaders(PackageName, TempHeaderPaths);
+	ExportUpdatedHeaders(MoveTemp(PackageName), MoveTemp(TempHeaderPaths));
 
 	// Delete stale *.generated.h files
-	DeleteUnusedGeneratedHeaders(PackageHeaderPaths);
+	DeleteUnusedGeneratedHeaders(MoveTemp(PackageHeaderPaths));
 }
 
-void FNativeClassHeaderGenerator::DeleteUnusedGeneratedHeaders(const TSet<FString>& PackageHeaderPathSet)
+void FNativeClassHeaderGenerator::DeleteUnusedGeneratedHeaders(TSet<FString>&& PackageHeaderPathSet)
 {
-	TSet<FString> AllIntermediateFolders;
-
-	for (const FString& PackageHeader : PackageHeaderPathSet)
+	auto DeleteUnusedGeneratedHeadersTask = [PackageHeaderPathSet = MoveTemp(PackageHeaderPathSet)]()
 	{
-		const FString IntermediatePath = FPaths::GetPath(PackageHeader);
+		TSet<FString> AllIntermediateFolders;
 
-		if (AllIntermediateFolders.Contains(IntermediatePath))
+		for (const FString& PackageHeader : PackageHeaderPathSet)
 		{
-			continue;
-		}
-
-		AllIntermediateFolders.Add( IntermediatePath );
-
-		TArray<FString> AllHeaders;
-		IFileManager::Get().FindFiles( AllHeaders, *(IntermediatePath / TEXT("*.generated.h")), true, false );
+			FString IntermediatePath = FPaths::GetPath(PackageHeader);
 
-		for (const FString& Header : AllHeaders)
-		{
-			const FString HeaderPath = IntermediatePath / Header;
-
-			if (PackageHeaderPathSet.Contains(HeaderPath))
+			if (AllIntermediateFolders.Contains(IntermediatePath))
 			{
 				continue;
 			}
 
-			// Check intrinsic classes. Get the class name from file name by removing .generated.h.
-			const FString HeaderFilename = FPaths::GetBaseFilename(HeaderPath);
-			const int32   GeneratedIndex = HeaderFilename.Find(TEXT(".generated"), ESearchCase::IgnoreCase, ESearchDir::FromEnd);
-			const FString ClassName      = HeaderFilename.Mid(0, GeneratedIndex);
-			UClass* IntrinsicClass       = FindObject<UClass>(ANY_PACKAGE, *ClassName);
-			if (!IntrinsicClass || !IntrinsicClass->HasAnyClassFlags(CLASS_Intrinsic))
+			TArray<FString> AllHeaders;
+			IFileManager::Get().FindFiles(AllHeaders, *(IntermediatePath / TEXT("*.generated.h")), true, false);
+
+			for (const FString& Header : AllHeaders)
 			{
-				IFileManager::Get().Delete(*HeaderPath);
+				const FString HeaderPath = IntermediatePath / Header;
+
+				if (PackageHeaderPathSet.Contains(HeaderPath))
+				{
+					continue;
+				}
+
+				// Check intrinsic classes. Get the class name from file name by removing .generated.h.
+				FString HeaderFilename = FPaths::GetBaseFilename(HeaderPath);
+				const int32   GeneratedIndex = HeaderFilename.Find(TEXT(".generated"), ESearchCase::IgnoreCase, ESearchDir::FromEnd);
+				const FString ClassName = MoveTemp(HeaderFilename).Mid(0, GeneratedIndex);
+				UClass* IntrinsicClass = FindObject<UClass>(ANY_PACKAGE, *ClassName);
+				if (!IntrinsicClass || !IntrinsicClass->HasAnyClassFlags(CLASS_Intrinsic))
+				{
+					IFileManager::Get().Delete(*HeaderPath);
+				}
 			}
+
+			AllIntermediateFolders.Add(MoveTemp(IntermediatePath));
 		}
-	}
+	};
+
+	GAsyncFileTasks.Add(FFunctionGraphTask::CreateAndDispatchWhenReady(MoveTemp(DeleteUnusedGeneratedHeadersTask), TStatId()));
 }
 
 /**
@@ -6219,41 +6228,37 @@ bool FNativeClassHeaderGenerator::SaveHeaderIfChanged(FReferenceGatherers& OutRe
 	return bHasChanged;
 }
 
-/**
-* Create a temp header file name from the header name
-*
-* @param	CurrentFilename		The filename off of which the current filename will be generated
-* @param	bReverseOperation	Get the header from the temp file name instead
-*
-* @return	The generated string
-*/
 FString FNativeClassHeaderGenerator::GenerateTempHeaderName( const FString& CurrentFilename, bool bReverseOperation )
 {
 	return bReverseOperation
-		? CurrentFilename.Replace(TEXT(".tmp"), TEXT(""))
+		? CurrentFilename.Replace(TEXT(".tmp"), TEXT(""), ESearchCase::CaseSensitive)
 		: CurrentFilename + TEXT(".tmp");
 }
 
-/**
-* Exports the temp header files into the .h files, then deletes the temp files.
-*
-* @param	PackageName	Name of the package being saved
-*/
-void FNativeClassHeaderGenerator::ExportUpdatedHeaders(const FString& PackageName, const TArray<FString>& TempHeaderPaths)
-{
-	ParallelFor(TempHeaderPaths.Num(), [&](int32 Index)
+void FNativeClassHeaderGenerator::ExportUpdatedHeaders(FString&& PackageName, TArray<FString>&& TempHeaderPaths)
+{	
+	// Asynchronously move the headers to the correct locations
+	if (TempHeaderPaths.Num() > 0)
 	{
-		const FString& TmpFilename = TempHeaderPaths[Index];
-		FString Filename = GenerateTempHeaderName( TmpFilename, true );
-		if (!IFileManager::Get().Move(*Filename, *TmpFilename, true, true))
+		auto MoveHeadersTask = [PackageName = MoveTemp(PackageName), TempHeaderPaths = MoveTemp(TempHeaderPaths)]()
 		{
-			UE_LOG(LogCompile, Error, TEXT("Error exporting %s: couldn't write file '%s'"), *PackageName, *Filename);
-		}
-		else
-		{
-			UE_LOG(LogCompile, Log, TEXT("Exported updated C++ header: %s"), *Filename);
-		}
-	});
+			ParallelFor(TempHeaderPaths.Num(), [&](int32 Index)
+			{
+				const FString& TmpFilename = TempHeaderPaths[Index];
+				FString Filename = GenerateTempHeaderName(TmpFilename, true);
+				if (!IFileManager::Get().Move(*Filename, *TmpFilename, true, true))
+				{
+					UE_LOG(LogCompile, Error, TEXT("Error exporting %s: couldn't write file '%s'"), *PackageName, *Filename);
+				}
+				else
+				{
+					UE_LOG(LogCompile, Log, TEXT("Exported updated C++ header: %s"), *Filename);
+				}
+			});
+		};
+
+		GAsyncFileTasks.Add(FFunctionGraphTask::CreateAndDispatchWhenReady(MoveTemp(MoveHeadersTask), TStatId()));
+	}
 }
 
 /**
@@ -6730,6 +6735,9 @@ ECompilationResult::Type UnrealHeaderTool_Main(const FString& ModuleInfoFilename
 	// Avoid TArray slack for meta data.
 	GScriptHelper.Shrink();
 
+	// Finish all async file tasks before stopping the clock
+	FTaskGraphInterface::Get().WaitUntilTasksComplete(GAsyncFileTasks);
+
 	MainTimer.Stop();
 
 	UE_LOG(LogCompile, Log, TEXT("Preparsing %i modules took %.2f seconds"), GManifest.Modules.Num(), TotalModulePreparseTime);
diff --git a/Engine/Source/Programs/UnrealHeaderTool/Private/NativeClassExporter.h b/Engine/Source/Programs/UnrealHeaderTool/Private/NativeClassExporter.h
--- a/Engine/Source/Programs/UnrealHeaderTool/Private/NativeClassExporter.h
+++ b/Engine/Source/Programs/UnrealHeaderTool/Private/NativeClassExporter.h
@@ -293,12 +293,13 @@ private:
 	 */
 	static void ExportEventParm(FUHTStringBuilder& Out, TSet<FString>& PropertyFwd, UFunction* Function, int32 Indent, bool bOutputConstructor, EExportingState ExportingState);
 
-	/** 
-	* Exports the temp header files into the .h files, then deletes the temp files.
+	/**
+	* Move the temp header files into the .h files
 	* 
 	* @param	PackageName	Name of the package being saved
+	* @param	TempHeaderPaths	Names of all the headers to move
 	*/
-	static void ExportUpdatedHeaders(const FString& PackageName, const TArray<FString>& TempHeaderPaths);
+	static void ExportUpdatedHeaders(FString&& PackageName, TArray<FString>&& TempHeaderPaths);
 
 	/**
 	 * Exports the generated cpp file for all functions/events/delegates in package.
@@ -483,7 +484,7 @@ private:
 	/**
 	 * Deletes all .generated.h files which do not correspond to any of the classes.
 	 */
-	static void DeleteUnusedGeneratedHeaders(const TSet<FString>& PackageHeaderPathSet);
+	static void DeleteUnusedGeneratedHeaders(TSet<FString>&& PackageHeaderPathSet);
 
 	/**
 	 * Exports macros that manages UObject constructors.
