UHT optimization: Remove autos and reorder == to do cheapest checks first
#rnx

diff --git a/Engine/Source/Programs/UnrealHeaderTool/Private/Manifest.cpp b/Engine/Source/Programs/UnrealHeaderTool/Private/Manifest.cpp
--- a/Engine/Source/Programs/UnrealHeaderTool/Private/Manifest.cpp
+++ b/Engine/Source/Programs/UnrealHeaderTool/Private/Manifest.cpp
@@ -36,7 +36,7 @@ namespace
 	template <typename T>
 	void GetJsonFieldValue(T& OutVal, const TSharedPtr<FJsonObject>& JsonObject, const TCHAR* FieldName, const TCHAR* Outer)
 	{
-		auto JsonValue = JsonObject->Values.Find(FieldName);
+		TSharedPtr<FJsonValue>* JsonValue = JsonObject->Values.Find(FieldName);
 
 		if (!JsonValue)
 		{
@@ -72,8 +72,8 @@ FManifest FManifest::LoadFromFile(const FString& Filename)
 		FError::Throwf(TEXT("Unable to load manifest: %s"), *Filename);
 	}
 
-	auto RootObject = TSharedPtr<FJsonObject>();
-	auto Reader = TJsonReaderFactory<TCHAR>::Create(Json);
+	TSharedPtr<FJsonObject> RootObject = TSharedPtr<FJsonObject>();
+	TSharedRef<TJsonReader<TCHAR>> Reader = TJsonReaderFactory<TCHAR>::Create(Json);
 
 	if (!FJsonSerializer::Deserialize(Reader, RootObject))
 	{
@@ -112,16 +112,16 @@ FManifest FManifest::LoadFromFile(const FString& Filename)
 
 	int32 ModuleIndex = 0;
 
-	for (auto Module : ModulesArray)
+	for (const TSharedPtr<FJsonValue>& Module : ModulesArray)
 	{
-		auto ModuleObj = Module->AsObject();
+		const TSharedPtr<FJsonObject>& ModuleObj = Module->AsObject();
 
 		TArray<TSharedPtr<FJsonValue>> ClassesHeaders;
 		TArray<TSharedPtr<FJsonValue>> PublicHeaders;
 		TArray<TSharedPtr<FJsonValue>> PrivateHeaders;
 
 		Result.Modules.AddZeroed();
-		auto& KnownModule = Result.Modules.Last();
+		FManifestModule& KnownModule = Result.Modules.Last();
 
 		FString Outer = FString::Printf(TEXT("Modules[%d]"), ModuleIndex);
 		FString GeneratedCodeVersionString;
@@ -236,8 +236,10 @@ bool FManifestModule::NeedsRegeneration() const
 
 bool FManifestModule::IsCompatibleWith(const FManifestModule& ManifestModule)
 {
-	return Name == ManifestModule.Name
-		&& ModuleType == ManifestModule.ModuleType
+	return ModuleType == ManifestModule.ModuleType
+		&& SaveExportedHeaders == ManifestModule.SaveExportedHeaders
+		&& GeneratedCodeVersion == ManifestModule.GeneratedCodeVersion
+		&& Name == ManifestModule.Name
 		&& LongPackageName == ManifestModule.LongPackageName
 		&& BaseDirectory == ManifestModule.BaseDirectory
 		&& IncludeBase == ManifestModule.IncludeBase
@@ -245,7 +247,5 @@ bool FManifestModule::IsCompatibleWith(const FManifestModule& ManifestModule)
 		&& PublicUObjectClassesHeaders == ManifestModule.PublicUObjectClassesHeaders
 		&& PublicUObjectHeaders == ManifestModule.PublicUObjectHeaders
 		&& PrivateUObjectHeaders == ManifestModule.PrivateUObjectHeaders
-		&& GeneratedCPPFilenameBase == ManifestModule.GeneratedCPPFilenameBase
-		&& SaveExportedHeaders == ManifestModule.SaveExportedHeaders
-		&& GeneratedCodeVersion == ManifestModule.GeneratedCodeVersion;
+		&& GeneratedCPPFilenameBase == ManifestModule.GeneratedCPPFilenameBase;
 }
