UHT cleanup: Fix constructor signature to be move explicit
#rnx

diff --git a/Engine/Source/Programs/UnrealHeaderTool/Private/SimplifiedParsingClassInfo.h b/Engine/Source/Programs/UnrealHeaderTool/Private/SimplifiedParsingClassInfo.h
--- a/Engine/Source/Programs/UnrealHeaderTool/Private/SimplifiedParsingClassInfo.h
+++ b/Engine/Source/Programs/UnrealHeaderTool/Private/SimplifiedParsingClassInfo.h
@@ -14,7 +14,7 @@ class FSimplifiedParsingClassInfo
 {
 public:
 	// Constructor.
-	FSimplifiedParsingClassInfo(FString InClassName, FString InBaseClassName, int32 InClassDefLine, bool bInClassIsAnInterface)
+	FSimplifiedParsingClassInfo(FString&& InClassName, FString&& InBaseClassName, int32 InClassDefLine, bool bInClassIsAnInterface)
 		: ClassName          (MoveTemp(InClassName))
 		, BaseClassName      (MoveTemp(InBaseClassName))
 		, ClassDefLine       (InClassDefLine)
