UHT Optimization: Determine and store in a static bool FailIfGeneratedCodeChanges rather than looking up per module.
#rnx

diff --git a/Engine/Source/Programs/UnrealHeaderTool/Private/CodeGenerator.cpp b/Engine/Source/Programs/UnrealHeaderTool/Private/CodeGenerator.cpp
--- a/Engine/Source/Programs/UnrealHeaderTool/Private/CodeGenerator.cpp
+++ b/Engine/Source/Programs/UnrealHeaderTool/Private/CodeGenerator.cpp
@@ -5645,7 +5645,6 @@ FNativeClassHeaderGenerator::FNativeClassHeaderGenerator(
 	, APIStringPrivate           (FString::Printf(TEXT("%s_API "), *API))
 	, Package                    (InPackage)
 	, bAllowSaveExportedHeaders  (InAllowSaveExportedHeaders)
-	, bFailIfGeneratedCodeChanges(FParse::Param(FCommandLine::Get(), TEXT("FailIfGeneratedCodeChanges")))
 {
 	FString PackageName = FPackageName::GetShortName(Package);
 
@@ -6198,6 +6197,7 @@ bool FNativeClassHeaderGenerator::SaveHeaderIfChanged(FReferenceGatherers& OutRe
 	const bool bHasChanged = OriginalHeaderLocal.Len() == 0 || FCString::Strcmp(*OriginalHeaderLocal, NewHeaderContents);
 	if (bHasChanged)
 	{
+		static const bool bFailIfGeneratedCodeChanges = FParse::Param(FCommandLine::Get(), TEXT("FailIfGeneratedCodeChanges"));
 		if (bFailIfGeneratedCodeChanges)
 		{
 			FString ConflictPath = HeaderPathStr + TEXT(".conflict");
diff --git a/Engine/Source/Programs/UnrealHeaderTool/Private/NativeClassExporter.h b/Engine/Source/Programs/UnrealHeaderTool/Private/NativeClassExporter.h
--- a/Engine/Source/Programs/UnrealHeaderTool/Private/NativeClassExporter.h
+++ b/Engine/Source/Programs/UnrealHeaderTool/Private/NativeClassExporter.h
@@ -103,8 +103,8 @@ private:
 		{
 		}
 
-	/** Set of already exported cross-module references, to prevent duplicates */
-	TSet<FString>* UniqueCrossModuleReferences;
+		/** Set of already exported cross-module references, to prevent duplicates */
+		TSet<FString>* UniqueCrossModuleReferences;
 		/** Array of all header filenames from the current package. */
 		TSet<FString>& PackageHeaderPaths;
 		/** Array of temp filenames that for files to overwrite headers */
@@ -120,9 +120,6 @@ private:
 	/** If false, exported headers will not be saved to disk */
 	bool bAllowSaveExportedHeaders;
 
-	/** If true, any change in the generated headers will result in UHT failure. */
-	bool bFailIfGeneratedCodeChanges;
-
 	/**
 	 * Exports the struct's C++ properties to the HeaderText output device and adds special
 	 * compiler directives for GCC to pack as we expect.
