UHT optimization: UnrealSourceFile API cleanup
Eliminate Get/AppendDefinedClasses and use the const TMap& GetDefinedClassesWithParsingInfo instead to avoid building/copying arrays unnecessarily.
Cache GeneratedHeaderFilename, StrippedFilename, and FileId instead of repeatedly rebuilding
Use Move parameters as possible
#rnx

diff --git a/Engine/Source/Programs/UnrealHeaderTool/Private/CodeGenerator.cpp b/Engine/Source/Programs/UnrealHeaderTool/Private/CodeGenerator.cpp
--- a/Engine/Source/Programs/UnrealHeaderTool/Private/CodeGenerator.cpp
+++ b/Engine/Source/Programs/UnrealHeaderTool/Private/CodeGenerator.cpp
@@ -5661,18 +5661,16 @@ FNativeClassHeaderGenerator::FNativeClassHeaderGenerator(
 	});
 	if (bPackageHasAnyExportClasses)
 	{
-		TArray<UClass*> DefinedClasses;
 		for (FUnrealSourceFile* SourceFile : SourceFiles)
 		{
-			DefinedClasses.Reset();
-			SourceFile->AppendDefinedClasses(DefinedClasses);
-			for (UClass* Class : DefinedClasses)
+			for (const TPair<UClass*, FSimplifiedParsingClassInfo>& ClassDataPair : SourceFile->GetDefinedClassesWithParsingInfo())
 			{
+				UClass* Class = ClassDataPair.Key;
 				if (!Class->HasAnyClassFlags(CLASS_Native))
 				{
 					Class->UnMark(EObjectMark(OBJECTMARK_TagImp | OBJECTMARK_TagExp));
 				}
-				else if (GTypeDefinitionInfoMap.Contains(Class) && !Class->HasAnyClassFlags(CLASS_NoExport))
+				else if (!Class->HasAnyClassFlags(CLASS_NoExport) && GTypeDefinitionInfoMap.Contains(Class))
 				{
 					bWriteClassesH = true;
 					Class->UnMark(OBJECTMARK_TagImp);
@@ -5724,8 +5722,9 @@ FNativeClassHeaderGenerator::FNativeClassHeaderGenerator(
 		// This needs to be done outside of parallel blocks because it will modify UClass memory.
 		// Later calls to SetUpUhtReplicationData inside parallel blocks should be fine, because
 		// they will see the memory has already been set up, and just return the parent pointer.
-		for (UClass* Class : SourceFile->GetDefinedClasses())
+		for (const TPair<UClass*, FSimplifiedParsingClassInfo>& ClassDataPair : SourceFile->GetDefinedClassesWithParsingInfo())
 		{
+			UClass* Class = ClassDataPair.Key;
 			if (ClassHasReplicatedProperties(Class))
 			{
 				Class->SetUpUhtReplicationData();
@@ -5811,10 +5810,9 @@ FNativeClassHeaderGenerator::FNativeClassHeaderGenerator(
 
 		EExportClassOutFlags ExportFlags = EExportClassOutFlags::None;
 		TSet<FString> AdditionalHeaders;
-		TArray<UClass*>	DefinedClasses;
-		SourceFile->AppendDefinedClasses(DefinedClasses);
-		for (UClass* Class : DefinedClasses)
+		for (const TPair<UClass*, FSimplifiedParsingClassInfo>& ClassDataPair : SourceFile->GetDefinedClassesWithParsingInfo())
 		{
+			UClass* Class = ClassDataPair.Key;
 			if (!(Class->ClassFlags & CLASS_Intrinsic))
 			{
 				ConstThis->ExportClassFromSourceFileInner(GeneratedHeaderText, OutText, GeneratedFunctionDeclarations, ReferenceGatherers, (FClass*)Class, *SourceFile, ExportFlags);
@@ -5848,10 +5846,9 @@ FNativeClassHeaderGenerator::FNativeClassHeaderGenerator(
 
 	for (FUnrealSourceFile* SourceFile : Exported)
 	{
-		DefinedClasses.Reset();
-		SourceFile->AppendDefinedClasses(DefinedClasses);
-		for (UClass* Class : DefinedClasses)
+		for (const TPair<UClass*, FSimplifiedParsingClassInfo>& ClassDataPair : SourceFile->GetDefinedClassesWithParsingInfo())
 		{
+			UClass* Class = ClassDataPair.Key;
 			GClassToSourceFileMap.Add(Class, SourceFile);
 		}
 
@@ -5871,10 +5868,9 @@ FNativeClassHeaderGenerator::FNativeClassHeaderGenerator(
 		bool bAddedArchiveUObjectFromStructuredArchiveHeader = false;
 
 		TArray<FString>& RelativeIncludes = GeneratedCPPs[SourceFile].RelativeIncludes;
-		DefinedClasses.Reset();
-		SourceFile->AppendDefinedClasses(DefinedClasses);
-		for (UClass* Class : DefinedClasses)
+		for (const TPair<UClass*, FSimplifiedParsingClassInfo>& ClassDataPair : SourceFile->GetDefinedClassesWithParsingInfo())
 		{
+			UClass* Class = ClassDataPair.Key;
 			if (Class->ClassWithin && Class->ClassWithin != UObject::StaticClass())
 			{
 				if (FUnrealSourceFile** WithinSourceFile = GClassToSourceFileMap.Find(Class->ClassWithin))
@@ -6486,7 +6482,6 @@ ECompilationResult::Type PreparseModules(const FString& ModuleInfoPath, int32& N
 
 					TSharedRef<FUnrealSourceFile> UnrealSourceFile = PerformInitialParseOnHeader(Package, *RawFilename, RF_Public | RF_Standalone, *HeaderFile);
 					FUnrealSourceFile* UnrealSourceFilePtr = &UnrealSourceFile.Get();
-					TArray<UClass*> DefinedClasses = UnrealSourceFile->GetDefinedClasses();
 					GUnrealSourceFilesMap.Add(FPaths::GetCleanFilename(RawFilename), UnrealSourceFile);
 
 					if (CurrentlyProcessing == PublicClassesHeaders)
diff --git a/Engine/Source/Programs/UnrealHeaderTool/Private/HeaderParser.cpp b/Engine/Source/Programs/UnrealHeaderTool/Private/HeaderParser.cpp
--- a/Engine/Source/Programs/UnrealHeaderTool/Private/HeaderParser.cpp
+++ b/Engine/Source/Programs/UnrealHeaderTool/Private/HeaderParser.cpp
@@ -8455,8 +8455,9 @@ ECompilationResult::Type FHeaderParser::ParseHeader(FClasses& AllClasses, FUnrea
 		// First-pass success.
 		Result = ECompilationResult::Succeeded;
 
-		for (UClass* Class : CurrentSrcFile->GetDefinedClasses())
+		for (const TPair<UClass*, FSimplifiedParsingClassInfo>& ClassDataPair : CurrentSrcFile->GetDefinedClassesWithParsingInfo())
 		{
+			UClass* Class = ClassDataPair.Key;
 			PostParsingClassSetup(Class);
 
 			// Clean up and exit.
@@ -8556,10 +8557,9 @@ ECompilationResult::Type FHeaderParser::ParseHeaders(FClasses& AllClasses, FHead
 		}
 	}
 
-	const TArray<UClass*>& Classes = SourceFile->GetDefinedClasses();
-
-	for (UClass* Class : Classes)
+	for (const TPair<UClass*, FSimplifiedParsingClassInfo>& ClassDataPair : SourceFile->GetDefinedClassesWithParsingInfo())
 	{
+		UClass* Class = ClassDataPair.Key;
 		for (UClass* ParentClass = Class->GetSuperClass(); ParentClass && !ParentClass->HasAnyClassFlags(CLASS_Parsed | CLASS_Intrinsic); ParentClass = ParentClass->GetSuperClass())
 		{
 			SourceFilesRequired.Add(&GTypeDefinitionInfoMap[ParentClass]->GetUnrealSourceFile());
@@ -8582,8 +8582,9 @@ ECompilationResult::Type FHeaderParser::ParseHeaders(FClasses& AllClasses, FHead
 	{
 		ECompilationResult::Type OneFileResult = HeaderParser.ParseHeader(AllClasses, SourceFile);
 
-		for (UClass* Class : Classes)
+		for (const TPair<UClass*, FSimplifiedParsingClassInfo>& ClassDataPair : SourceFile->GetDefinedClassesWithParsingInfo())
 		{
+			UClass* Class = ClassDataPair.Key;
 			Class->ClassFlags |= CLASS_Parsed;
 		}
 
diff --git a/Engine/Source/Programs/UnrealHeaderTool/Private/UnrealSourceFile.cpp b/Engine/Source/Programs/UnrealHeaderTool/Private/UnrealSourceFile.cpp
--- a/Engine/Source/Programs/UnrealHeaderTool/Private/UnrealSourceFile.cpp
+++ b/Engine/Source/Programs/UnrealHeaderTool/Private/UnrealSourceFile.cpp
@@ -5,53 +5,63 @@
 #include "Misc/PackageName.h"
 #include "HeaderParser.h"
 
-void FUnrealSourceFile::AddDefinedClass(UClass* Class, FSimplifiedParsingClassInfo ParsingInfo)
+void FUnrealSourceFile::AddDefinedClass(UClass* Class, FSimplifiedParsingClassInfo&& ParsingInfo)
 {
 	DefinedClasses.Add(Class, MoveTemp(ParsingInfo));
 }
 
-FString FUnrealSourceFile::GetFileId() const
+const FString& FUnrealSourceFile::GetFileId() const
 {
-	FString StdFilename = Filename;
-
-	FPaths::MakeStandardFilename(StdFilename);
-
-	bool RelativePath = FPaths::IsRelative(StdFilename);
-
-	if (!RelativePath)
+	if (FileId.Len() == 0)
 	{
-		// If path is still absolute that means MakeStandardFilename has failed
-		// In this case make it relative to the current project. 
-		RelativePath = FPaths::MakePathRelativeTo(StdFilename, *FPaths::GetPath(FPaths::GetProjectFilePath()));
-	}
+		FString StdFilename = Filename;
 
-	// If the path has passed either MakeStandardFilename or MakePathRelativeTo it should be using internal path seperators
-	if (RelativePath)
-	{
-		// Remove any preceding parent directory paths
-		while (StdFilename.RemoveFromStart(TEXT("../")));
-	}
+		FPaths::MakeStandardFilename(StdFilename);
 
-	FStringOutputDevice Out;
+		bool bRelativePath = FPaths::IsRelative(StdFilename);
 
-	for (auto Char : StdFilename)
-	{
-		if (FChar::IsAlnum(Char))
+		if (!bRelativePath)
 		{
-			Out.AppendChar(Char);
+			// If path is still absolute that means MakeStandardFilename has failed
+			// In this case make it relative to the current project. 
+			bRelativePath = FPaths::MakePathRelativeTo(StdFilename, *FPaths::GetPath(FPaths::GetProjectFilePath()));
 		}
-		else
+
+		// If the path has passed either MakeStandardFilename or MakePathRelativeTo it should be using internal path seperators
+		if (bRelativePath)
 		{
-			Out.AppendChar('_');
+			// Remove any preceding parent directory paths
+			while (StdFilename.RemoveFromStart(TEXT("../")));
 		}
+
+		FStringOutputDevice Out;
+
+		for (TCHAR Char : StdFilename)
+		{
+			if (FChar::IsAlnum(Char))
+			{
+				Out.AppendChar(Char);
+			}
+			else
+			{
+				Out.AppendChar('_');
+			}
+		}
+
+		FileId = MoveTemp(Out);
 	}
 
-	return MoveTemp(Out);
+	return FileId;
 }
 
-FString FUnrealSourceFile::GetStrippedFilename() const
+const FString& FUnrealSourceFile::GetStrippedFilename() const
 {
-	return FPaths::GetBaseFilename(Filename);
+	if (StrippedFilename.Len() == 0)
+	{
+		StrippedFilename = FPaths::GetBaseFilename(Filename);
+	}
+
+	return StrippedFilename;
 }
 
 FString FUnrealSourceFile::GetGeneratedMacroName(FClassMetaData* ClassData, const TCHAR* Suffix) const
@@ -74,7 +84,7 @@ FString FUnrealSourceFile::GetGeneratedBodyMacroName(int32 LineNumber, bool bLeg
 	return GetGeneratedMacroName(LineNumber, *FString::Printf(TEXT("%s%s"), TEXT("_GENERATED_BODY"), bLegacy ? TEXT("_LEGACY") : TEXT("")));
 }
 
-void FUnrealSourceFile::SetGeneratedFilename(FString InGeneratedFilename)
+void FUnrealSourceFile::SetGeneratedFilename(FString&& InGeneratedFilename)
 {
 	GeneratedFilename = MoveTemp(InGeneratedFilename);
 }
@@ -84,12 +94,12 @@ void FUnrealSourceFile::SetHasChanged(bool bInHasChanged)
 	bHasChanged = bInHasChanged;
 }
 
-void FUnrealSourceFile::SetModuleRelativePath(FString InModuleRelativePath)
+void FUnrealSourceFile::SetModuleRelativePath(FString&& InModuleRelativePath)
 {
 	ModuleRelativePath = MoveTemp(InModuleRelativePath);
 }
 
-void FUnrealSourceFile::SetIncludePath(FString InIncludePath)
+void FUnrealSourceFile::SetIncludePath(FString&& InIncludePath)
 {
 	IncludePath = MoveTemp(InIncludePath);
 }
@@ -147,12 +157,8 @@ bool FUnrealSourceFile::HasChanged() const
 	return bHasChanged;
 }
 
-FString FUnrealSourceFile::GetAPI() const
-{
-	return FPackageName::GetShortName(Package).ToUpper();
-}
-
 FString FUnrealSourceFile::GetFileDefineName() const
 {
-	return FString::Printf(TEXT("%s_%s_generated_h"), *GetAPI(), *GetStrippedFilename());
+	const FString API = FPackageName::GetShortName(Package).ToUpper();
+	return FString::Printf(TEXT("%s_%s_generated_h"), *API, *GetStrippedFilename());
 }
diff --git a/Engine/Source/Programs/UnrealHeaderTool/Private/UnrealSourceFile.h b/Engine/Source/Programs/UnrealHeaderTool/Private/UnrealSourceFile.h
--- a/Engine/Source/Programs/UnrealHeaderTool/Private/UnrealSourceFile.h
+++ b/Engine/Source/Programs/UnrealHeaderTool/Private/UnrealSourceFile.h
@@ -26,44 +26,17 @@ public:
 		: Scope                (MakeShareable(new FFileScope(*(FString(TEXT("__")) + FPaths::GetBaseFilename(InFilename) + FString(TEXT("__File"))), this)))
 		, Filename             (MoveTemp(InFilename))
 		, Package              (InPackage)
-		, bHasChanged          (false)
 		, Content              (MoveTemp(InContent))
+		, bHasChanged          (false)
 		, bParsed              (false)
 		, bDependenciesResolved(false)
 	{
 		if (GetStrippedFilename() != "NoExportTypes")
 		{
-			Includes.Add(FHeaderProvider(EHeaderProviderSourceType::FileName, "NoExportTypes.h"));
-		}
-	}
-
-	/**
-	 * Appends array with classes defined in this source file.
-	 *
-	 * @param OutClasses (Output parameter) Array to fill with classes defined.
-	 */
-	void AppendDefinedClasses(TArray<UClass*>& OutClasses) const
-	{
-		for (const auto& KeyValuePair : DefinedClasses)
-		{
-			OutClasses.Add(KeyValuePair.Key);
+			Includes.Emplace(FHeaderProvider(EHeaderProviderSourceType::FileName, "NoExportTypes.h"));
 		}
 	}
 
-	/**
-	 * Gets array with classes defined in this source file.
-	 *
-	 * @returns Array with classes defined in this source file.
-	 */
-	TArray<UClass*> GetDefinedClasses() const
-	{
-		TArray<UClass*> Array;
-
-		AppendDefinedClasses(Array);
-
-		return Array;
-	}
-
 	/**
 	 * Gets parsing info for class defined in this source file.
 	 *
@@ -89,15 +62,20 @@ public:
 	 */
 	int32 GetDefinedClassesCount() const
 	{
-		return GetDefinedClasses().Num();
+		return DefinedClasses.Num();
 	}
 
 	/**
 	 * Gets generated header filename.
 	 */
-	FString GetGeneratedHeaderFilename() const
+	const FString& GetGeneratedHeaderFilename() const
 	{
-		return FString::Printf(TEXT("%s.generated.h"), *FPaths::GetBaseFilename(Filename));
+		if (GeneratedHeaderFilename.Len() == 0)
+		{
+			GeneratedHeaderFilename = FString::Printf(TEXT("%s.generated.h"), *FPaths::GetBaseFilename(Filename));
+		}
+
+		return GeneratedHeaderFilename;
 	}
 
 	/**
@@ -111,17 +89,12 @@ public:
 	/**
 	 * Gets stripped filename.
 	 */
-	FString GetStrippedFilename() const;
+	const FString& GetStrippedFilename() const;
 
 	/**
 	 * Gets unique file id.
 	 */
-	FString GetFileId() const;
-
-	/**
-	 * Gets source file API.
-	 */
-	FString GetAPI() const;
+	const FString& GetFileId() const;
 
 	/**
 	 * Gets define name of this source file.
@@ -158,7 +131,7 @@ public:
 	 * @param Class Class to add to list.
 	 * @param ParsingInfo Data from simplified parsing step.
 	 */
-	void AddDefinedClass(UClass* Class, FSimplifiedParsingClassInfo ParsingInfo);
+	void AddDefinedClass(UClass* Class, FSimplifiedParsingClassInfo&& ParsingInfo);
 
 	/**
 	 * Gets scope for this file.
@@ -245,7 +218,7 @@ public:
 	/**
 	 * Sets generated filename.
 	 */
-	void SetGeneratedFilename(FString GeneratedFilename);
+	void SetGeneratedFilename(FString&& GeneratedFilename);
 
 	/**
 	 * Sets has changed flag.
@@ -255,12 +228,12 @@ public:
 	/**
 	 * Sets module relative path.
 	 */
-	void SetModuleRelativePath(FString ModuleRelativePath);
+	void SetModuleRelativePath(FString&& ModuleRelativePath);
 
 	/**
 	 * Sets include path.
 	 */
-	void SetIncludePath(FString IncludePath);
+	void SetIncludePath(FString&& IncludePath);
 
 	/**
 	 * Mark this file as parsed.
@@ -297,14 +270,20 @@ private:
 	// Path of this file.
 	FString Filename;
 
+	// Stripped path of this file
+	mutable FString StrippedFilename;
+
+	// Cached FileId for this file
+	mutable FString FileId;
+
 	// Package of this file.
 	UPackage* Package;
 
 	// File name of the generated header file associated with this file.
 	FString GeneratedFilename;
 
-	// Tells if generated header file was changed.
-	bool bHasChanged;
+	// File name of the generated header file associated with this file.
+	mutable FString GeneratedHeaderFilename;
 
 	// Module relative path.
 	FString ModuleRelativePath;
@@ -315,6 +294,9 @@ private:
 	// Source file content.
 	FString Content;
 
+	// Tells if generated header file was changed.
+	bool bHasChanged;
+
 	// Tells if this file was parsed.
 	bool bParsed;
 
