UHT optimization: Encapsulate GUnrealSourceFilesMap, GTypeDefinitionInfoMap, and GPublicSourceFileSet to allow indexing in multiple ways
#rnx

diff --git a/Engine/Source/Programs/UnrealHeaderTool/Private/ClassMaps.cpp b/Engine/Source/Programs/UnrealHeaderTool/Private/ClassMaps.cpp
--- a/Engine/Source/Programs/UnrealHeaderTool/Private/ClassMaps.cpp
+++ b/Engine/Source/Programs/UnrealHeaderTool/Private/ClassMaps.cpp
@@ -4,11 +4,11 @@
 #include "UnrealHeaderTool.h"
 #include "UnrealTypeDefinitionInfo.h"
 
-TMap<FString, TSharedRef<FUnrealSourceFile> > GUnrealSourceFilesMap;
-TMap<UField*, TSharedRef<FUnrealTypeDefinitionInfo> > GTypeDefinitionInfoMap;
+FUnrealSourceFiles GUnrealSourceFilesMap;
+FTypeDefinitionInfoMap GTypeDefinitionInfoMap;
 TMap<const UPackage*, TArray<UField*>> GPackageSingletons;
 FCriticalSection GPackageSingletonsCriticalSection;
-TSet<FUnrealSourceFile*> GPublicSourceFileSet;
+FPublicSourceFileSet GPublicSourceFileSet;
 TMap<FProperty*, FString> GArrayDimensions;
 TMap<UPackage*,  const FManifestModule*> GPackageToManifestModuleMap;
 TMap<void*, uint32> GGeneratedCodeHashes;
diff --git a/Engine/Source/Programs/UnrealHeaderTool/Private/ClassMaps.h b/Engine/Source/Programs/UnrealHeaderTool/Private/ClassMaps.h
--- a/Engine/Source/Programs/UnrealHeaderTool/Private/ClassMaps.h
+++ b/Engine/Source/Programs/UnrealHeaderTool/Private/ClassMaps.h
@@ -32,11 +32,69 @@ struct FArchiveTypeDefinePair
 	FString EnclosingDefine;
 };
 
-extern TMap<FString, TSharedRef<FUnrealSourceFile> > GUnrealSourceFilesMap;
-extern TMap<UField*, TSharedRef<FUnrealTypeDefinitionInfo> > GTypeDefinitionInfoMap;
+// Wrapper class around TypeDefinition map so we can maintain a parallel by name map
+struct FTypeDefinitionInfoMap
+{
+	void Add(UField* Field, TSharedRef<FUnrealTypeDefinitionInfo>&& Definition)
+	{
+		DefinitionsByField.Add(Field, Definition);
+		DefinitionsByName.Add(Field->GetFName(), MoveTemp(Definition));
+	}
+	bool Contains(const UField* Field) { return DefinitionsByField.Contains(Field); }
+	TSharedRef<FUnrealTypeDefinitionInfo>* Find(const UField* Field) { return DefinitionsByField.Find(Field); }
+	TSharedRef<FUnrealTypeDefinitionInfo>* FindByName(const FName Name) { return DefinitionsByName.Find(Name); }
+	TSharedRef<FUnrealTypeDefinitionInfo>& operator[](const UField* Field) { return DefinitionsByField[Field]; }
+
+private:
+
+	TMap<UField*, TSharedRef<FUnrealTypeDefinitionInfo>> DefinitionsByField;
+	TMap<FName, TSharedRef<FUnrealTypeDefinitionInfo>> DefinitionsByName;
+};
+
+// Wrapper class around SourceFiles map so we can quickly get a list of source files for a given package
+struct FUnrealSourceFiles
+{
+	void Add(FString&& Filename, TSharedRef<FUnrealSourceFile> SourceFile)
+	{
+		SourceFilesByString.Add(MoveTemp(Filename), SourceFile);
+		SourceFilesByPackage.FindOrAdd(SourceFile->GetPackage()).Add(&SourceFile.Get());
+	}
+	const TSharedRef<FUnrealSourceFile>* Find(const FString& Id) const { return SourceFilesByString.Find(Id); }
+	const TArray<FUnrealSourceFile*>* FindFilesForPackage(const UPackage* Package) const { return SourceFilesByPackage.Find(Package); }
+
+private:
+	// A map of all source files indexed by string.
+	TMap<FString, TSharedRef<FUnrealSourceFile>> SourceFilesByString;
+
+	// The list of source files per package. Stored as raw pointer since SourceFilesByString holds shared ref.
+	TMap<UPackage*, TArray<FUnrealSourceFile*>> SourceFilesByPackage;
+};
+
+// Wrapper class around PublicSourceFile set so we can quickly get a list of source files for a given package
+struct FPublicSourceFileSet
+{
+	void Add(FUnrealSourceFile* SourceFile)
+	{
+		SourceFileSet.Add(SourceFile);
+		SourceFilesByPackage.FindOrAdd(SourceFile->GetPackage()).Add(SourceFile);
+	}
+	bool Contains(FUnrealSourceFile* SourceFile) const { return SourceFileSet.Contains(SourceFile); }
+	const TArray<FUnrealSourceFile*>* FindFilesForPackage(const UPackage* Package) const { return SourceFilesByPackage.Find(Package); }
+
+private:
+
+	// The set of all public source files. Stored as raw pointer since FUnrealSourceFiles::SourceFilesByString holds shared ref.
+	TSet<FUnrealSourceFile*> SourceFileSet;
+
+	// The list of public source files per package. Stored as raw pointer since FUnrealSourceFiles::SourceFilesByString holds shared ref.
+	TMap<UPackage*, TArray<FUnrealSourceFile*>> SourceFilesByPackage;
+};
+
+extern FUnrealSourceFiles GUnrealSourceFilesMap;
+extern FTypeDefinitionInfoMap GTypeDefinitionInfoMap;
 extern TMap<const UPackage*, TArray<UField*>> GPackageSingletons;
 extern FCriticalSection GPackageSingletonsCriticalSection;
-extern TSet<FUnrealSourceFile*> GPublicSourceFileSet;
+extern FPublicSourceFileSet GPublicSourceFileSet;
 extern TMap<FProperty*, FString> GArrayDimensions;
 extern TMap<UPackage*,  const FManifestModule*> GPackageToManifestModuleMap;
 extern TMap<void*, uint32> GGeneratedCodeHashes;
diff --git a/Engine/Source/Programs/UnrealHeaderTool/Private/CodeGenerator.cpp b/Engine/Source/Programs/UnrealHeaderTool/Private/CodeGenerator.cpp
--- a/Engine/Source/Programs/UnrealHeaderTool/Private/CodeGenerator.cpp
+++ b/Engine/Source/Programs/UnrealHeaderTool/Private/CodeGenerator.cpp
@@ -5910,12 +5910,9 @@ FNativeClassHeaderGenerator::FNativeClassHeaderGenerator(
 		ClassesHText.Log(TEXT("\r\n"));
 
 		// Fill with the rest source files from this package.
-		for (FUnrealSourceFile* SourceFile : GPublicSourceFileSet)
+		if (const TArray<FUnrealSourceFile*>* SourceFilesForPackage = GPublicSourceFileSet.FindFilesForPackage(InPackage))
 		{
-			if (SourceFile->GetPackage() == InPackage)
-			{
-				PublicHeaderGroupIncludes.Add(SourceFile);
-			}
+			PublicHeaderGroupIncludes.Append(*SourceFilesForPackage);
 		}
 
 		for (FUnrealSourceFile* SourceFile : PublicHeaderGroupIncludes)
diff --git a/Engine/Source/Programs/UnrealHeaderTool/Private/HeaderParser.cpp b/Engine/Source/Programs/UnrealHeaderTool/Private/HeaderParser.cpp
--- a/Engine/Source/Programs/UnrealHeaderTool/Private/HeaderParser.cpp
+++ b/Engine/Source/Programs/UnrealHeaderTool/Private/HeaderParser.cpp
@@ -8511,16 +8511,17 @@ ECompilationResult::Type FHeaderParser::ParseHeader(FClasses& AllClasses, FUnrea
 
 ECompilationResult::Type FHeaderParser::ParseRestOfModulesSourceFiles(FClasses& AllClasses, UPackage* ModulePackage, FHeaderParser& HeaderParser)
 {
-	for (auto& Pair : GUnrealSourceFilesMap)
+	if (const TArray<FUnrealSourceFile*>* SourceFiles = GUnrealSourceFilesMap.FindFilesForPackage(ModulePackage))
 	{
-		FUnrealSourceFile* SourceFile = &Pair.Value.Get();
-
-		if (SourceFile->GetPackage() == ModulePackage && (!SourceFile->IsParsed() || SourceFile->GetDefinedClassesCount() == 0))
+		for (FUnrealSourceFile* SourceFile : *SourceFiles)
 		{
-			ECompilationResult::Type Result;
-			if ((Result = ParseHeaders(AllClasses, HeaderParser, SourceFile)) != ECompilationResult::Succeeded)
+			if (!SourceFile->IsParsed() || SourceFile->GetDefinedClassesCount() == 0)
 			{
-				return Result;
+				ECompilationResult::Type Result;
+				if ((Result = ParseHeaders(AllClasses, HeaderParser, SourceFile)) != ECompilationResult::Succeeded)
+				{
+					return Result;
+				}
 			}
 		}
 	}
@@ -8644,14 +8645,14 @@ TSet<FUnrealSourceFile*> GetSourceFilesWithInheritanceOrdering(UPackage* Current
 	}
 
 	// Then add the rest.
-	for (auto& Pair : GUnrealSourceFilesMap)
+	if (const TArray<FUnrealSourceFile*>* SourceFilesForPackage = GUnrealSourceFilesMap.FindFilesForPackage(CurrentPackage))
 	{
-		auto& SourceFile = Pair.Value.Get();
-
-		if (SourceFile.GetPackage() == CurrentPackage
-			&& SourceFile.GetScope()->ContainsTypes())
+		for (FUnrealSourceFile* SourceFile : *SourceFilesForPackage)
 		{
-			SourceFiles.Add(&SourceFile);
+			if (SourceFile->GetScope()->ContainsTypes())
+			{
+				SourceFiles.Add(SourceFile);
+			}
 		}
 	}
 
@@ -8865,14 +8866,17 @@ ECompilationResult::Type FHeaderParser::ParseAllHeadersInside(
 
 		HeaderParser.Filename = IFileManager::Get().ConvertToAbsolutePathForExternalAppForRead(*RootFilename);
 
-		for (FUnrealSourceFile* SourceFile : GPublicSourceFileSet)
+		if (const TArray<FUnrealSourceFile*>* SourceFiles = GPublicSourceFileSet.FindFilesForPackage(CurrentPackage))
 		{
-			if (SourceFile->GetPackage() == CurrentPackage && (!SourceFile->IsParsed() || SourceFile->GetDefinedClassesCount() == 0))
+			for (FUnrealSourceFile* SourceFile : *SourceFiles)
 			{
-				Result = ParseHeaders(ModuleClasses, HeaderParser, SourceFile);
-				if (Result != ECompilationResult::Succeeded)
+				if ((!SourceFile->IsParsed() || SourceFile->GetDefinedClassesCount() == 0))
 				{
-					return Result;
+					Result = ParseHeaders(ModuleClasses, HeaderParser, SourceFile);
+					if (Result != ECompilationResult::Succeeded)
+					{
+						return Result;
+					}
 				}
 			}
 		}
diff --git a/Engine/Source/Programs/UnrealHeaderTool/Private/HeaderProvider.cpp b/Engine/Source/Programs/UnrealHeaderTool/Private/HeaderProvider.cpp
--- a/Engine/Source/Programs/UnrealHeaderTool/Private/HeaderProvider.cpp
+++ b/Engine/Source/Programs/UnrealHeaderTool/Private/HeaderProvider.cpp
@@ -19,14 +19,10 @@ FUnrealSourceFile* FHeaderProvider::Resolve()
 	{
 		if (Type == EHeaderProviderSourceType::ClassName)
 		{
-			FName IdName(*Id);
-			for (const auto& Pair : GTypeDefinitionInfoMap)
+			FName IdName(*Id, FNAME_Find);
+			if (TSharedRef<FUnrealTypeDefinitionInfo>* Source = GTypeDefinitionInfoMap.FindByName(IdName))
 			{
-				if (Pair.Key->GetFName() == IdName)
-				{
-					Cache = &Pair.Value->GetUnrealSourceFile();
-					break;
-				}
+				Cache = &(*Source)->GetUnrealSourceFile();
 			}
 		}
 		else if (const TSharedRef<FUnrealSourceFile>* Source = GUnrealSourceFilesMap.Find(Id))
