UHT optimization: Avoid copy of GetGeneratedHeaderFilename and use the const& instead
(should have been part of CL# 13181587)
#rnx

diff --git a/Engine/Source/Programs/UnrealHeaderTool/Private/HeaderParser.cpp b/Engine/Source/Programs/UnrealHeaderTool/Private/HeaderParser.cpp
--- a/Engine/Source/Programs/UnrealHeaderTool/Private/HeaderParser.cpp
+++ b/Engine/Source/Programs/UnrealHeaderTool/Private/HeaderParser.cpp
@@ -3468,10 +3468,10 @@ void FHeaderParser::CompileDirective(FClasses& AllClasses)
 	}
 	else if (Directive.Matches(TEXT("include"), ESearchCase::CaseSensitive))
 	{
-		FString ExpectedHeaderName = CurrentSrcFile->GetGeneratedHeaderFilename();
 		FToken IncludeName;
 		if (GetToken(IncludeName) && (IncludeName.TokenType == TOKEN_Const) && (IncludeName.Type == CPT_String))
 		{
+			const FString& ExpectedHeaderName = CurrentSrcFile->GetGeneratedHeaderFilename();
 			if (FCString::Stricmp(IncludeName.String, *ExpectedHeaderName) == 0)
 			{
 				bSpottedAutogeneratedHeaderInclude = true;
@@ -8473,7 +8473,7 @@ ECompilationResult::Type FHeaderParser::ParseHeader(FClasses& AllClasses, FUnrea
 
 		if (!bSpottedAutogeneratedHeaderInclude && !bEmptyFile && !bNoExportClassesOnly)
 		{
-			const FString ExpectedHeaderName = CurrentSrcFile->GetGeneratedHeaderFilename();
+			const FString& ExpectedHeaderName = CurrentSrcFile->GetGeneratedHeaderFilename();
 			FError::Throwf(TEXT("Expected an include at the top of the header: '#include \"%s\"'"), *ExpectedHeaderName);
 		}
 
