UHT Optimization: Parallelize the file loading and initial parsing of headers
#rnx

diff --git a/Engine/Source/Programs/UnrealHeaderTool/Private/ClassMaps.cpp b/Engine/Source/Programs/UnrealHeaderTool/Private/ClassMaps.cpp
--- a/Engine/Source/Programs/UnrealHeaderTool/Private/ClassMaps.cpp
+++ b/Engine/Source/Programs/UnrealHeaderTool/Private/ClassMaps.cpp
@@ -14,9 +14,31 @@ TMap<UPackage*,  const FManifestModule*> GPackageToManifestModuleMap;
 TMap<void*, uint32> GGeneratedCodeHashes;
 FRWLock GGeneratedCodeHashesLock;
 TMap<UEnum*,  EUnderlyingEnumType> GEnumUnderlyingTypes;
-TMap<FName, TSharedRef<FClassDeclarationMetaData> > GClassDeclarations;
+FClassDeclarations GClassDeclarations;
 TSet<FProperty*> GUnsizedProperties;
 TSet<UField*> GEditorOnlyDataTypes;
 TMap<UStruct*, TTuple<TSharedRef<FUnrealSourceFile>, int32>> GStructToSourceLine;
 TMap<UClass*, FArchiveTypeDefinePair> GClassSerializerMap;
 TSet<FProperty*> GPropertyUsesMemoryImageAllocator;
+
+void FClassDeclarations::AddIfMissing(FName Name, TUniqueFunction<TSharedRef<FClassDeclarationMetaData>()>&& DeclConstructFunc)
+{
+	FRWScopeLock Lock(ClassDeclLock, SLT_Write);
+	if (!ClassDeclarations.Contains(Name))
+	{
+		TSharedRef<FClassDeclarationMetaData> ClassDecl = DeclConstructFunc();
+		ClassDeclarations.Add(Name, MoveTemp(ClassDecl));
+	}
+}
+
+TSharedRef<FClassDeclarationMetaData>* FClassDeclarations::Find(FName Name)
+{
+	FRWScopeLock Lock(ClassDeclLock, SLT_ReadOnly);
+	return ClassDeclarations.Find(Name);
+}
+
+TSharedRef<FClassDeclarationMetaData>& FClassDeclarations::FindChecked(FName Name)
+{
+	FRWScopeLock Lock(ClassDeclLock, SLT_ReadOnly);
+	return ClassDeclarations.FindChecked(Name);
+}
diff --git a/Engine/Source/Programs/UnrealHeaderTool/Private/ClassMaps.h b/Engine/Source/Programs/UnrealHeaderTool/Private/ClassMaps.h
--- a/Engine/Source/Programs/UnrealHeaderTool/Private/ClassMaps.h
+++ b/Engine/Source/Programs/UnrealHeaderTool/Private/ClassMaps.h
@@ -51,6 +51,19 @@ private:
 	TMap<FName, TSharedRef<FUnrealTypeDefinitionInfo>> DefinitionsByName;
 };
 
+// Wrapper class around ClassDeclarations map so we can control access in a threadsafe manner
+struct FClassDeclarations
+{
+	void AddIfMissing(FName Name, TUniqueFunction<TSharedRef<FClassDeclarationMetaData>()>&& DeclConstructFunc);
+	TSharedRef<FClassDeclarationMetaData>* Find(FName Name);
+	TSharedRef<FClassDeclarationMetaData>& FindChecked(FName Name);
+
+private:
+	TMap<FName, TSharedRef<FClassDeclarationMetaData>> ClassDeclarations;
+
+	FRWLock ClassDeclLock;
+};
+
 // Wrapper class around SourceFiles map so we can quickly get a list of source files for a given package
 struct FUnrealSourceFiles
 {
@@ -100,7 +113,7 @@ extern TMap<UPackage*,  const FManifestModule*> GPackageToManifestModuleMap;
 extern TMap<void*, uint32> GGeneratedCodeHashes;
 extern FRWLock GGeneratedCodeHashesLock;
 extern TMap<UEnum*, EUnderlyingEnumType> GEnumUnderlyingTypes;
-extern TMap<FName, TSharedRef<FClassDeclarationMetaData> > GClassDeclarations;
+extern FClassDeclarations GClassDeclarations;
 extern TSet<FProperty*> GUnsizedProperties;
 extern TSet<UField*> GEditorOnlyDataTypes;
 extern TMap<UStruct*, TTuple<TSharedRef<FUnrealSourceFile>, int32>> GStructToSourceLine;
diff --git a/Engine/Source/Programs/UnrealHeaderTool/Private/CodeGenerator.cpp b/Engine/Source/Programs/UnrealHeaderTool/Private/CodeGenerator.cpp
--- a/Engine/Source/Programs/UnrealHeaderTool/Private/CodeGenerator.cpp
+++ b/Engine/Source/Programs/UnrealHeaderTool/Private/CodeGenerator.cpp
@@ -76,7 +76,15 @@ static TArray<FString> ChangeMessages;
 static bool bWriteContents = false;
 static bool bVerifyContents = false;
 
-static TSharedRef<FUnrealSourceFile> PerformInitialParseOnHeader(UPackage* InParent, const TCHAR* FileName, EObjectFlags Flags, const TCHAR* Buffer);
+struct FPerHeaderData
+{
+	TSharedPtr<FUnrealSourceFile> UnrealSourceFile;
+	TArray<FHeaderProvider> DependsOn;
+	TArray<FSimplifiedParsingClassInfo> ParsedClassArray;
+};
+
+static void PerformSimplifiedClassParse(UPackage* InParent, const TCHAR* FileName, const TCHAR* Buffer, FPerHeaderData& PerHeaderData);
+static void ProcessInitialClassParse(FPerHeaderData& PerHeaderData);
 
 FCompilerMetadataManager GScriptHelper;
 
@@ -6395,6 +6403,40 @@ ECompilationResult::Type PreparseModules(const FString& ModuleInfoPath, int32& N
 	};
 
 	ECompilationResult::Type Result = ECompilationResult::Succeeded;
+
+#if !PLATFORM_EXCEPTIONS_DISABLED
+	FGraphEventArray ExceptionTasks;
+
+	auto LogException = [&Result, &NumFailures, &ExceptionTasks](FString&& Filename, int32 Line, const FString& Message)
+	{
+		auto LogExceptionTask = [&Result, &NumFailures, Filename = MoveTemp(Filename), Line, Message]()
+		{
+			TGuardValue<ELogTimes::Type> DisableLogTimes(GPrintLogTimes, ELogTimes::None);
+
+			FString FormattedErrorMessage = FString::Printf(TEXT("%s(%d): Error: %s\r\n"), *Filename, Line, *Message);
+			Result = ECompilationResult::OtherCompilationError;
+
+			UE_LOG(LogCompile, Log, TEXT("%s"), *FormattedErrorMessage);
+			GWarn->Log(ELogVerbosity::Error, FormattedErrorMessage);
+
+			++NumFailures;
+		};
+
+		if (IsInGameThread())
+		{
+			LogExceptionTask();
+		}
+		else
+		{
+			FGraphEventRef EventRef = FFunctionGraphTask::CreateAndDispatchWhenReady(MoveTemp(LogExceptionTask), TStatId(), nullptr, ENamedThreads::GameThread);
+
+			static FCriticalSection ExceptionCS;
+			FScopeLock Lock(&ExceptionCS);
+			ExceptionTasks.Add(EventRef);
+		}
+	};
+#endif
+
 	for (FManifestModule& Module : GManifest.Modules)
 	{
 		if (Result != ECompilationResult::Succeeded)
@@ -6459,25 +6501,94 @@ ECompilationResult::Type PreparseModules(const FString& ModuleInfoPath, int32& N
 
 			NumHeadersPreparsed += UObjectHeaders.Num();
 
-			for (const FString& RawFilename : UObjectHeaders)
+			TArray<FString> HeaderFiles;
+			HeaderFiles.SetNum(UObjectHeaders.Num());
+
 			{
-#if !PLATFORM_EXCEPTIONS_DISABLED
-				try
-#endif
+				SCOPE_SECONDS_COUNTER_UHT(LoadHeaderContentFromFile);
+				ParallelFor(UObjectHeaders.Num(), [&](int32 Index)
 				{
-					// Import class.
-					const FString FullFilename = FPaths::ConvertRelativePathToFull(ModuleInfoPath, RawFilename);
+					const FString& RawFilename = UObjectHeaders[Index];
 
-					FString HeaderFile;
+#if !PLATFORM_EXCEPTIONS_DISABLED
+					try
+#endif
 					{
-						SCOPE_SECONDS_COUNTER_UHT(LoadHeaderContentFromFile);
-						if (!FFileHelper::LoadFileToString(HeaderFile, *FullFilename))
+						const FString FullFilename = FPaths::ConvertRelativePathToFull(ModuleInfoPath, RawFilename);
+
+						if (!FFileHelper::LoadFileToString(HeaderFiles[Index], *FullFilename))
 						{
 							FError::Throwf(TEXT("UnrealHeaderTool was unable to load source file '%s'"), *FullFilename);
 						}
 					}
+#if !PLATFORM_EXCEPTIONS_DISABLED
+					catch (TCHAR* ErrorMsg)
+					{
+						FString AbsFilename = IFileManager::Get().ConvertToAbsolutePathForExternalAppForRead(*RawFilename);
+						LogException(MoveTemp(AbsFilename), 1, ErrorMsg);
+					}				
+#endif
+				});
+			}
+
+#if !PLATFORM_EXCEPTIONS_DISABLED
+			FTaskGraphInterface::Get().WaitUntilTasksComplete(ExceptionTasks);
+#endif
 
-					TSharedRef<FUnrealSourceFile> UnrealSourceFile = PerformInitialParseOnHeader(Package, *RawFilename, RF_Public | RF_Standalone, *HeaderFile);
+			if (Result != ECompilationResult::Succeeded)
+			{
+				continue;
+			}
+
+			TArray<FPerHeaderData> PerHeaderData;
+			PerHeaderData.SetNum(UObjectHeaders.Num());
+
+			ParallelFor(UObjectHeaders.Num(), [&](int32 Index)
+			{
+				const FString& RawFilename = UObjectHeaders[Index];
+
+#if !PLATFORM_EXCEPTIONS_DISABLED
+				try
+#endif
+				{
+					PerformSimplifiedClassParse(Package, *RawFilename, *HeaderFiles[Index], PerHeaderData[Index]);
+				}
+#if !PLATFORM_EXCEPTIONS_DISABLED
+				catch (const FFileLineException& Ex)
+				{
+					FString AbsFilename = IFileManager::Get().ConvertToAbsolutePathForExternalAppForRead(*Ex.Filename);
+					LogException(MoveTemp(AbsFilename), Ex.Line, Ex.Message);
+				}
+				catch (TCHAR* ErrorMsg)
+				{
+					FString AbsFilename = IFileManager::Get().ConvertToAbsolutePathForExternalAppForRead(*RawFilename);
+					LogException(MoveTemp(AbsFilename), 1, ErrorMsg);
+				}
+#endif
+			});
+
+#if !PLATFORM_EXCEPTIONS_DISABLED
+			FTaskGraphInterface::Get().WaitUntilTasksComplete(ExceptionTasks);
+#endif
+
+			if (Result != ECompilationResult::Succeeded)
+			{
+				continue;
+			}
+
+			for (int32 Index = 0; Index < UObjectHeaders.Num(); ++Index)
+			{
+				const FString& RawFilename = UObjectHeaders[Index];
+
+#if !PLATFORM_EXCEPTIONS_DISABLED
+				try
+#endif
+				{
+					// Import class.
+					const FString FullFilename = FPaths::ConvertRelativePathToFull(ModuleInfoPath, RawFilename);
+
+					ProcessInitialClassParse(PerHeaderData[Index]);
+					TSharedRef<FUnrealSourceFile> UnrealSourceFile = PerHeaderData[Index].UnrealSourceFile.ToSharedRef();
 					FUnrealSourceFile* UnrealSourceFilePtr = &UnrealSourceFile.Get();
 					GUnrealSourceFilesMap.Add(FPaths::GetCleanFilename(RawFilename), UnrealSourceFile);
 
@@ -6530,31 +6641,13 @@ ECompilationResult::Type PreparseModules(const FString& ModuleInfoPath, int32& N
 #if !PLATFORM_EXCEPTIONS_DISABLED
 				catch (const FFileLineException& Ex)
 				{
-					TGuardValue<ELogTimes::Type> DisableLogTimes(GPrintLogTimes, ELogTimes::None);
-
-					FString AbsFilename           = IFileManager::Get().ConvertToAbsolutePathForExternalAppForRead(*Ex.Filename);
-					FString Prefix                = FString::Printf(TEXT("%s(%d): "), *AbsFilename, Ex.Line);
-					FString FormattedErrorMessage = FString::Printf(TEXT("%sError: %s\r\n"), *Prefix, *Ex.Message);
-					Result = GCompilationResult;
-
-					UE_LOG(LogCompile, Log, TEXT("%s"), *FormattedErrorMessage);
-					GWarn->Log(ELogVerbosity::Error, FormattedErrorMessage);
-
-					++NumFailures;
+					FString AbsFilename = IFileManager::Get().ConvertToAbsolutePathForExternalAppForRead(*Ex.Filename);
+					LogException(MoveTemp(AbsFilename), Ex.Line, Ex.Message);
 				}
 				catch (TCHAR* ErrorMsg)
 				{
-					TGuardValue<ELogTimes::Type> DisableLogTimes(GPrintLogTimes, ELogTimes::None);
-
-					FString AbsFilename           = IFileManager::Get().ConvertToAbsolutePathForExternalAppForRead(*RawFilename);
-					FString Prefix                = FString::Printf(TEXT("%s(1): "), *AbsFilename);
-					FString FormattedErrorMessage = FString::Printf(TEXT("%sError: %s\r\n"), *Prefix, ErrorMsg);
-					Result = GCompilationResult;
-
-					UE_LOG(LogCompile, Log, TEXT("%s"), *FormattedErrorMessage);
-					GWarn->Log(ELogVerbosity::Error, FormattedErrorMessage);
-
-					++NumFailures;
+					FString AbsFilename = IFileManager::Get().ConvertToAbsolutePathForExternalAppForRead(*RawFilename);
+					LogException(MoveTemp(AbsFilename), 1, ErrorMsg);
 				}
 #endif
 			}
@@ -6866,38 +6959,34 @@ UClass* ProcessParsedClass(bool bClassIsAnInterface, TArray<FHeaderProvider>& De
 	return ResultClass;
 }
 
-
-TSharedRef<FUnrealSourceFile> PerformInitialParseOnHeader(UPackage* InParent, const TCHAR* FileName, EObjectFlags Flags, const TCHAR* Buffer)
+void PerformSimplifiedClassParse(UPackage* InParent, const TCHAR* FileName, const TCHAR* Buffer, FPerHeaderData& PerHeaderData)
 {
-	const TCHAR* InBuffer = Buffer;
-
-	// is the parsed class name an interface?
-	bool bClassIsAnInterface = false;
-
-	TArray<FHeaderProvider> DependsOn;
-
 	// Parse the header to extract the information needed
 	FUHTStringBuilder ClassHeaderTextStrippedOfCppText;
-	TArray<FSimplifiedParsingClassInfo> ParsedClassArray;
-	FHeaderParser::SimplifiedClassParse(FileName, Buffer, /*out*/ ParsedClassArray, /*out*/ DependsOn, ClassHeaderTextStrippedOfCppText);
+
+	FHeaderParser::SimplifiedClassParse(FileName, Buffer, /*out*/ PerHeaderData.ParsedClassArray, /*out*/ PerHeaderData.DependsOn, ClassHeaderTextStrippedOfCppText);
 
 	FUnrealSourceFile* UnrealSourceFilePtr = new FUnrealSourceFile(InParent, FileName, MoveTemp(ClassHeaderTextStrippedOfCppText));
-	TSharedRef<FUnrealSourceFile> UnrealSourceFile = MakeShareable(UnrealSourceFilePtr);
-	for (FSimplifiedParsingClassInfo& ParsedClassInfo : ParsedClassArray)
+	PerHeaderData.UnrealSourceFile = MakeShareable(UnrealSourceFilePtr);
+}
+
+void ProcessInitialClassParse(FPerHeaderData& PerHeaderData)
+{
+	TSharedRef<FUnrealSourceFile> UnrealSourceFile = PerHeaderData.UnrealSourceFile.ToSharedRef();
+	UPackage* InParent = UnrealSourceFile->GetPackage();
+	for (FSimplifiedParsingClassInfo& ParsedClassInfo : PerHeaderData.ParsedClassArray)
 	{
-		UClass* ResultClass = ProcessParsedClass(ParsedClassInfo.IsInterface(), DependsOn, ParsedClassInfo.GetClassName(), ParsedClassInfo.GetBaseClassName(), InParent, Flags);
+		UClass* ResultClass = ProcessParsedClass(ParsedClassInfo.IsInterface(), PerHeaderData.DependsOn, ParsedClassInfo.GetClassName(), ParsedClassInfo.GetBaseClassName(), InParent, RF_Public | RF_Standalone);
 		GStructToSourceLine.Add(ResultClass, MakeTuple(UnrealSourceFile, ParsedClassInfo.GetClassDefLine()));
 
 		FScope::AddTypeScope(ResultClass, &UnrealSourceFile->GetScope().Get());
 
-		GTypeDefinitionInfoMap.Add(ResultClass, MakeShared<FUnrealTypeDefinitionInfo>(*UnrealSourceFilePtr, ParsedClassInfo.GetClassDefLine()));
+		GTypeDefinitionInfoMap.Add(ResultClass, MakeShared<FUnrealTypeDefinitionInfo>(UnrealSourceFile.Get(), ParsedClassInfo.GetClassDefLine()));
 		UnrealSourceFile->AddDefinedClass(ResultClass, MoveTemp(ParsedClassInfo));
 	}
 
-	for (FHeaderProvider& DependsOnElement : DependsOn)
+	for (FHeaderProvider& DependsOnElement : PerHeaderData.DependsOn)
 	{
 		UnrealSourceFile->GetIncludes().AddUnique(DependsOnElement);
 	}
-
-	return UnrealSourceFile;
-}
+}
\ No newline at end of file
diff --git a/Engine/Source/Programs/UnrealHeaderTool/Private/HeaderParser.cpp b/Engine/Source/Programs/UnrealHeaderTool/Private/HeaderParser.cpp
--- a/Engine/Source/Programs/UnrealHeaderTool/Private/HeaderParser.cpp
+++ b/Engine/Source/Programs/UnrealHeaderTool/Private/HeaderParser.cpp
@@ -9581,14 +9581,18 @@ void FHeaderPreParser::ParseClassDeclaration(const TCHAR* Filename, const TCHAR*
 
 	FString ClassNameWithoutPrefixStr = GetClassNameWithPrefixRemoved(out_ClassName);
 	out_StrippedClassName = *ClassNameWithoutPrefixStr;
-	TSharedRef<FClassDeclarationMetaData>* DeclarationDataPtr = GClassDeclarations.Find(out_StrippedClassName);
-	if (!DeclarationDataPtr)
-	{
-		// Add class declaration meta data so that we can access class flags before the class is fully parsed
-		TSharedRef<FClassDeclarationMetaData> DeclarationData = MakeShareable(new FClassDeclarationMetaData());
-		DeclarationData->MetaData = MoveTemp(MetaData);
-		DeclarationData->ParseClassProperties(MoveTemp(SpecifiersFound), RequiredAPIMacroIfPresent);
-		GClassDeclarations.Add(out_StrippedClassName, DeclarationData);
+
+	{
+		auto ConstructDeclarationData = [MetaData = MoveTemp(MetaData), SpecifiersFound = MoveTemp(SpecifiersFound), &RequiredAPIMacroIfPresent]() mutable
+		{
+			// Add class declaration meta data so that we can access class flags before the class is fully parsed
+			TSharedRef<FClassDeclarationMetaData> DeclarationData = MakeShareable(new FClassDeclarationMetaData());
+			DeclarationData->MetaData = MoveTemp(MetaData);
+			DeclarationData->ParseClassProperties(MoveTemp(SpecifiersFound), RequiredAPIMacroIfPresent);
+			return DeclarationData;
+		};
+
+		GClassDeclarations.AddIfMissing(out_StrippedClassName, MoveTemp(ConstructDeclarationData));
 	}
 
 	// Skip optional final keyword
