Merging CL# 13989043

Support for namespaced native types in UHT.
Some obvious string optimizations.

diff --git a/Engine/Source/Programs/UnrealHeaderTool/Private/HeaderParser.cpp b/Engine/Source/Programs/UnrealHeaderTool/Private/HeaderParser.cpp
--- a/Engine/Source/Programs/UnrealHeaderTool/Private/HeaderParser.cpp
+++ b/Engine/Source/Programs/UnrealHeaderTool/Private/HeaderParser.cpp
@@ -6143,40 +6143,60 @@ FClass* FHeaderParser::ParseClassNameDeclaration(FClasses& AllClasses, FString&
 		{
 			RequireIdentifier(TEXT("public"), ESearchCase::CaseSensitive, TEXT("Interface inheritance must be public"));
 
-			FToken Token;
-			if (!GetIdentifier(Token, true))
-				FError::Throwf(TEXT("Failed to get interface class identifier"));
-
-			FString InterfaceName = Token.Identifier;
+			FString InterfaceName;
 
-			// Handle templated native classes
-			if (MatchSymbol(TEXT('<')))
+			FToken Token;
+			for (;;)
 			{
-				InterfaceName += TEXT('<');
+				if (!GetIdentifier(Token, true))
+				{
+					FError::Throwf(TEXT("Failed to get interface class identifier"));
+				}
+
+				InterfaceName += Token.Identifier;
 
-				int32 NestedScopes = 1;
-				while (NestedScopes)
+				// Handle templated native classes
+				if (MatchSymbol(TEXT('<')))
 				{
-					if (!GetToken(Token))
-						FError::Throwf(TEXT("Unexpected end of file"));
+					InterfaceName += TEXT('<');
 
-					if (Token.TokenType == TOKEN_Symbol)
+					int32 NestedScopes = 1;
+					while (NestedScopes)
 					{
-						if (Token.Matches(TEXT('<')))
+						if (!GetToken(Token))
 						{
-							++NestedScopes;
+							FError::Throwf(TEXT("Unexpected end of file"));
 						}
-						else if (Token.Matches(TEXT('>')))
+
+						if (Token.TokenType == TOKEN_Symbol)
 						{
-							--NestedScopes;
+							if (Token.Matches(TEXT('<')))
+							{
+								++NestedScopes;
+							}
+							else if (Token.Matches(TEXT('>')))
+							{
+								--NestedScopes;
+							}
 						}
+
+						InterfaceName += Token.Identifier;
 					}
+				}
 
-					InterfaceName += Token.Identifier;
+				// Handle scoped native classes
+				if (MatchSymbol(TEXT("::")))
+				{
+					InterfaceName += TEXT("::");
+
+					// Keep reading nested identifiers
+					continue;
 				}
+
+				break;
 			}
 
-			HandleOneInheritedClass(AllClasses, FoundClass, *InterfaceName);
+			HandleOneInheritedClass(AllClasses, FoundClass, MoveTemp(InterfaceName));
 		}
 	}
 	else if (FoundClass->GetSuperClass())
@@ -6187,7 +6207,7 @@ FClass* FHeaderParser::ParseClassNameDeclaration(FClasses& AllClasses, FString&
 	return FoundClass;
 }
 
-void FHeaderParser::HandleOneInheritedClass(FClasses& AllClasses, UClass* Class, FString InterfaceName)
+void FHeaderParser::HandleOneInheritedClass(FClasses& AllClasses, UClass* Class, FString&& InterfaceName)
 {
 	FUnrealSourceFile* CurrentSrcFile = GetCurrentSourceFile();
 	// Check for UInterface derived interface inheritance
@@ -6215,7 +6235,7 @@ void FHeaderParser::HandleOneInheritedClass(FClasses& AllClasses, UClass* Class,
 		// Non-UObject inheritance
 		FClassMetaData* ClassData = GScriptHelper.FindClassData(Class);
 		check(ClassData);
-		ClassData->AddInheritanceParent(InterfaceName, CurrentSrcFile);
+		ClassData->AddInheritanceParent(MoveTemp(InterfaceName), CurrentSrcFile);
 	}
 }
 
diff --git a/Engine/Source/Programs/UnrealHeaderTool/Private/HeaderParser.h b/Engine/Source/Programs/UnrealHeaderTool/Private/HeaderParser.h
--- a/Engine/Source/Programs/UnrealHeaderTool/Private/HeaderParser.h
+++ b/Engine/Source/Programs/UnrealHeaderTool/Private/HeaderParser.h
@@ -781,7 +781,7 @@ protected:
 	bool SkipDeclaration(FToken& Token);
 	/** Similar to MatchSymbol() but will return to the exact location as on entry if the symbol was not found. */
 	bool SafeMatchSymbol(const TCHAR Match);
-	void HandleOneInheritedClass(FClasses& AllClasses, UClass* Class, FString InterfaceName);
+	void HandleOneInheritedClass(FClasses& AllClasses, UClass* Class, FString&& InterfaceName);
 	FClass* ParseClassNameDeclaration(FClasses& AllClasses, FString& DeclaredClassName, FString& RequiredAPIMacroIfPresent);
 
 	/** The property style of a variable declaration being parsed */
diff --git a/Engine/Source/Programs/UnrealHeaderTool/Private/ParserHelper.cpp b/Engine/Source/Programs/UnrealHeaderTool/Private/ParserHelper.cpp
--- a/Engine/Source/Programs/UnrealHeaderTool/Private/ParserHelper.cpp
+++ b/Engine/Source/Programs/UnrealHeaderTool/Private/ParserHelper.cpp
@@ -83,9 +83,9 @@ FTokenData* FClassMetaData::FindTokenData( FProperty* Prop )
 	return Result;
 }
 
-void FClassMetaData::AddInheritanceParent(const FString& InParent, FUnrealSourceFile* UnrealSourceFile)
+void FClassMetaData::AddInheritanceParent(FString&& InParent, FUnrealSourceFile* UnrealSourceFile)
 {
-	MultipleInheritanceParents.Add(new FMultipleInheritanceBaseClass(InParent));
+	MultipleInheritanceParents.Add(new FMultipleInheritanceBaseClass(MoveTemp(InParent)));
 }
 
 void FClassMetaData::AddInheritanceParent(UClass* ImplementedInterfaceClass, FUnrealSourceFile* UnrealSourceFile)
diff --git a/Engine/Source/Programs/UnrealHeaderTool/Private/ParserHelper.h b/Engine/Source/Programs/UnrealHeaderTool/Private/ParserHelper.h
--- a/Engine/Source/Programs/UnrealHeaderTool/Private/ParserHelper.h
+++ b/Engine/Source/Programs/UnrealHeaderTool/Private/ParserHelper.h
@@ -1458,17 +1458,18 @@ struct FMultipleInheritanceBaseClass
 	 * For multiple inheritance parents declared using 'Implements', corresponds to the UClass for the interface.  For multiple inheritance parents declared
 	 * using 'Inherits', this value will be NULL.
 	 */
-	UClass* InterfaceClass;
+	UClass* InterfaceClass = nullptr;
 
 	/**
 	 * Constructors
 	 */
-	FMultipleInheritanceBaseClass(const FString& BaseClassName)
-	: ClassName(BaseClassName), InterfaceClass(NULL)
-	{}
+	explicit FMultipleInheritanceBaseClass(FString&& BaseClassName)
+		: ClassName(MoveTemp(BaseClassName))
+	{
+	}
 
-	FMultipleInheritanceBaseClass(UClass* ImplementedInterfaceClass)
-	: InterfaceClass(ImplementedInterfaceClass)
+	explicit FMultipleInheritanceBaseClass(UClass* ImplementedInterfaceClass)
+		: InterfaceClass(ImplementedInterfaceClass)
 	{
 		ClassName = FString::Printf(TEXT("I%s"), *ImplementedInterfaceClass->GetName());
 	}
@@ -1719,7 +1720,7 @@ public:
 	 * @param Inparent The C++ class name to add to the multiple inheritance list
 	 * @param UnrealSourceFile Currently parsed source file.
 	 */
-	void AddInheritanceParent(const FString& InParent, FUnrealSourceFile* UnrealSourceFile);
+	void AddInheritanceParent(FString&& InParent, FUnrealSourceFile* UnrealSourceFile);
 
 	/**
 	 * Add a string to the list of inheritance parents for this class.
diff --git a/Engine/Source/Programs/UnrealHeaderTool/Resources/UHTDebugging/TestObject.h b/Engine/Source/Programs/UnrealHeaderTool/Resources/UHTDebugging/TestObject.h
--- a/Engine/Source/Programs/UnrealHeaderTool/Resources/UHTDebugging/TestObject.h
+++ b/Engine/Source/Programs/UnrealHeaderTool/Resources/UHTDebugging/TestObject.h
@@ -31,7 +31,7 @@ class alignas(8) UAlignedObject : public UObject
 };
 
 UCLASS()
-class UTestObject : public UObject, public ITestInterface
+class UTestObject : public UObject, public ITestInterface, public SomeNamespace::FSomeNonReflectedType
 {
 	GENERATED_BODY()
 
