Fixed UHT failing to discard inline keywords for void returns

GetVarType is the function that strips the inline and FORCEINLINE keywords, but that function call was being skipped for void return types. The result was that things like "inline void" failed to parse as a UHT return type, but "inline int32" worked fine.

diff --git a/Engine/Source/Programs/UnrealHeaderTool/Private/HeaderParser.cpp b/Engine/Source/Programs/UnrealHeaderTool/Private/HeaderParser.cpp
--- a/Engine/Source/Programs/UnrealHeaderTool/Private/HeaderParser.cpp
+++ b/Engine/Source/Programs/UnrealHeaderTool/Private/HeaderParser.cpp
@@ -4201,7 +4201,11 @@ void FHeaderParser::GetVarType(
 
 	RedirectTypeIdentifier(VarType);
 
-	if ( VarType.Matches(TEXT("int8"), ESearchCase::CaseSensitive) )
+	if ( VariableCategory == EVariableCategory::Return && VarType.Matches(TEXT("void"), ESearchCase::CaseSensitive) )
+	{
+		VarProperty = FPropertyBase(CPT_None);
+	}
+	else if ( VarType.Matches(TEXT("int8"), ESearchCase::CaseSensitive) )
 	{
 		VarProperty = FPropertyBase(CPT_Int8);
 	}
@@ -7562,11 +7566,8 @@ void FHeaderParser::CompileFunctionDeclaration(FClasses& AllClasses)
 	FToken ReturnType( CPT_None );
 
 	// C++ style functions always have a return value type, even if it's void
-	bool bHasReturnValue = !MatchIdentifier(TEXT("void"), ESearchCase::CaseSensitive);
-	if (bHasReturnValue)
-	{
-		GetVarType(AllClasses, GetCurrentScope(), ReturnType, CPF_None, nullptr, EPropertyDeclarationStyle::None, EVariableCategory::Return);
-	}
+	GetVarType(AllClasses, GetCurrentScope(), ReturnType, CPF_None, nullptr, EPropertyDeclarationStyle::None, EVariableCategory::Return);
+	bool bHasReturnValue = ReturnType.Type != CPT_None;
 
 	// Skip whitespaces to get InputPos exactly on beginning of function name.
 	while (FChar::IsWhitespace(PeekChar())) { GetChar(); }
