RigVM: Backing out UHT changes @ CL14544263, causing UHT failure in Frosty

diff --git a/Engine/Source/Programs/UnrealHeaderTool/Private/CodeGenerator.cpp b/Engine/Source/Programs/UnrealHeaderTool/Private/CodeGenerator.cpp
--- a/Engine/Source/Programs/UnrealHeaderTool/Private/CodeGenerator.cpp
+++ b/Engine/Source/Programs/UnrealHeaderTool/Private/CodeGenerator.cpp
@@ -3872,16 +3872,10 @@ void FNativeClassHeaderGenerator::ExportGeneratedStructBodyMacros(FOutputDevice&
 		// registry on construction of the static struct
 		if (StructRigVMInfo)
 		{
-			// All RigVM methods are in the Units folder
-			FString HeaderFilePath = SourceFile.GetFilename();
-			check(HeaderFilePath.Find(TEXT("Units")) > 0);
-			HeaderFilePath = HeaderFilePath.RightChop(HeaderFilePath.Find(TEXT("Units"))); 
-			HeaderFilePath.ReplaceInline(TEXT("\\"), TEXT("/"));
-
 			for (const FRigVMMethodInfo& MethodInfo : StructRigVMInfo->Methods)
 			{
-				Out.Logf(TEXT("\t\tFRigVMRegistry::Get().Register(TEXT(\"%s::%s\"), &%s::RigVM%s, Singleton, TEXT(\"%s\"));\r\n"),
-					*StructNameCPP, *MethodInfo.Name, *StructNameCPP, *MethodInfo.Name, *HeaderFilePath);
+				Out.Logf(TEXT("\t\tFRigVMRegistry::Get().Register(TEXT(\"%s::%s\"), &%s::RigVM%s, Singleton);\r\n"),
+					*StructNameCPP, *MethodInfo.Name, *StructNameCPP, *MethodInfo.Name);
 			}
 		}
 
