Fix missing code in the way ObjectPtr properties handle inline attributes.

diff --git a/Engine/Source/Programs/UnrealHeaderTool/Private/HeaderParser.cpp b/Engine/Source/Programs/UnrealHeaderTool/Private/HeaderParser.cpp
--- a/Engine/Source/Programs/UnrealHeaderTool/Private/HeaderParser.cpp
+++ b/Engine/Source/Programs/UnrealHeaderTool/Private/HeaderParser.cpp
@@ -874,6 +874,13 @@ namespace
 			{
 				check(VarProperty.PropertyClass);
 
+				// @TODO: OBJPTR: Further work is needed to handle UClass derivatives in a manner similar to CPT_ObjectReference or CPT_SoftObjectReference
+				if (DoesAnythingInHierarchyHaveDefaultToInstanced(VarProperty.PropertyClass))
+				{
+					VarProperty.PropertyFlags |= CPF_InstancedReference;
+					AddEditInlineMetaData(VarProperty.MetaData);
+				}
+
 				FObjectPtrProperty* Result = new FObjectPtrProperty(Scope, Name, ObjectFlags);
 				Result->PropertyClass = VarProperty.PropertyClass;
 				return Result;
