Fixing a race condition with -WRITEREF and -VERIFYREF that would cause files to be missing.

Removed some unused functionality in FScriptLocation that wasn't being used and depended on a global variable.

diff --git a/Engine/Source/Programs/UnrealHeaderTool/Private/CodeGenerator.cpp b/Engine/Source/Programs/UnrealHeaderTool/Private/CodeGenerator.cpp
--- a/Engine/Source/Programs/UnrealHeaderTool/Private/CodeGenerator.cpp
+++ b/Engine/Source/Programs/UnrealHeaderTool/Private/CodeGenerator.cpp
@@ -6361,6 +6361,7 @@ void FNativeClassHeaderGenerator::DeleteUnusedGeneratedHeaders(TSet<FString>&& P
  * exceptions. Needs to be fixed in future versions of UHT.
  */
 ECompilationResult::Type GCompilationResult = ECompilationResult::OtherCompilationError;
+FCriticalSection TestCommandLineCS;
 
 bool FNativeClassHeaderGenerator::SaveHeaderIfChanged(FReferenceGatherers& OutReferenceGatherers, const FPreloadHeaderFileInfo& FileInfo, FString&& InNewHeaderContents, FGraphEventRef& OutSaveTaskRef) const
 {
@@ -6373,30 +6374,33 @@ bool FNativeClassHeaderGenerator::SaveHeaderIfChanged(FReferenceGatherers& OutRe
 	static bool bTestedCmdLine = false;
 	if (!bTestedCmdLine)
 	{
-		bTestedCmdLine = true;
-
-		const FString& ProjectSavedDir = FPaths::ProjectSavedDir();
-
-		if (FParse::Param(FCommandLine::Get(), TEXT("WRITEREF")))
+		FScopeLock Lock(&TestCommandLineCS);
+		if (!bTestedCmdLine)
 		{
-			const FString ReferenceGeneratedCodePath = ProjectSavedDir / TEXT("ReferenceGeneratedCode/");
+			const FString& ProjectSavedDir = FPaths::ProjectSavedDir();
 
-			bWriteContents = true;
-			UE_LOG(LogCompile, Log, TEXT("********************************* Writing reference generated code to %s."), *ReferenceGeneratedCodePath);
-			UE_LOG(LogCompile, Log, TEXT("********************************* Deleting all files in ReferenceGeneratedCode."));
-			IFileManager::Get().DeleteDirectory(*ReferenceGeneratedCodePath, false, true);
-			IFileManager::Get().MakeDirectory(*ReferenceGeneratedCodePath);
-		}
-		else if (FParse::Param( FCommandLine::Get(), TEXT("VERIFYREF")))
-		{
-			const FString ReferenceGeneratedCodePath = ProjectSavedDir / TEXT("ReferenceGeneratedCode/");
-			const FString VerifyGeneratedCodePath = ProjectSavedDir / TEXT("VerifyGeneratedCode/");
+			if (FParse::Param(FCommandLine::Get(), TEXT("WRITEREF")))
+			{
+				const FString ReferenceGeneratedCodePath = ProjectSavedDir / TEXT("ReferenceGeneratedCode/");
 
-			bVerifyContents = true;
-			UE_LOG(LogCompile, Log, TEXT("********************************* Writing generated code to %s and comparing to %s"), *VerifyGeneratedCodePath, *ReferenceGeneratedCodePath);
-			UE_LOG(LogCompile, Log, TEXT("********************************* Deleting all files in VerifyGeneratedCode."));
-			IFileManager::Get().DeleteDirectory(*VerifyGeneratedCodePath, false, true);
-			IFileManager::Get().MakeDirectory(*VerifyGeneratedCodePath);
+				bWriteContents = true;
+				UE_LOG(LogCompile, Log, TEXT("********************************* Writing reference generated code to %s."), *ReferenceGeneratedCodePath);
+				UE_LOG(LogCompile, Log, TEXT("********************************* Deleting all files in ReferenceGeneratedCode."));
+				IFileManager::Get().DeleteDirectory(*ReferenceGeneratedCodePath, false, true);
+				IFileManager::Get().MakeDirectory(*ReferenceGeneratedCodePath);
+			}
+			else if (FParse::Param(FCommandLine::Get(), TEXT("VERIFYREF")))
+			{
+				const FString ReferenceGeneratedCodePath = ProjectSavedDir / TEXT("ReferenceGeneratedCode/");
+				const FString VerifyGeneratedCodePath = ProjectSavedDir / TEXT("VerifyGeneratedCode/");
+
+				bVerifyContents = true;
+				UE_LOG(LogCompile, Log, TEXT("********************************* Writing generated code to %s and comparing to %s"), *VerifyGeneratedCodePath, *ReferenceGeneratedCodePath);
+				UE_LOG(LogCompile, Log, TEXT("********************************* Deleting all files in VerifyGeneratedCode."));
+				IFileManager::Get().DeleteDirectory(*VerifyGeneratedCodePath, false, true);
+				IFileManager::Get().MakeDirectory(*VerifyGeneratedCodePath);
+			}
+			bTestedCmdLine = true;
 		}
 	}
 
@@ -6408,31 +6412,15 @@ bool FNativeClassHeaderGenerator::SaveHeaderIfChanged(FReferenceGatherers& OutRe
 
 		if (bWriteContents)
 		{
-			int32 i;
-			for (i = 0 ;i < 10; i++)
-			{
-				if (FFileHelper::SaveStringToFile(InNewHeaderContents, *Ref))
-				{
-					break;
-				}
-				FPlatformProcess::Sleep(1.0f); // I don't know why this fails after we delete the directory
-			}
-			check(i<10);
+			bool Written = FFileHelper::SaveStringToFile(InNewHeaderContents, *Ref);
+			check(Written);
 		}
 		else
 		{
 			const FString Verify = ProjectSavedDir / TEXT("VerifyGeneratedCode") / CleanFilename;
+			bool Written = FFileHelper::SaveStringToFile(InNewHeaderContents, *Verify);
+			check(Written);
 
-			int32 i;
-			for (i = 0 ;i < 10; i++)
-			{
-				if (FFileHelper::SaveStringToFile(InNewHeaderContents, *Verify))
-				{
-					break;
-				}
-				FPlatformProcess::Sleep(1.0f); // I don't know why this fails after we delete the directory
-			}
-			check(i<10);
 			FString RefHeader;
 			FString Message;
 			{
diff --git a/Engine/Source/Programs/UnrealHeaderTool/Private/HeaderParser.cpp b/Engine/Source/Programs/UnrealHeaderTool/Private/HeaderParser.cpp
--- a/Engine/Source/Programs/UnrealHeaderTool/Private/HeaderParser.cpp
+++ b/Engine/Source/Programs/UnrealHeaderTool/Private/HeaderParser.cpp
@@ -1353,19 +1353,6 @@ namespace
 	};
 }
 
-/////////////////////////////////////////////////////
-// FScriptLocation
-
-FHeaderParser* FScriptLocation::Compiler = NULL;
-
-FScriptLocation::FScriptLocation()
-{
-	if ( Compiler != NULL )
-	{
-		Compiler->InitScriptLocation(*this);
-	}
-}
-
 /////////////////////////////////////////////////////
 // FHeaderParser
 
@@ -1632,8 +1619,6 @@ UEnum* FHeaderParser::CompileEnum()
 	TArray<FPropertySpecifier> SpecifiersFound;
 	ReadSpecifierSetInsideMacro(SpecifiersFound, TEXT("Enum"), EnumToken.MetaData);
 
-	FScriptLocation DeclarationPosition;
-
 	// Check enum type. This can be global 'enum', 'namespace' or 'enum class' enums.
 	bool            bReadEnumName = false;
 	UEnum::ECppForm CppForm       = UEnum::ECppForm::Regular;
@@ -2417,8 +2402,6 @@ UScriptStruct* FHeaderParser::CompileStructDeclaration(FClasses& AllClasses)
 	// Make sure structs can be declared here.
 	CheckAllow( TEXT("'struct'"), ENestAllowFlags::TypeDecl );
 
-	FScriptLocation StructDeclaration;
-
 	bool IsNative = false;
 	bool IsExport = false;
 	bool IsTransient = false;
@@ -9002,8 +8985,6 @@ FHeaderParser::FHeaderParser(FFeedbackContext* InWarn, const FManifestModule& In
 		check(false);
 	}
 
-	FScriptLocation::Compiler = this;
-
 	static bool bConfigOptionsInitialized = false;
 
 	if (!bConfigOptionsInitialized)
diff --git a/Engine/Source/Programs/UnrealHeaderTool/Private/HeaderParser.h b/Engine/Source/Programs/UnrealHeaderTool/Private/HeaderParser.h
--- a/Engine/Source/Programs/UnrealHeaderTool/Private/HeaderParser.h
+++ b/Engine/Source/Programs/UnrealHeaderTool/Private/HeaderParser.h
@@ -695,14 +695,6 @@ protected:
 	// Constructor.
 	explicit FHeaderParser(FFeedbackContext* InWarn, const FManifestModule& InModule);
 
-	virtual ~FHeaderParser()
-	{
-		if ( FScriptLocation::Compiler == this )
-		{
-			FScriptLocation::Compiler = NULL;
-		}
-	}
-
 	// Returns true if the token is a dynamic delegate declaration
 	bool IsValidDelegateDeclaration(const FToken& Token) const;
 
diff --git a/Engine/Source/Programs/UnrealHeaderTool/Private/ParserHelper.h b/Engine/Source/Programs/UnrealHeaderTool/Private/ParserHelper.h
--- a/Engine/Source/Programs/UnrealHeaderTool/Private/ParserHelper.h
+++ b/Engine/Source/Programs/UnrealHeaderTool/Private/ParserHelper.h
@@ -1861,8 +1861,6 @@ private:
  */
 struct FScriptLocation
 {
-	static class FHeaderParser* Compiler;
-
 	/** the text buffer for the class associated with this retry point */
 	const TCHAR* Input;
 
@@ -1871,9 +1869,6 @@ struct FScriptLocation
 
 	/** the LineNumber of the compiler when this retry point was created */
 	int32 InputLine;
-
-	/** Constructor */
-	FScriptLocation();
 };
 
 /////////////////////////////////////////////////////
