Reduce preprocessing time by 60%.  Utilizing JobGraph to allow reading and initial parsing of files to be 100% concurrent.

diff --git a/Engine/Source/Programs/UnrealHeaderTool/Private/BaseParser.cpp b/Engine/Source/Programs/UnrealHeaderTool/Private/BaseParser.cpp
--- a/Engine/Source/Programs/UnrealHeaderTool/Private/BaseParser.cpp
+++ b/Engine/Source/Programs/UnrealHeaderTool/Private/BaseParser.cpp
@@ -79,56 +79,62 @@ namespace
 		EMetadataValueArgument::Type ValueArgument;
 	};
 
-	FMetadataKeyword* GetMetadataKeyword(const TCHAR* Keyword)
+	TMap<FString, FMetadataKeyword> MetadataDictionary;
+
+	void InitMetadataKeyswordsInternal()
 	{
-		static TMap<FString, FMetadataKeyword> Dictionary;
-		if (!Dictionary.Num())
+		check(MetadataDictionary.Num() == 0);
+		if (!MetadataDictionary.Num())
 		{
-			FMetadataKeyword& DisplayName = Dictionary.Add(TEXT("DisplayName"), EMetadataValueArgument::Required);
+			FMetadataKeyword& DisplayName = MetadataDictionary.Add(TEXT("DisplayName"), EMetadataValueArgument::Required);
 			DisplayName.InsertAddAction(TEXT("DisplayName"), TEXT(""));
 
-			FMetadataKeyword& FriendlyName = Dictionary.Add(TEXT("FriendlyName"), EMetadataValueArgument::Required);
+			FMetadataKeyword& FriendlyName = MetadataDictionary.Add(TEXT("FriendlyName"), EMetadataValueArgument::Required);
 			FriendlyName.InsertAddAction(TEXT("FriendlyName"), TEXT(""));
 
-			FMetadataKeyword& BlueprintInternalUseOnly = Dictionary.Add(TEXT("BlueprintInternalUseOnly"), EMetadataValueArgument::None);
+			FMetadataKeyword& BlueprintInternalUseOnly = MetadataDictionary.Add(TEXT("BlueprintInternalUseOnly"), EMetadataValueArgument::None);
 			BlueprintInternalUseOnly.InsertAddAction(TEXT("BlueprintInternalUseOnly"), TEXT("true"));
 			BlueprintInternalUseOnly.InsertAddAction(TEXT("BlueprintType"), TEXT("true"));
 
-			FMetadataKeyword& BlueprintType = Dictionary.Add(TEXT("BlueprintType"), EMetadataValueArgument::None);
+			FMetadataKeyword& BlueprintType = MetadataDictionary.Add(TEXT("BlueprintType"), EMetadataValueArgument::None);
 			BlueprintType.InsertAddAction(TEXT("BlueprintType"), TEXT("true"));
 
-			FMetadataKeyword& NotBlueprintType = Dictionary.Add(TEXT("NotBlueprintType"), EMetadataValueArgument::None);
+			FMetadataKeyword& NotBlueprintType = MetadataDictionary.Add(TEXT("NotBlueprintType"), EMetadataValueArgument::None);
 			NotBlueprintType.InsertAddAction(TEXT("NotBlueprintType"), TEXT("true"));
 			NotBlueprintType.InsertRemoveAction(TEXT("BlueprintType"));
 
-			FMetadataKeyword& Blueprintable = Dictionary.Add(TEXT("Blueprintable"), EMetadataValueArgument::None);
+			FMetadataKeyword& Blueprintable = MetadataDictionary.Add(TEXT("Blueprintable"), EMetadataValueArgument::None);
 			Blueprintable.InsertAddAction(TEXT("IsBlueprintBase"), TEXT("true"));
-			Blueprintable.InsertAddAction(TEXT("BlueprintType"),   TEXT("true"));
+			Blueprintable.InsertAddAction(TEXT("BlueprintType"), TEXT("true"));
 
-			FMetadataKeyword& CallInEditor = Dictionary.Add(TEXT("CallInEditor"), EMetadataValueArgument::None);
+			FMetadataKeyword& CallInEditor = MetadataDictionary.Add(TEXT("CallInEditor"), EMetadataValueArgument::None);
 			CallInEditor.InsertAddAction(TEXT("CallInEditor"), TEXT("true"));
 
-			FMetadataKeyword& NotBlueprintable = Dictionary.Add(TEXT("NotBlueprintable"), EMetadataValueArgument::None);
-			NotBlueprintable.InsertAddAction   (TEXT("IsBlueprintBase"), TEXT("false"));
+			FMetadataKeyword& NotBlueprintable = MetadataDictionary.Add(TEXT("NotBlueprintable"), EMetadataValueArgument::None);
+			NotBlueprintable.InsertAddAction(TEXT("IsBlueprintBase"), TEXT("false"));
 			NotBlueprintable.InsertRemoveAction(TEXT("BlueprintType"));
 
-			FMetadataKeyword& Category = Dictionary.Add(TEXT("Category"), EMetadataValueArgument::Required);
+			FMetadataKeyword& Category = MetadataDictionary.Add(TEXT("Category"), EMetadataValueArgument::Required);
 			Category.InsertAddAction(TEXT("Category"), TEXT(""));
 
-			FMetadataKeyword& ExperimentalFeature = Dictionary.Add(TEXT("Experimental"), EMetadataValueArgument::None);
+			FMetadataKeyword& ExperimentalFeature = MetadataDictionary.Add(TEXT("Experimental"), EMetadataValueArgument::None);
 			ExperimentalFeature.InsertAddAction(TEXT("DevelopmentStatus"), TEXT("Experimental"));
 
-			FMetadataKeyword& EarlyAccessFeature = Dictionary.Add(TEXT("EarlyAccessPreview"), EMetadataValueArgument::None);
+			FMetadataKeyword& EarlyAccessFeature = MetadataDictionary.Add(TEXT("EarlyAccessPreview"), EMetadataValueArgument::None);
 			EarlyAccessFeature.InsertAddAction(TEXT("DevelopmentStatus"), TEXT("EarlyAccess"));
 
-			FMetadataKeyword& DocumentationPolicy = Dictionary.Add(TEXT("DocumentationPolicy"), EMetadataValueArgument::None);
+			FMetadataKeyword& DocumentationPolicy = MetadataDictionary.Add(TEXT("DocumentationPolicy"), EMetadataValueArgument::None);
 			DocumentationPolicy.InsertAddAction(TEXT("DocumentationPolicy"), TEXT("Strict"));
 
-			FMetadataKeyword& SparseClassDataType = Dictionary.Add(TEXT("SparseClassDataType"), EMetadataValueArgument::Required);
+			FMetadataKeyword& SparseClassDataType = MetadataDictionary.Add(TEXT("SparseClassDataType"), EMetadataValueArgument::Required);
 			SparseClassDataType.InsertAddAction(TEXT("SparseClassDataType"), TEXT(""));
 		}
+	}
 
-		return Dictionary.Find(Keyword);
+	FMetadataKeyword* GetMetadataKeyword(const TCHAR* Keyword)
+	{
+		check(MetadataDictionary.Num() > 0);
+		return MetadataDictionary.Find(Keyword);
 	}
 }
 
@@ -1217,4 +1223,7 @@ void FBaseParser::InsertMetaDataPair(TMap<FName, FString>& MetaData, FName KeyNa
 	MetaData.Add(KeyName, MoveTemp(Value));
 }
 
-
+void FBaseParser::InitMetadataKeywords()
+{
+	InitMetadataKeyswordsInternal();
+}
\ No newline at end of file
diff --git a/Engine/Source/Programs/UnrealHeaderTool/Private/BaseParser.h b/Engine/Source/Programs/UnrealHeaderTool/Private/BaseParser.h
--- a/Engine/Source/Programs/UnrealHeaderTool/Private/BaseParser.h
+++ b/Engine/Source/Programs/UnrealHeaderTool/Private/BaseParser.h
@@ -187,4 +187,7 @@ public:
 	static void InsertMetaDataPair(TMap<FName, FString>& MetaData, FName InKey, FString InValue);
 
 	//////////////
+
+	// Initialize the metadata keywords prior to parsing
+	static void InitMetadataKeywords();
 };
diff --git a/Engine/Source/Programs/UnrealHeaderTool/Private/CodeGenerator.cpp b/Engine/Source/Programs/UnrealHeaderTool/Private/CodeGenerator.cpp
--- a/Engine/Source/Programs/UnrealHeaderTool/Private/CodeGenerator.cpp
+++ b/Engine/Source/Programs/UnrealHeaderTool/Private/CodeGenerator.cpp
@@ -83,6 +83,16 @@ struct FPerHeaderData
 	TArray<FSimplifiedParsingClassInfo> ParsedClassArray;
 };
 
+struct FManifestPerHeaderData
+{
+	TArray<FString> PublicUObjectClassesHeadersText;
+	TArray<FPerHeaderData> PublicUObjectClassesHeadersData;
+	TArray<FString> PublicUObjectHeadersText;
+	TArray<FPerHeaderData> PublicUObjectHeadersData;
+	TArray<FString> PrivateUObjectHeadersText;
+	TArray<FPerHeaderData> PrivateUObjectHeadersData;
+};
+
 static void PerformSimplifiedClassParse(UPackage* InParent, const TCHAR* FileName, const TCHAR* Buffer, FPerHeaderData& PerHeaderData);
 static void ProcessInitialClassParse(FPerHeaderData& PerHeaderData);
 
@@ -6702,6 +6712,68 @@ void ResolveSuperClasses(UPackage* Package)
 	}
 }
 
+UPackage* GetModulePackage(FManifestModule& Module)
+{
+	UPackage* Package = Cast<UPackage>(StaticFindObjectFast(UPackage::StaticClass(), NULL, FName(*Module.LongPackageName), false, false));
+	if (Package == NULL)
+	{
+		Package = CreatePackage(*Module.LongPackageName);
+	}
+	// Set some package flags for indicating that this package contains script
+	// NOTE: We do this even if we didn't have to create the package, because CoreUObject is compiled into UnrealHeaderTool and we still
+	//       want to make sure our flags get set
+	Package->SetPackageFlags(PKG_ContainsScript | PKG_Compiling);
+	Package->ClearPackageFlags(PKG_ClientOptional | PKG_ServerSideOnly);
+
+	if (Module.OverrideModuleType == EPackageOverrideType::None)
+	{
+		switch (Module.ModuleType)
+		{
+		case EBuildModuleType::GameEditor:
+		case EBuildModuleType::EngineEditor:
+			Package->SetPackageFlags(PKG_EditorOnly);
+			break;
+
+		case EBuildModuleType::GameDeveloper:
+		case EBuildModuleType::EngineDeveloper:
+			Package->SetPackageFlags(PKG_Developer);
+			break;
+
+		case EBuildModuleType::GameUncooked:
+		case EBuildModuleType::EngineUncooked:
+			Package->SetPackageFlags(PKG_UncookedOnly);
+			break;
+		}
+	}
+	else
+	{
+		// If the user has specified this module to have another package flag, then OR it on
+		switch (Module.OverrideModuleType)
+		{
+		case EPackageOverrideType::EditorOnly:
+			Package->SetPackageFlags(PKG_EditorOnly);
+			break;
+
+		case EPackageOverrideType::EngineDeveloper:
+		case EPackageOverrideType::GameDeveloper:
+			Package->SetPackageFlags(PKG_Developer);
+			break;
+
+		case EPackageOverrideType::EngineUncookedOnly:
+		case EPackageOverrideType::GameUncookedOnly:
+			Package->SetPackageFlags(PKG_UncookedOnly);
+			break;
+		}
+	}
+
+	// Add new module or overwrite whatever we had loaded, that data is obsolete.
+	GPackageToManifestModuleMap.Add(Package, &Module);
+
+	return Package;
+}
+
+//FGraphEventRef CreateModuleParse
+
 ECompilationResult::Type PreparseModules(const FString& ModuleInfoPath, int32& NumFailures)
 {
 	// Three passes.  1) Public 'Classes' headers (legacy)  2) Public headers   3) Private headers
@@ -6714,6 +6786,9 @@ ECompilationResult::Type PreparseModules(const FString& ModuleInfoPath, int32& N
 		FolderType_Count
 	};
 
+	// The meta data keywords must be initialized prior to going wide
+	FBaseParser::InitMetadataKeywords();
+
 	ECompilationResult::Type Result = ECompilationResult::Succeeded;
 
 #if !PLATFORM_EXCEPTIONS_DISABLED
@@ -6749,291 +6824,312 @@ ECompilationResult::Type PreparseModules(const FString& ModuleInfoPath, int32& N
 	};
 #endif
 
+	// Pre-initialization
+	int32 NumHeaders = 0;
 	for (FManifestModule& Module : GManifest.Modules)
 	{
-		if (Result != ECompilationResult::Succeeded)
-		{
-			break;
-		}
-
 		// Force regeneration of all subsequent modules, otherwise data will get corrupted.
 		Module.ForceRegeneration();
+		GetModulePackage(Module);
 
-		UPackage* Package = Cast<UPackage>(StaticFindObjectFast(UPackage::StaticClass(), NULL, FName(*Module.LongPackageName), false, false));
-		if (Package == NULL)
-		{
-			Package = CreatePackage(*Module.LongPackageName);
-		}
-		// Set some package flags for indicating that this package contains script
-		// NOTE: We do this even if we didn't have to create the package, because CoreUObject is compiled into UnrealHeaderTool and we still
-		//       want to make sure our flags get set
-		Package->SetPackageFlags(PKG_ContainsScript | PKG_Compiling);
-		Package->ClearPackageFlags(PKG_ClientOptional | PKG_ServerSideOnly);
-		
-		if(Module.OverrideModuleType == EPackageOverrideType::None)
+		// Count the headers
+		for (int32 PassIndex = 0; PassIndex < FolderType_Count && Result == ECompilationResult::Succeeded; ++PassIndex)
 		{
-			switch (Module.ModuleType)
-			{
-			case EBuildModuleType::GameEditor:
-			case EBuildModuleType::EngineEditor:
-				Package->SetPackageFlags(PKG_EditorOnly);
-				break;
-
-			case EBuildModuleType::GameDeveloper:
-			case EBuildModuleType::EngineDeveloper:
-				Package->SetPackageFlags(PKG_Developer);
-				break;
-
-			case EBuildModuleType::GameUncooked:
-			case EBuildModuleType::EngineUncooked:
-				Package->SetPackageFlags(PKG_UncookedOnly);
-				break;
-			}
+			EHeaderFolderTypes CurrentlyProcessing = (EHeaderFolderTypes)PassIndex;
+			const TArray<FString>& UObjectHeaders =
+				(CurrentlyProcessing == PublicClassesHeaders) ? Module.PublicUObjectClassesHeaders :
+				(CurrentlyProcessing == PublicHeaders) ? Module.PublicUObjectHeaders :
+				Module.PrivateUObjectHeaders;
+			NumHeaders += UObjectHeaders.Num();
 		}
-		else
-		{
-			// If the user has specified this module to have another package flag, then OR it on
-			switch (Module.OverrideModuleType)
-			{
-			case EPackageOverrideType::EditorOnly:
-				Package->SetPackageFlags(PKG_EditorOnly);
-				break;
-
-			case EPackageOverrideType::EngineDeveloper:
-			case EPackageOverrideType::GameDeveloper:
-				Package->SetPackageFlags(PKG_Developer);
-				break;
+	}
 
-			case EPackageOverrideType::EngineUncookedOnly:
-			case EPackageOverrideType::GameUncookedOnly:
-				Package->SetPackageFlags(PKG_UncookedOnly);
-				break;
-			}
-		}
+	TArray<FManifestPerHeaderData> ManifestPerHeaderDatas;
+	ManifestPerHeaderDatas.SetNum(GManifest.Modules.Num());
 
-		// Add new module or overwrite whatever we had loaded, that data is obsolete.
-		GPackageToManifestModuleMap.Add(Package, &Module);
+	// Pre-parse
+	{
+		double PreparseTime = 0.0;
+		FDurationTimer PreparseTimer(PreparseTime);
+		PreparseTimer.Start();
 
-		double ThisModulePreparseTime = 0.0;
-		int32 NumHeadersPreparsed = 0;
-		FDurationTimer ThisModuleTimer(ThisModulePreparseTime);
-		ThisModuleTimer.Start();
+		// Create the load tasks
+		FGraphEventRef LastClassParseTask;
+		FGraphEventArray ClassParseTaskPreReqs;
+		ClassParseTaskPreReqs.SetNum(2);
 
-		// Pre-parse the headers
-		for (int32 PassIndex = 0; PassIndex < FolderType_Count && Result == ECompilationResult::Succeeded; ++PassIndex)
+		FGraphEventArray LoadTasks;
+		LoadTasks.SetNum(NumHeaders);
+		for (int32 ModuleIndex = 0, NumModules = GManifest.Modules.Num(); ModuleIndex < NumModules; ++ModuleIndex)
 		{
-			EHeaderFolderTypes CurrentlyProcessing = (EHeaderFolderTypes)PassIndex;
+			FManifestModule& Module = GManifest.Modules[ModuleIndex];
+			FManifestPerHeaderData& ModulePerHeaderData = ManifestPerHeaderDatas[ModuleIndex];
+			UPackage* Package = GetModulePackage(Module);
 
-			// We'll make an ordered list of all UObject headers we care about.
-			// @todo uht: Ideally 'dependson' would not be allowed from public -> private, or NOT at all for new style headers
-			const TArray<FString>& UObjectHeaders =
-				(CurrentlyProcessing == PublicClassesHeaders) ? Module.PublicUObjectClassesHeaders :
-				(CurrentlyProcessing == PublicHeaders       ) ? Module.PublicUObjectHeaders        :
-				                                                Module.PrivateUObjectHeaders;
-			if (!UObjectHeaders.Num())
+			// Pre-parse the headers
+			for (int32 PassIndex = 0; PassIndex < FolderType_Count && Result == ECompilationResult::Succeeded; ++PassIndex)
 			{
-				continue;
-			}
+				EHeaderFolderTypes CurrentlyProcessing = (EHeaderFolderTypes)PassIndex;
 
-			NumHeadersPreparsed += UObjectHeaders.Num();
+				// We'll make an ordered list of all UObject headers we care about.
+				// @todo uht: Ideally 'dependson' would not be allowed from public -> private, or NOT at all for new style headers
+				const TArray<FString>& UObjectHeaders =
+					(CurrentlyProcessing == PublicClassesHeaders) ? Module.PublicUObjectClassesHeaders :
+					(CurrentlyProcessing == PublicHeaders       ) ? Module.PublicUObjectHeaders :
+					                                                Module.PrivateUObjectHeaders;
+				if (!UObjectHeaders.Num())
+				{
+					continue;
+				}
 
-			TArray<FString> HeaderFiles;
-			HeaderFiles.SetNum(UObjectHeaders.Num());
+				TArray<FString>& HeaderFiles =
+					(CurrentlyProcessing == PublicClassesHeaders) ? ModulePerHeaderData.PublicUObjectClassesHeadersText :
+					(CurrentlyProcessing == PublicHeaders       ) ? ModulePerHeaderData.PublicUObjectHeadersText :
+					                                                ModulePerHeaderData.PrivateUObjectHeadersText;
+				HeaderFiles.SetNum(UObjectHeaders.Num());
 
-			{
-				SCOPE_SECONDS_COUNTER_UHT(LoadHeaderContentFromFile);
-				ParallelFor(UObjectHeaders.Num(), [&](int32 Index)
+				TArray<FPerHeaderData>& PerHeaderData =
+					(CurrentlyProcessing == PublicClassesHeaders) ? ModulePerHeaderData.PublicUObjectClassesHeadersData :
+					(CurrentlyProcessing == PublicHeaders       ) ? ModulePerHeaderData.PublicUObjectHeadersData :
+					                                                ModulePerHeaderData.PrivateUObjectHeadersData;
+				PerHeaderData.SetNum(UObjectHeaders.Num());
+
+				for (int32 Index = 0, EIndex = UObjectHeaders.Num(); Index < EIndex; ++Index)
 				{
-					const FString& RawFilename = UObjectHeaders[Index];
 
-#if !PLATFORM_EXCEPTIONS_DISABLED
-					try
-#endif
+					// Phase #1: Load the file
+					auto LoadLambda = [&UObjectHeaders, &HeaderFiles, &LogException, &ModuleInfoPath, &Result, Index]()
 					{
-						const FString FullFilename = FPaths::ConvertRelativePathToFull(ModuleInfoPath, RawFilename);
-
-						if (!FFileHelper::LoadFileToString(HeaderFiles[Index], *FullFilename))
+						if (Result != ECompilationResult::Succeeded)
 						{
-							FError::Throwf(TEXT("UnrealHeaderTool was unable to load source file '%s'"), *FullFilename);
+							return;
 						}
-					}
-#if !PLATFORM_EXCEPTIONS_DISABLED
-					catch (TCHAR* ErrorMsg)
-					{
-						FString AbsFilename = IFileManager::Get().ConvertToAbsolutePathForExternalAppForRead(*RawFilename);
-						LogException(MoveTemp(AbsFilename), 1, ErrorMsg);
-					}				
-#endif
-				});
-			}
-
-#if !PLATFORM_EXCEPTIONS_DISABLED
-			FTaskGraphInterface::Get().WaitUntilTasksComplete(ExceptionTasks);
-#endif
-
-			if (Result != ECompilationResult::Succeeded)
-			{
-				continue;
-			}
 
-			TArray<FPerHeaderData> PerHeaderData;
-			PerHeaderData.SetNum(UObjectHeaders.Num());
-
-			ParallelFor(UObjectHeaders.Num(), [&](int32 Index)
-			{
-				const FString& RawFilename = UObjectHeaders[Index];
+						const FString& RawFilename = UObjectHeaders[Index];
 
 #if !PLATFORM_EXCEPTIONS_DISABLED
-				try
+						try
 #endif
-				{
-					PerformSimplifiedClassParse(Package, *RawFilename, *HeaderFiles[Index], PerHeaderData[Index]);
-				}
-#if !PLATFORM_EXCEPTIONS_DISABLED
-				catch (const FFileLineException& Ex)
-				{
-					FString AbsFilename = IFileManager::Get().ConvertToAbsolutePathForExternalAppForRead(*Ex.Filename);
-					LogException(MoveTemp(AbsFilename), Ex.Line, Ex.Message);
-				}
-				catch (TCHAR* ErrorMsg)
-				{
-					FString AbsFilename = IFileManager::Get().ConvertToAbsolutePathForExternalAppForRead(*RawFilename);
-					LogException(MoveTemp(AbsFilename), 1, ErrorMsg);
-				}
-#endif
-			});
+						{
+							const FString FullFilename = FPaths::ConvertRelativePathToFull(ModuleInfoPath, RawFilename);
 
+							if (!FFileHelper::LoadFileToString(HeaderFiles[Index], *FullFilename))
+							{
+								FError::Throwf(TEXT("UnrealHeaderTool was unable to load source file '%s'"), *FullFilename);
+							}
+						}
 #if !PLATFORM_EXCEPTIONS_DISABLED
-			FTaskGraphInterface::Get().WaitUntilTasksComplete(ExceptionTasks);
+						catch (TCHAR* ErrorMsg)
+						{
+							FString AbsFilename = IFileManager::Get().ConvertToAbsolutePathForExternalAppForRead(*RawFilename);
+							LogException(MoveTemp(AbsFilename), 1, ErrorMsg);
+						}
 #endif
+					};
 
-			if (Result != ECompilationResult::Succeeded)
-			{
-				continue;
-			}
+					// Phase #2: Perform simplified class parse (can run concurrenrtly)
+					auto PreProcessLambda = [Package, &UObjectHeaders, &HeaderFiles, &PerHeaderData, &LogException, &ModuleInfoPath, &Result, Index]()
+					{
+						if (Result != ECompilationResult::Succeeded)
+						{
+							return;
+						}
 
-			for (int32 Index = 0; Index < UObjectHeaders.Num(); ++Index)
-			{
-				const FString& RawFilename = UObjectHeaders[Index];
+						const FString& RawFilename = UObjectHeaders[Index];
 
 #if !PLATFORM_EXCEPTIONS_DISABLED
-				try
+						try
 #endif
-				{
-					// Import class.
-					const FString FullFilename = FPaths::ConvertRelativePathToFull(ModuleInfoPath, RawFilename);
-
-					ProcessInitialClassParse(PerHeaderData[Index]);
-					TSharedRef<FUnrealSourceFile> UnrealSourceFile = PerHeaderData[Index].UnrealSourceFile.ToSharedRef();
-					FUnrealSourceFile* UnrealSourceFilePtr = &UnrealSourceFile.Get();
-					FString CleanFilename     = FPaths::GetCleanFilename(RawFilename);
-					uint32  CleanFilenameHash = GetTypeHash(CleanFilename);
-					if (const TSharedRef<FUnrealSourceFile>* ExistingSourceFile = GUnrealSourceFilesMap.FindByHash(CleanFilenameHash, CleanFilename))
-					{
-						FString NormalizedFullFilename     = FullFilename;
-						FString NormalizedExistingFilename = (*ExistingSourceFile)->GetFilename();
-
-						FPaths::NormalizeFilename(NormalizedFullFilename);
-						FPaths::NormalizeFilename(NormalizedExistingFilename);
-
-						if (NormalizedFullFilename != NormalizedExistingFilename)
 						{
-							FError::Throwf(TEXT("Duplicate leaf header name found: %s (original: %s)"), *NormalizedFullFilename, *NormalizedExistingFilename);
+							PerformSimplifiedClassParse(Package, *RawFilename, *HeaderFiles[Index], PerHeaderData[Index]);
+							HeaderFiles[Index].Empty();
 						}
-					}
-					GUnrealSourceFilesMap.AddByHash(CleanFilenameHash, MoveTemp(CleanFilename), UnrealSourceFile);
-
-					if (CurrentlyProcessing == PublicClassesHeaders)
-					{
-						GPublicSourceFileSet.Add(UnrealSourceFilePtr);
-					}
-
-					// Save metadata for the class path, both for it's include path and relative to the module base directory
-					if (FullFilename.StartsWith(Module.BaseDirectory))
-					{
-						// Get the path relative to the module directory
-						const TCHAR* ModuleRelativePath = *FullFilename + Module.BaseDirectory.Len();
-
-						UnrealSourceFilePtr->SetModuleRelativePath(ModuleRelativePath);
-
-						// Calculate the include path
-						const TCHAR* IncludePath = ModuleRelativePath;
-
-						// Walk over the first potential slash
-						if (*IncludePath == TEXT('/'))
+#if !PLATFORM_EXCEPTIONS_DISABLED
+						catch (const FFileLineException& Ex)
 						{
-							IncludePath++;
+							FString AbsFilename = IFileManager::Get().ConvertToAbsolutePathForExternalAppForRead(*Ex.Filename);
+							LogException(MoveTemp(AbsFilename), Ex.Line, Ex.Message);
 						}
+						catch (TCHAR* ErrorMsg)
+						{
+							FString AbsFilename = IFileManager::Get().ConvertToAbsolutePathForExternalAppForRead(*RawFilename);
+							LogException(MoveTemp(AbsFilename), 1, ErrorMsg);
+						}
+#endif
+					};
 
-						// Does this module path start with a known include path location? If so, we can cut that part out of the include path
-						static const TCHAR PublicFolderName[]  = TEXT("Public/");
-						static const TCHAR PrivateFolderName[] = TEXT("Private/");
-						static const TCHAR ClassesFolderName[] = TEXT("Classes/");
-						if (FCString::Strnicmp(IncludePath, PublicFolderName, UE_ARRAY_COUNT(PublicFolderName) - 1) == 0)
+					// Phase #3: Perform initial class parse, can not run until previous header is done
+					auto ClassParseLambda = [&Module, Package, CurrentlyProcessing, &UObjectHeaders, &PerHeaderData, &LogException, &ModuleInfoPath, &Result, Index]()
+					{
+						if (Result != ECompilationResult::Succeeded)
 						{
-							IncludePath += (UE_ARRAY_COUNT(PublicFolderName) - 1);
+							return;
 						}
-						else if (FCString::Strnicmp(IncludePath, PrivateFolderName, UE_ARRAY_COUNT(PrivateFolderName) - 1) == 0)
+
+						const FString& RawFilename = UObjectHeaders[Index];
+
+#if !PLATFORM_EXCEPTIONS_DISABLED
+						try
+#endif
 						{
-							IncludePath += (UE_ARRAY_COUNT(PrivateFolderName) - 1);
+							// Import class.
+							const FString FullFilename = FPaths::ConvertRelativePathToFull(ModuleInfoPath, RawFilename);
+
+							ProcessInitialClassParse(PerHeaderData[Index]);
+							TSharedRef<FUnrealSourceFile> UnrealSourceFile = PerHeaderData[Index].UnrealSourceFile.ToSharedRef();
+							FUnrealSourceFile* UnrealSourceFilePtr = &UnrealSourceFile.Get();
+							FString CleanFilename = FPaths::GetCleanFilename(RawFilename);
+							uint32  CleanFilenameHash = GetTypeHash(CleanFilename);
+							if (const TSharedRef<FUnrealSourceFile>* ExistingSourceFile = GUnrealSourceFilesMap.FindByHash(CleanFilenameHash, CleanFilename))
+							{
+								FString NormalizedFullFilename = FullFilename;
+								FString NormalizedExistingFilename = (*ExistingSourceFile)->GetFilename();
+
+								FPaths::NormalizeFilename(NormalizedFullFilename);
+								FPaths::NormalizeFilename(NormalizedExistingFilename);
+
+								if (NormalizedFullFilename != NormalizedExistingFilename)
+								{
+									FError::Throwf(TEXT("Duplicate leaf header name found: %s (original: %s)"), *NormalizedFullFilename, *NormalizedExistingFilename);
+								}
+							}
+							GUnrealSourceFilesMap.AddByHash(CleanFilenameHash, MoveTemp(CleanFilename), UnrealSourceFile);
+
+							if (CurrentlyProcessing == PublicClassesHeaders)
+							{
+								GPublicSourceFileSet.Add(UnrealSourceFilePtr);
+							}
+
+							// Save metadata for the class path, both for it's include path and relative to the module base directory
+							if (FullFilename.StartsWith(Module.BaseDirectory))
+							{
+								// Get the path relative to the module directory
+								const TCHAR* ModuleRelativePath = *FullFilename + Module.BaseDirectory.Len();
+
+								UnrealSourceFilePtr->SetModuleRelativePath(ModuleRelativePath);
+
+								// Calculate the include path
+								const TCHAR* IncludePath = ModuleRelativePath;
+
+								// Walk over the first potential slash
+								if (*IncludePath == TEXT('/'))
+								{
+									IncludePath++;
+								}
+
+								// Does this module path start with a known include path location? If so, we can cut that part out of the include path
+								static const TCHAR PublicFolderName[] = TEXT("Public/");
+								static const TCHAR PrivateFolderName[] = TEXT("Private/");
+								static const TCHAR ClassesFolderName[] = TEXT("Classes/");
+								if (FCString::Strnicmp(IncludePath, PublicFolderName, UE_ARRAY_COUNT(PublicFolderName) - 1) == 0)
+								{
+									IncludePath += (UE_ARRAY_COUNT(PublicFolderName) - 1);
+								}
+								else if (FCString::Strnicmp(IncludePath, PrivateFolderName, UE_ARRAY_COUNT(PrivateFolderName) - 1) == 0)
+								{
+									IncludePath += (UE_ARRAY_COUNT(PrivateFolderName) - 1);
+								}
+								else if (FCString::Strnicmp(IncludePath, ClassesFolderName, UE_ARRAY_COUNT(ClassesFolderName) - 1) == 0)
+								{
+									IncludePath += (UE_ARRAY_COUNT(ClassesFolderName) - 1);
+								}
+
+								// Add the include path
+								if (*IncludePath != 0)
+								{
+									UnrealSourceFilePtr->SetIncludePath(MoveTemp(IncludePath));
+								}
+							}
 						}
-						else if (FCString::Strnicmp(IncludePath, ClassesFolderName, UE_ARRAY_COUNT(ClassesFolderName) - 1) == 0)
+#if !PLATFORM_EXCEPTIONS_DISABLED
+						catch (const FFileLineException& Ex)
 						{
-							IncludePath += (UE_ARRAY_COUNT(ClassesFolderName) - 1);
+							FString AbsFilename = IFileManager::Get().ConvertToAbsolutePathForExternalAppForRead(*Ex.Filename);
+							LogException(MoveTemp(AbsFilename), Ex.Line, Ex.Message);
 						}
-
-						// Add the include path
-						if (*IncludePath != 0)
+						catch (TCHAR* ErrorMsg)
 						{
-							UnrealSourceFilePtr->SetIncludePath(MoveTemp(IncludePath));
+							FString AbsFilename = IFileManager::Get().ConvertToAbsolutePathForExternalAppForRead(*RawFilename);
+							LogException(MoveTemp(AbsFilename), 1, ErrorMsg);
 						}
+#endif
+					};
+
+					// Only the first ClassParseTask uses only the pre process task.  All others must wait on the previous 
+					// ClassParseTask.  In the future we should modify the code to not require this.
+					FGraphEventRef LoadTask = FFunctionGraphTask::CreateAndDispatchWhenReady(MoveTemp(LoadLambda), TStatId());
+					FGraphEventRef PreProcessTask = FFunctionGraphTask::CreateAndDispatchWhenReady(MoveTemp(PreProcessLambda), TStatId(), LoadTask);
+					FGraphEventRef ClassParseTask;
+					if (LastClassParseTask.IsValid())
+					{
+						ClassParseTaskPreReqs[0] = LastClassParseTask;
+						ClassParseTaskPreReqs[1] = PreProcessTask;
+						ClassParseTask = FFunctionGraphTask::CreateAndDispatchWhenReady(MoveTemp(ClassParseLambda), TStatId(), &ClassParseTaskPreReqs);
 					}
+					else
+					{
+						ClassParseTask = FFunctionGraphTask::CreateAndDispatchWhenReady(MoveTemp(ClassParseLambda), TStatId(), PreProcessTask);
+					}
+					LastClassParseTask = ClassParseTask;
+
+					LoadTasks.Add(MoveTemp(ClassParseTask));
 				}
-#if !PLATFORM_EXCEPTIONS_DISABLED
-				catch (const FFileLineException& Ex)
-				{
-					FString AbsFilename = IFileManager::Get().ConvertToAbsolutePathForExternalAppForRead(*Ex.Filename);
-					LogException(MoveTemp(AbsFilename), Ex.Line, Ex.Message);
-				}
-				catch (TCHAR* ErrorMsg)
-				{
-					FString AbsFilename = IFileManager::Get().ConvertToAbsolutePathForExternalAppForRead(*RawFilename);
-					LogException(MoveTemp(AbsFilename), 1, ErrorMsg);
-				}
-#endif
-			}
-			if (Result == ECompilationResult::Succeeded && NumFailures != 0)
-			{
-				Result = ECompilationResult::OtherCompilationError;
 			}
 		}
 
-		// Don't resolve superclasses for module when loading from makefile.
-		// Data is only partially loaded at this point.
+		// Wait for all the loading and preparsing to complete
+		FTaskGraphInterface::Get().WaitUntilTasksComplete(LoadTasks);
+
 #if !PLATFORM_EXCEPTIONS_DISABLED
-		try
+		FTaskGraphInterface::Get().WaitUntilTasksComplete(ExceptionTasks);
 #endif
+
+		PreparseTimer.Stop();
+		UE_LOG(LogCompile, Log, TEXT("Loaded and preparsed classes in %d modules containing %d files(s) in %.3f secs."), GManifest.Modules.Num(), NumHeaders, PreparseTime);
+	}
+
+	// Perform the second phase of loading
+	if (Result == ECompilationResult::Succeeded)
+	{
+		double ResolveTime = 0.0;
+		FDurationTimer ResolveTimer(ResolveTime);
+		ResolveTimer.Start();
+	
+		for (FManifestModule& Module : GManifest.Modules)
 		{
-			ResolveSuperClasses(Package);
-		}
+			UPackage* Package = GetModulePackage(Module);
+
+			// Don't resolve superclasses for module when loading from makefile.
+			// Data is only partially loaded at this point.
 #if !PLATFORM_EXCEPTIONS_DISABLED
-		catch (TCHAR* ErrorMsg)
-		{
-			TGuardValue<ELogTimes::Type> DisableLogTimes(GPrintLogTimes, ELogTimes::None);
+			try
+#endif
+			{
+				ResolveSuperClasses(Package);
+			}
+#if !PLATFORM_EXCEPTIONS_DISABLED
+			catch (TCHAR* ErrorMsg)
+			{
+				TGuardValue<ELogTimes::Type> DisableLogTimes(GPrintLogTimes, ELogTimes::None);
 
-			FString FormattedErrorMessage = FString::Printf(TEXT("Error: %s\r\n"), ErrorMsg);
+				FString FormattedErrorMessage = FString::Printf(TEXT("Error: %s\r\n"), ErrorMsg);
 
-			Result = GCompilationResult;
+				Result = GCompilationResult;
 
-			UE_LOG(LogCompile, Log, TEXT("%s"), *FormattedErrorMessage);
-			GWarn->Log(ELogVerbosity::Error, FormattedErrorMessage);
+				UE_LOG(LogCompile, Log, TEXT("%s"), *FormattedErrorMessage);
+				GWarn->Log(ELogVerbosity::Error, FormattedErrorMessage);
 
-			++NumFailures;
+				++NumFailures;
+			}
+#endif
 		}
+
+#if !PLATFORM_EXCEPTIONS_DISABLED
+		FTaskGraphInterface::Get().WaitUntilTasksComplete(ExceptionTasks);
 #endif
 
-		ThisModuleTimer.Stop();
-		UE_LOG(LogCompile, Log, TEXT("Preparsed module %s containing %i files(s) in %.2f secs."), *Module.LongPackageName, NumHeadersPreparsed, ThisModulePreparseTime);
+		ResolveTimer.Stop();
+		UE_LOG(LogCompile, Log, TEXT("Resolving classes in %d modules containing %d files(s) in %.3f secs."), GManifest.Modules.Num(), NumHeaders, ResolveTime);
 	}
 
 	return Result;
