Various UE4->UE5 renames.

diff --git a/Engine/Plugins/Developer/RiderSourceCodeAccess/Source/RiderSourceCodeAccess/Private/RiderSourceCodeAccessor.cpp b/Engine/Plugins/Developer/RiderSourceCodeAccess/Source/RiderSourceCodeAccess/Private/RiderSourceCodeAccessor.cpp
--- a/Engine/Plugins/Developer/RiderSourceCodeAccess/Source/RiderSourceCodeAccess/Private/RiderSourceCodeAccessor.cpp
+++ b/Engine/Plugins/Developer/RiderSourceCodeAccess/Source/RiderSourceCodeAccess/Private/RiderSourceCodeAccessor.cpp
@@ -273,7 +273,7 @@ void FRiderSourceCodeAccessor::CachePathToSln() const
 			FString MasterProjectName;
 			if (!FFileHelper::LoadFileToString(MasterProjectName, *(FPaths::EngineIntermediateDir() / TEXT("ProjectFiles/MasterProjectName.txt"))))
 			{
-				MasterProjectName = "UE4";
+				MasterProjectName = "UE5";
 			}
 			CachedSolutionPath = FPaths::Combine(FPaths::RootDir(), MasterProjectName + TEXT(".sln"));
 		}
diff --git a/Engine/Plugins/Developer/VisualStudioCodeSourceCodeAccess/Source/VisualStudioCodeSourceCodeAccess/Private/VisualStudioCodeSourceCodeAccessor.cpp b/Engine/Plugins/Developer/VisualStudioCodeSourceCodeAccess/Source/VisualStudioCodeSourceCodeAccess/Private/VisualStudioCodeSourceCodeAccessor.cpp
--- a/Engine/Plugins/Developer/VisualStudioCodeSourceCodeAccess/Source/VisualStudioCodeSourceCodeAccess/Private/VisualStudioCodeSourceCodeAccessor.cpp
+++ b/Engine/Plugins/Developer/VisualStudioCodeSourceCodeAccess/Source/VisualStudioCodeSourceCodeAccess/Private/VisualStudioCodeSourceCodeAccessor.cpp
@@ -39,7 +39,7 @@ FString FVisualStudioCodeSourceCodeAccessor::GetSolutionPath() const
 
 		if (!FUProjectDictionary::GetDefault().IsForeignProject(CachedSolutionPath))
 		{
-			CachedSolutionPath = FPaths::Combine(FPaths::RootDir(), FString("UE4") + GVSCodeWorkspaceExtension);
+			CachedSolutionPath = FPaths::Combine(FPaths::RootDir(), FString("UE5") + GVSCodeWorkspaceExtension);
 		}
 		else
 		{
diff --git a/Engine/Plugins/Developer/VisualStudioSourceCodeAccess/Source/VisualStudioSourceCodeAccess/Private/VisualStudioSourceCodeAccessor.cpp b/Engine/Plugins/Developer/VisualStudioSourceCodeAccess/Source/VisualStudioSourceCodeAccess/Private/VisualStudioSourceCodeAccessor.cpp
--- a/Engine/Plugins/Developer/VisualStudioSourceCodeAccess/Source/VisualStudioSourceCodeAccess/Private/VisualStudioSourceCodeAccessor.cpp
+++ b/Engine/Plugins/Developer/VisualStudioSourceCodeAccess/Source/VisualStudioSourceCodeAccess/Private/VisualStudioSourceCodeAccessor.cpp
@@ -1412,7 +1412,7 @@ FString FVisualStudioSourceCodeAccessor::GetSolutionPath() const
 				FString MasterProjectName;
 				if (!FFileHelper::LoadFileToString(MasterProjectName, *(FPaths::EngineIntermediateDir() / TEXT("ProjectFiles/MasterProjectName.txt"))))
 				{
-					MasterProjectName = "UE4";
+					MasterProjectName = "UE5";
 				}
 				CachedSolutionPath = FPaths::Combine(FPaths::RootDir(), MasterProjectName + TEXT(".sln"));
 			}
diff --git a/Engine/Source/Editor/GameProjectGeneration/Private/GameProjectUtils.cpp b/Engine/Source/Editor/GameProjectGeneration/Private/GameProjectUtils.cpp
--- a/Engine/Source/Editor/GameProjectGeneration/Private/GameProjectUtils.cpp
+++ b/Engine/Source/Editor/GameProjectGeneration/Private/GameProjectUtils.cpp
@@ -821,7 +821,7 @@ bool GameProjectUtils::OpenCodeIDE(const FString& ProjectFile, FText& OutFailRea
 	else
 	{
 		SolutionFolder = IFileManager::Get().ConvertToAbsolutePathForExternalAppForRead(*FPaths::RootDir());
-		SolutionFilenameWithoutExtension = TEXT("UE4");
+		SolutionFilenameWithoutExtension = TEXT("UE5");
 	}
 
 	if (!FSourceCodeNavigation::OpenProjectSolution(FPaths::Combine(SolutionFolder, SolutionFilenameWithoutExtension)))
diff --git a/Engine/Source/Programs/LiveCodingConsole/Private/LiveCodingConsole.cpp b/Engine/Source/Programs/LiveCodingConsole/Private/LiveCodingConsole.cpp
--- a/Engine/Source/Programs/LiveCodingConsole/Private/LiveCodingConsole.cpp
+++ b/Engine/Source/Programs/LiveCodingConsole/Private/LiveCodingConsole.cpp
@@ -512,7 +512,7 @@ bool LiveCodingConsoleMain(const TCHAR* CmdLine)
 				return false;
 			}
 
-			// Set the normal UE4 IsEngineExitRequested() when outer frame is closed
+			// Set the normal UE IsEngineExitRequested() when outer frame is closed
 			Slate->SetExitRequestedHandler(FSimpleDelegate::CreateStatic(&OnRequestExit));
 
 			// Prepare the custom Slate styles
diff --git a/Engine/Source/Programs/UnrealHeaderTool/Private/CodeGenerator.cpp b/Engine/Source/Programs/UnrealHeaderTool/Private/CodeGenerator.cpp
--- a/Engine/Source/Programs/UnrealHeaderTool/Private/CodeGenerator.cpp
+++ b/Engine/Source/Programs/UnrealHeaderTool/Private/CodeGenerator.cpp
@@ -83,16 +83,6 @@ struct FPerHeaderData
 	TArray<FSimplifiedParsingClassInfo> ParsedClassArray;
 };
 
-struct FManifestPerHeaderData
-{
-	TArray<FString> PublicUObjectClassesHeadersText;
-	TArray<FPerHeaderData> PublicUObjectClassesHeadersData;
-	TArray<FString> PublicUObjectHeadersText;
-	TArray<FPerHeaderData> PublicUObjectHeadersData;
-	TArray<FString> PrivateUObjectHeadersText;
-	TArray<FPerHeaderData> PrivateUObjectHeadersData;
-};
-
 static void PerformSimplifiedClassParse(UPackage* InParent, const TCHAR* FileName, const TCHAR* Buffer, FPerHeaderData& PerHeaderData);
 static void ProcessInitialClassParse(FPerHeaderData& PerHeaderData);
 
@@ -748,11 +738,11 @@ static FString OutputMetaDataCodeForObject(FOutputDevice& OutDeclaration, FOutpu
 		}
 
 		OutDeclaration.Log (TEXT("#if WITH_METADATA\r\n"));
-		OutDeclaration.Logf(TEXT("%sstatic const UE4CodeGen_Private::FMetaDataPairParam %s[];\r\n"), DeclSpaces, *MetaDataBlockNameWithoutScope);
+		OutDeclaration.Logf(TEXT("%sstatic const UECodeGen_Private::FMetaDataPairParam %s[];\r\n"), DeclSpaces, *MetaDataBlockNameWithoutScope);
 		OutDeclaration.Log (TEXT("#endif\r\n"));
 
 		Out.Log (TEXT("#if WITH_METADATA\r\n"));
-		Out.Logf(TEXT("%sconst UE4CodeGen_Private::FMetaDataPairParam %s[] = {\r\n"), Spaces, MetaDataBlockName);
+		Out.Logf(TEXT("%sconst UECodeGen_Private::FMetaDataPairParam %s[] = {\r\n"), Spaces, MetaDataBlockName);
 
 		for (const KVPType& KVP : KVPs)
 		{
@@ -1107,10 +1097,10 @@ void FNativeClassHeaderGenerator::PropertyNew(FOutputDevice& DeclOut, FOutputDev
 
 	if (FByteProperty* TypedProp = CastField<FByteProperty>(Prop))
 	{
-		DeclOut.Logf(TEXT("%sstatic const UE4CodeGen_Private::FBytePropertyParams %s;\r\n"), DeclSpaces, *NameWithoutScope);
+		DeclOut.Logf(TEXT("%sstatic const UECodeGen_Private::FBytePropertyParams %s;\r\n"), DeclSpaces, *NameWithoutScope);
 
 		Out.Logf(
-			TEXT("%sconst UE4CodeGen_Private::FBytePropertyParams %s = { %s, %s, (EPropertyFlags)0x%016llx, UE4CodeGen_Private::EPropertyGenFlags::Byte, %s, %s, %s, %s, %s };%s\r\n"),
+			TEXT("%sconst UECodeGen_Private::FBytePropertyParams %s = { %s, %s, (EPropertyFlags)0x%016llx, UECodeGen_Private::EPropertyGenFlags::Byte, %s, %s, %s, %s, %s };%s\r\n"),
 			Spaces,
 			Name,
 			*PropName,
@@ -1129,10 +1119,10 @@ void FNativeClassHeaderGenerator::PropertyNew(FOutputDevice& DeclOut, FOutputDev
 
 	if (FInt8Property* TypedProp = CastField<FInt8Property>(Prop))
 	{
-		DeclOut.Logf(TEXT("%sstatic const UE4CodeGen_Private::FInt8PropertyParams %s;\r\n"), DeclSpaces, *NameWithoutScope);
+		DeclOut.Logf(TEXT("%sstatic const UECodeGen_Private::FInt8PropertyParams %s;\r\n"), DeclSpaces, *NameWithoutScope);
 
 		Out.Logf(
-			TEXT("%sconst UE4CodeGen_Private::FInt8PropertyParams %s = { %s, %s, (EPropertyFlags)0x%016llx, UE4CodeGen_Private::EPropertyGenFlags::Int8, %s, %s, %s, %s };%s\r\n"),
+			TEXT("%sconst UECodeGen_Private::FInt8PropertyParams %s = { %s, %s, (EPropertyFlags)0x%016llx, UECodeGen_Private::EPropertyGenFlags::Int8, %s, %s, %s, %s };%s\r\n"),
 			Spaces,
 			Name,
 			*PropName,
@@ -1150,10 +1140,10 @@ void FNativeClassHeaderGenerator::PropertyNew(FOutputDevice& DeclOut, FOutputDev
 
 	if (FInt16Property* TypedProp = CastField<FInt16Property>(Prop))
 	{
-		DeclOut.Logf(TEXT("%sstatic const UE4CodeGen_Private::FInt16PropertyParams %s;\r\n"), DeclSpaces, *NameWithoutScope);
+		DeclOut.Logf(TEXT("%sstatic const UECodeGen_Private::FInt16PropertyParams %s;\r\n"), DeclSpaces, *NameWithoutScope);
 
 		Out.Logf(
-			TEXT("%sconst UE4CodeGen_Private::FInt16PropertyParams %s = { %s, %s, (EPropertyFlags)0x%016llx, UE4CodeGen_Private::EPropertyGenFlags::Int16, %s, %s, %s, %s };%s\r\n"),
+			TEXT("%sconst UECodeGen_Private::FInt16PropertyParams %s = { %s, %s, (EPropertyFlags)0x%016llx, UECodeGen_Private::EPropertyGenFlags::Int16, %s, %s, %s, %s };%s\r\n"),
 			Spaces,
 			Name,
 			*PropName,
@@ -1173,10 +1163,10 @@ void FNativeClassHeaderGenerator::PropertyNew(FOutputDevice& DeclOut, FOutputDev
 	{
 		const TCHAR* PropTypeName = GUnsizedProperties.Contains(TypedProp) ? TEXT("FUnsizedIntPropertyParams") : TEXT("FIntPropertyParams");
 
-		DeclOut.Logf(TEXT("%sstatic const UE4CodeGen_Private::%s %s;\r\n"), DeclSpaces, PropTypeName, *NameWithoutScope);
+		DeclOut.Logf(TEXT("%sstatic const UECodeGen_Private::%s %s;\r\n"), DeclSpaces, PropTypeName, *NameWithoutScope);
 
 		Out.Logf(
-			TEXT("%sconst UE4CodeGen_Private::%s %s = { %s, %s, (EPropertyFlags)0x%016llx, UE4CodeGen_Private::EPropertyGenFlags::Int, %s, %s, %s, %s };%s\r\n"),
+			TEXT("%sconst UECodeGen_Private::%s %s = { %s, %s, (EPropertyFlags)0x%016llx, UECodeGen_Private::EPropertyGenFlags::Int, %s, %s, %s, %s };%s\r\n"),
 			Spaces,
 			PropTypeName,
 			Name,
@@ -1195,10 +1185,10 @@ void FNativeClassHeaderGenerator::PropertyNew(FOutputDevice& DeclOut, FOutputDev
 
 	if (FInt64Property* TypedProp = CastField<FInt64Property>(Prop))
 	{
-		DeclOut.Logf(TEXT("%sstatic const UE4CodeGen_Private::FInt64PropertyParams %s;\r\n"), DeclSpaces, *NameWithoutScope);
+		DeclOut.Logf(TEXT("%sstatic const UECodeGen_Private::FInt64PropertyParams %s;\r\n"), DeclSpaces, *NameWithoutScope);
 
 		Out.Logf(
-			TEXT("%sconst UE4CodeGen_Private::FInt64PropertyParams %s = { %s, %s, (EPropertyFlags)0x%016llx, UE4CodeGen_Private::EPropertyGenFlags::Int64, %s, %s, %s, %s };%s\r\n"),
+			TEXT("%sconst UECodeGen_Private::FInt64PropertyParams %s = { %s, %s, (EPropertyFlags)0x%016llx, UECodeGen_Private::EPropertyGenFlags::Int64, %s, %s, %s, %s };%s\r\n"),
 			Spaces,
 			Name,
 			*PropName,
@@ -1216,10 +1206,10 @@ void FNativeClassHeaderGenerator::PropertyNew(FOutputDevice& DeclOut, FOutputDev
 
 	if (FUInt16Property* TypedProp = CastField<FUInt16Property>(Prop))
 	{
-		DeclOut.Logf(TEXT("%sstatic const UE4CodeGen_Private::FFInt16PropertyParams %s;\r\n"), DeclSpaces, *NameWithoutScope);
+		DeclOut.Logf(TEXT("%sstatic const UECodeGen_Private::FFInt16PropertyParams %s;\r\n"), DeclSpaces, *NameWithoutScope);
 
 		Out.Logf(
-			TEXT("%sconst UE4CodeGen_Private::FFInt16PropertyParams %s = { %s, %s, (EPropertyFlags)0x%016llx, UE4CodeGen_Private::EPropertyGenFlags::UInt16, %s, %s, %s, %s };%s\r\n"),
+			TEXT("%sconst UECodeGen_Private::FFInt16PropertyParams %s = { %s, %s, (EPropertyFlags)0x%016llx, UECodeGen_Private::EPropertyGenFlags::UInt16, %s, %s, %s, %s };%s\r\n"),
 			Spaces,
 			Name,
 			*PropName,
@@ -1239,10 +1229,10 @@ void FNativeClassHeaderGenerator::PropertyNew(FOutputDevice& DeclOut, FOutputDev
 	{
 		const TCHAR* PropTypeName = GUnsizedProperties.Contains(TypedProp) ? TEXT("FUnsizedFIntPropertyParams") : TEXT("FUInt32PropertyParams");
 
-		DeclOut.Logf(TEXT("%sstatic const UE4CodeGen_Private::%s %s;\r\n"), DeclSpaces, PropTypeName, *NameWithoutScope);
+		DeclOut.Logf(TEXT("%sstatic const UECodeGen_Private::%s %s;\r\n"), DeclSpaces, PropTypeName, *NameWithoutScope);
 
 		Out.Logf(
-			TEXT("%sconst UE4CodeGen_Private::%s %s = { %s, %s, (EPropertyFlags)0x%016llx, UE4CodeGen_Private::EPropertyGenFlags::UInt32, %s, %s, %s, %s };%s\r\n"),
+			TEXT("%sconst UECodeGen_Private::%s %s = { %s, %s, (EPropertyFlags)0x%016llx, UECodeGen_Private::EPropertyGenFlags::UInt32, %s, %s, %s, %s };%s\r\n"),
 			Spaces,
 			PropTypeName,
 			Name,
@@ -1261,10 +1251,10 @@ void FNativeClassHeaderGenerator::PropertyNew(FOutputDevice& DeclOut, FOutputDev
 
 	if (FUInt64Property* TypedProp = CastField<FUInt64Property>(Prop))
 	{
-		DeclOut.Logf(TEXT("%sstatic const UE4CodeGen_Private::FFInt64PropertyParams %s;\r\n"), DeclSpaces, *NameWithoutScope);
+		DeclOut.Logf(TEXT("%sstatic const UECodeGen_Private::FFInt64PropertyParams %s;\r\n"), DeclSpaces, *NameWithoutScope);
 
 		Out.Logf(
-			TEXT("%sconst UE4CodeGen_Private::FFInt64PropertyParams %s = { %s, %s, (EPropertyFlags)0x%016llx, UE4CodeGen_Private::EPropertyGenFlags::UInt64, %s, %s, %s, %s };%s\r\n"),
+			TEXT("%sconst UECodeGen_Private::FFInt64PropertyParams %s = { %s, %s, (EPropertyFlags)0x%016llx, UECodeGen_Private::EPropertyGenFlags::UInt64, %s, %s, %s, %s };%s\r\n"),
 			Spaces,
 			Name,
 			*PropName,
@@ -1282,10 +1272,10 @@ void FNativeClassHeaderGenerator::PropertyNew(FOutputDevice& DeclOut, FOutputDev
 
 	if (FFloatProperty* TypedProp = CastField<FFloatProperty>(Prop))
 	{
-		DeclOut.Logf(TEXT("%sstatic const UE4CodeGen_Private::FFloatPropertyParams %s;\r\n"), DeclSpaces, *NameWithoutScope);
+		DeclOut.Logf(TEXT("%sstatic const UECodeGen_Private::FFloatPropertyParams %s;\r\n"), DeclSpaces, *NameWithoutScope);
 
 		Out.Logf(
-			TEXT("%sconst UE4CodeGen_Private::FFloatPropertyParams %s = { %s, %s, (EPropertyFlags)0x%016llx, UE4CodeGen_Private::EPropertyGenFlags::Float, %s, %s, %s, %s };%s\r\n"),
+			TEXT("%sconst UECodeGen_Private::FFloatPropertyParams %s = { %s, %s, (EPropertyFlags)0x%016llx, UECodeGen_Private::EPropertyGenFlags::Float, %s, %s, %s, %s };%s\r\n"),
 			Spaces,
 			Name,
 			*PropName,
@@ -1303,10 +1293,10 @@ void FNativeClassHeaderGenerator::PropertyNew(FOutputDevice& DeclOut, FOutputDev
 
 	if (FDoubleProperty* TypedProp = CastField<FDoubleProperty>(Prop))
 	{
-		DeclOut.Logf(TEXT("%sstatic const UE4CodeGen_Private::FDoublePropertyParams %s;\r\n"), DeclSpaces, *NameWithoutScope);
+		DeclOut.Logf(TEXT("%sstatic const UECodeGen_Private::FDoublePropertyParams %s;\r\n"), DeclSpaces, *NameWithoutScope);
 
 		Out.Logf(
-			TEXT("%sconst UE4CodeGen_Private::FDoublePropertyParams %s = { %s, %s, (EPropertyFlags)0x%016llx, UE4CodeGen_Private::EPropertyGenFlags::Double, %s, %s, %s, %s };%s\r\n"),
+			TEXT("%sconst UECodeGen_Private::FDoublePropertyParams %s = { %s, %s, (EPropertyFlags)0x%016llx, UECodeGen_Private::EPropertyGenFlags::Double, %s, %s, %s, %s };%s\r\n"),
 			Spaces,
 			Name,
 			*PropName,
@@ -1345,16 +1335,16 @@ void FNativeClassHeaderGenerator::PropertyNew(FOutputDevice& DeclOut, FOutputDev
 			Setter = FString::Printf(TEXT("&%s_SetBit"), Name);
 		}
 
-		DeclOut.Logf(TEXT("%sstatic const UE4CodeGen_Private::FBoolPropertyParams %s;\r\n"), DeclSpaces, *NameWithoutScope);
+		DeclOut.Logf(TEXT("%sstatic const UECodeGen_Private::FBoolPropertyParams %s;\r\n"), DeclSpaces, *NameWithoutScope);
 
 		Out.Logf(
-			TEXT("%sconst UE4CodeGen_Private::FBoolPropertyParams %s = { %s, %s, (EPropertyFlags)0x%016llx, UE4CodeGen_Private::EPropertyGenFlags::Bool %s, %s, %s, sizeof(%s), %s, %s, %s };%s\r\n"),
+			TEXT("%sconst UECodeGen_Private::FBoolPropertyParams %s = { %s, %s, (EPropertyFlags)0x%016llx, UECodeGen_Private::EPropertyGenFlags::Bool %s, %s, %s, sizeof(%s), %s, %s, %s };%s\r\n"),
 			Spaces,
 			Name,
 			*PropName,
 			*PropNotifyFunc,
 			PropFlags,
-			TypedProp->IsNativeBool() ? TEXT("| UE4CodeGen_Private::EPropertyGenFlags::NativeBool") : TEXT(""),
+			TypedProp->IsNativeBool() ? TEXT("| UECodeGen_Private::EPropertyGenFlags::NativeBool") : TEXT(""),
 			FPropertyObjectFlags,
 			*ArrayDim,
 			*TypedProp->GetCPPType(nullptr, 0),
@@ -1369,10 +1359,10 @@ void FNativeClassHeaderGenerator::PropertyNew(FOutputDevice& DeclOut, FOutputDev
 
 	if (FSoftClassProperty* TypedProp = CastField<FSoftClassProperty>(Prop))
 	{
-		DeclOut.Logf(TEXT("%sstatic const UE4CodeGen_Private::FSoftClassPropertyParams %s;\r\n"), DeclSpaces, *NameWithoutScope);
+		DeclOut.Logf(TEXT("%sstatic const UECodeGen_Private::FSoftClassPropertyParams %s;\r\n"), DeclSpaces, *NameWithoutScope);
 
 		Out.Logf(
-			TEXT("%sconst UE4CodeGen_Private::FSoftClassPropertyParams %s = { %s, %s, (EPropertyFlags)0x%016llx, UE4CodeGen_Private::EPropertyGenFlags::SoftClass, %s, %s, %s, %s, %s };%s\r\n"),
+			TEXT("%sconst UECodeGen_Private::FSoftClassPropertyParams %s = { %s, %s, (EPropertyFlags)0x%016llx, UECodeGen_Private::EPropertyGenFlags::SoftClass, %s, %s, %s, %s, %s };%s\r\n"),
 			Spaces,
 			Name,
 			*PropName,
@@ -1391,10 +1381,10 @@ void FNativeClassHeaderGenerator::PropertyNew(FOutputDevice& DeclOut, FOutputDev
 
 	if (FWeakObjectProperty* TypedProp = CastField<FWeakObjectProperty>(Prop))
 	{
-		DeclOut.Logf(TEXT("%sstatic const UE4CodeGen_Private::FWeakObjectPropertyParams %s;\r\n"), DeclSpaces, *NameWithoutScope);
+		DeclOut.Logf(TEXT("%sstatic const UECodeGen_Private::FWeakObjectPropertyParams %s;\r\n"), DeclSpaces, *NameWithoutScope);
 
 		Out.Logf(
-			TEXT("%sconst UE4CodeGen_Private::FWeakObjectPropertyParams %s = { %s, %s, (EPropertyFlags)0x%016llx, UE4CodeGen_Private::EPropertyGenFlags::WeakObject, %s, %s, %s, %s, %s };%s\r\n"),
+			TEXT("%sconst UECodeGen_Private::FWeakObjectPropertyParams %s = { %s, %s, (EPropertyFlags)0x%016llx, UECodeGen_Private::EPropertyGenFlags::WeakObject, %s, %s, %s, %s, %s };%s\r\n"),
 			Spaces,
 			Name,
 			*PropName,
@@ -1413,10 +1403,10 @@ void FNativeClassHeaderGenerator::PropertyNew(FOutputDevice& DeclOut, FOutputDev
 
 	if (FLazyObjectProperty* TypedProp = CastField<FLazyObjectProperty>(Prop))
 	{
-		DeclOut.Logf(TEXT("%sstatic const UE4CodeGen_Private::FLazyObjectPropertyParams %s;\r\n"), DeclSpaces, *NameWithoutScope);
+		DeclOut.Logf(TEXT("%sstatic const UECodeGen_Private::FLazyObjectPropertyParams %s;\r\n"), DeclSpaces, *NameWithoutScope);
 
 		Out.Logf(
-			TEXT("%sconst UE4CodeGen_Private::FLazyObjectPropertyParams %s = { %s, %s, (EPropertyFlags)0x%016llx, UE4CodeGen_Private::EPropertyGenFlags::LazyObject, %s, %s, %s, %s, %s };%s\r\n"),
+			TEXT("%sconst UECodeGen_Private::FLazyObjectPropertyParams %s = { %s, %s, (EPropertyFlags)0x%016llx, UECodeGen_Private::EPropertyGenFlags::LazyObject, %s, %s, %s, %s, %s };%s\r\n"),
 			Spaces,
 			Name,
 			*PropName,
@@ -1435,10 +1425,10 @@ void FNativeClassHeaderGenerator::PropertyNew(FOutputDevice& DeclOut, FOutputDev
 
 	if (FObjectPtrProperty* TypedProp = CastField<FObjectPtrProperty>(Prop))
 	{
-		DeclOut.Logf(TEXT("%sstatic const UE4CodeGen_Private::FObjectPtrPropertyParams %s;\r\n"), DeclSpaces, *NameWithoutScope);
+		DeclOut.Logf(TEXT("%sstatic const UECodeGen_Private::FObjectPtrPropertyParams %s;\r\n"), DeclSpaces, *NameWithoutScope);
 
 		Out.Logf(
-			TEXT("%sconst UE4CodeGen_Private::FObjectPtrPropertyParams %s = { %s, %s, (EPropertyFlags)0x%016llx, UE4CodeGen_Private::EPropertyGenFlags::Object | UE4CodeGen_Private::EPropertyGenFlags::ObjectPtr, %s, %s, %s, %s, %s };%s\r\n"),
+			TEXT("%sconst UECodeGen_Private::FObjectPtrPropertyParams %s = { %s, %s, (EPropertyFlags)0x%016llx, UECodeGen_Private::EPropertyGenFlags::Object | UECodeGen_Private::EPropertyGenFlags::ObjectPtr, %s, %s, %s, %s, %s };%s\r\n"),
 			Spaces,
 			Name,
 			*PropName,
@@ -1457,10 +1447,10 @@ void FNativeClassHeaderGenerator::PropertyNew(FOutputDevice& DeclOut, FOutputDev
 
 	if (FClassPtrProperty* TypedProp = CastField<FClassPtrProperty>(Prop))
 	{
-		DeclOut.Logf(TEXT("%sstatic const UE4CodeGen_Private::FClassPtrPropertyParams %s;\r\n"), DeclSpaces, *NameWithoutScope);
+		DeclOut.Logf(TEXT("%sstatic const UECodeGen_Private::FClassPtrPropertyParams %s;\r\n"), DeclSpaces, *NameWithoutScope);
 
 		Out.Logf(
-			TEXT("%sconst UE4CodeGen_Private::FClassPtrPropertyParams %s = { %s, %s, (EPropertyFlags)0x%016llx, UE4CodeGen_Private::EPropertyGenFlags::Class | UE4CodeGen_Private::EPropertyGenFlags::ObjectPtr, %s, %s, %s, %s, %s, %s };%s\r\n"),
+			TEXT("%sconst UECodeGen_Private::FClassPtrPropertyParams %s = { %s, %s, (EPropertyFlags)0x%016llx, UECodeGen_Private::EPropertyGenFlags::Class | UECodeGen_Private::EPropertyGenFlags::ObjectPtr, %s, %s, %s, %s, %s, %s };%s\r\n"),
 			Spaces,
 			Name,
 			*PropName,
@@ -1480,10 +1470,10 @@ void FNativeClassHeaderGenerator::PropertyNew(FOutputDevice& DeclOut, FOutputDev
 
 	if (FSoftObjectProperty* TypedProp = CastField<FSoftObjectProperty>(Prop))
 	{
-		DeclOut.Logf(TEXT("%sstatic const UE4CodeGen_Private::FSoftObjectPropertyParams %s;\r\n"), DeclSpaces, *NameWithoutScope);
+		DeclOut.Logf(TEXT("%sstatic const UECodeGen_Private::FSoftObjectPropertyParams %s;\r\n"), DeclSpaces, *NameWithoutScope);
 
 		Out.Logf(
-			TEXT("%sconst UE4CodeGen_Private::FSoftObjectPropertyParams %s = { %s, %s, (EPropertyFlags)0x%016llx, UE4CodeGen_Private::EPropertyGenFlags::SoftObject, %s, %s, %s, %s, %s };%s\r\n"),
+			TEXT("%sconst UECodeGen_Private::FSoftObjectPropertyParams %s = { %s, %s, (EPropertyFlags)0x%016llx, UECodeGen_Private::EPropertyGenFlags::SoftObject, %s, %s, %s, %s, %s };%s\r\n"),
 			Spaces,
 			Name,
 			*PropName,
@@ -1502,10 +1492,10 @@ void FNativeClassHeaderGenerator::PropertyNew(FOutputDevice& DeclOut, FOutputDev
 
 	if (FClassProperty* TypedProp = CastField<FClassProperty>(Prop))
 	{
-		DeclOut.Logf(TEXT("%sstatic const UE4CodeGen_Private::FClassPropertyParams %s;\r\n"), DeclSpaces, *NameWithoutScope);
+		DeclOut.Logf(TEXT("%sstatic const UECodeGen_Private::FClassPropertyParams %s;\r\n"), DeclSpaces, *NameWithoutScope);
 
 		Out.Logf(
-			TEXT("%sconst UE4CodeGen_Private::FClassPropertyParams %s = { %s, %s, (EPropertyFlags)0x%016llx, UE4CodeGen_Private::EPropertyGenFlags::Class, %s, %s, %s, %s, %s, %s };%s\r\n"),
+			TEXT("%sconst UECodeGen_Private::FClassPropertyParams %s = { %s, %s, (EPropertyFlags)0x%016llx, UECodeGen_Private::EPropertyGenFlags::Class, %s, %s, %s, %s, %s, %s };%s\r\n"),
 			Spaces,
 			Name,
 			*PropName,
@@ -1525,10 +1515,10 @@ void FNativeClassHeaderGenerator::PropertyNew(FOutputDevice& DeclOut, FOutputDev
 
 	if (FObjectProperty* TypedProp = CastField<FObjectProperty>(Prop))
 	{
-		DeclOut.Logf(TEXT("%sstatic const UE4CodeGen_Private::FObjectPropertyParams %s;\r\n"), DeclSpaces, *NameWithoutScope);
+		DeclOut.Logf(TEXT("%sstatic const UECodeGen_Private::FObjectPropertyParams %s;\r\n"), DeclSpaces, *NameWithoutScope);
 
 		Out.Logf(
-			TEXT("%sconst UE4CodeGen_Private::FObjectPropertyParams %s = { %s, %s, (EPropertyFlags)0x%016llx, UE4CodeGen_Private::EPropertyGenFlags::Object, %s, %s, %s, %s, %s };%s\r\n"),
+			TEXT("%sconst UECodeGen_Private::FObjectPropertyParams %s = { %s, %s, (EPropertyFlags)0x%016llx, UECodeGen_Private::EPropertyGenFlags::Object, %s, %s, %s, %s, %s };%s\r\n"),
 			Spaces,
 			Name,
 			*PropName,
@@ -1547,10 +1537,10 @@ void FNativeClassHeaderGenerator::PropertyNew(FOutputDevice& DeclOut, FOutputDev
 
 	if (FInterfaceProperty* TypedProp = CastField<FInterfaceProperty>(Prop))
 	{
-		DeclOut.Logf(TEXT("%sstatic const UE4CodeGen_Private::FInterfacePropertyParams %s;\r\n"), DeclSpaces, *NameWithoutScope);
+		DeclOut.Logf(TEXT("%sstatic const UECodeGen_Private::FInterfacePropertyParams %s;\r\n"), DeclSpaces, *NameWithoutScope);
 
 		Out.Logf(
-			TEXT("%sconst UE4CodeGen_Private::FInterfacePropertyParams %s = { %s, %s, (EPropertyFlags)0x%016llx, UE4CodeGen_Private::EPropertyGenFlags::Interface, %s, %s, %s, %s, %s };%s\r\n"),
+			TEXT("%sconst UECodeGen_Private::FInterfacePropertyParams %s = { %s, %s, (EPropertyFlags)0x%016llx, UECodeGen_Private::EPropertyGenFlags::Interface, %s, %s, %s, %s, %s };%s\r\n"),
 			Spaces,
 			Name,
 			*PropName,
@@ -1569,10 +1559,10 @@ void FNativeClassHeaderGenerator::PropertyNew(FOutputDevice& DeclOut, FOutputDev
 
 	if (FNameProperty* TypedProp = CastField<FNameProperty>(Prop))
 	{
-		DeclOut.Logf(TEXT("%sstatic const UE4CodeGen_Private::FNamePropertyParams %s;\r\n"), DeclSpaces, *NameWithoutScope);
+		DeclOut.Logf(TEXT("%sstatic const UECodeGen_Private::FNamePropertyParams %s;\r\n"), DeclSpaces, *NameWithoutScope);
 
 		Out.Logf(
-			TEXT("%sconst UE4CodeGen_Private::FNamePropertyParams %s = { %s, %s, (EPropertyFlags)0x%016llx, UE4CodeGen_Private::EPropertyGenFlags::Name, %s, %s, %s, %s };%s\r\n"),
+			TEXT("%sconst UECodeGen_Private::FNamePropertyParams %s = { %s, %s, (EPropertyFlags)0x%016llx, UECodeGen_Private::EPropertyGenFlags::Name, %s, %s, %s, %s };%s\r\n"),
 			Spaces,
 			Name,
 			*PropName,
@@ -1590,10 +1580,10 @@ void FNativeClassHeaderGenerator::PropertyNew(FOutputDevice& DeclOut, FOutputDev
 
 	if (FStrProperty* TypedProp = CastField<FStrProperty>(Prop))
 	{
-		DeclOut.Logf(TEXT("%sstatic const UE4CodeGen_Private::FStrPropertyParams %s;\r\n"), DeclSpaces, *NameWithoutScope);
+		DeclOut.Logf(TEXT("%sstatic const UECodeGen_Private::FStrPropertyParams %s;\r\n"), DeclSpaces, *NameWithoutScope);
 
 		Out.Logf(
-			TEXT("%sconst UE4CodeGen_Private::FStrPropertyParams %s = { %s, %s, (EPropertyFlags)0x%016llx, UE4CodeGen_Private::EPropertyGenFlags::Str, %s, %s, %s, %s };%s\r\n"),
+			TEXT("%sconst UECodeGen_Private::FStrPropertyParams %s = { %s, %s, (EPropertyFlags)0x%016llx, UECodeGen_Private::EPropertyGenFlags::Str, %s, %s, %s, %s };%s\r\n"),
 			Spaces,
 			Name,
 			*PropName,
@@ -1611,10 +1601,10 @@ void FNativeClassHeaderGenerator::PropertyNew(FOutputDevice& DeclOut, FOutputDev
 
 	if (FArrayProperty* TypedProp = CastField<FArrayProperty>(Prop))
 	{
-		DeclOut.Logf(TEXT("%sstatic const UE4CodeGen_Private::FArrayPropertyParams %s;\r\n"), DeclSpaces, *NameWithoutScope);
+		DeclOut.Logf(TEXT("%sstatic const UECodeGen_Private::FArrayPropertyParams %s;\r\n"), DeclSpaces, *NameWithoutScope);
 
 		Out.Logf(
-			TEXT("%sconst UE4CodeGen_Private::FArrayPropertyParams %s = { %s, %s, (EPropertyFlags)0x%016llx, UE4CodeGen_Private::EPropertyGenFlags::Array, %s, %s, %s, %s, %s };%s\r\n"),
+			TEXT("%sconst UECodeGen_Private::FArrayPropertyParams %s = { %s, %s, (EPropertyFlags)0x%016llx, UECodeGen_Private::EPropertyGenFlags::Array, %s, %s, %s, %s, %s };%s\r\n"),
 			Spaces,
 			Name,
 			*PropName,
@@ -1633,10 +1623,10 @@ void FNativeClassHeaderGenerator::PropertyNew(FOutputDevice& DeclOut, FOutputDev
 
 	if (FMapProperty* TypedProp = CastField<FMapProperty>(Prop))
 	{
-		DeclOut.Logf(TEXT("%sstatic const UE4CodeGen_Private::FMapPropertyParams %s;\r\n"), DeclSpaces, *NameWithoutScope);
+		DeclOut.Logf(TEXT("%sstatic const UECodeGen_Private::FMapPropertyParams %s;\r\n"), DeclSpaces, *NameWithoutScope);
 
 		Out.Logf(
-			TEXT("%sconst UE4CodeGen_Private::FMapPropertyParams %s = { %s, %s, (EPropertyFlags)0x%016llx, UE4CodeGen_Private::EPropertyGenFlags::Map, %s, %s, %s, %s, %s };%s\r\n"),
+			TEXT("%sconst UECodeGen_Private::FMapPropertyParams %s = { %s, %s, (EPropertyFlags)0x%016llx, UECodeGen_Private::EPropertyGenFlags::Map, %s, %s, %s, %s, %s };%s\r\n"),
 			Spaces,
 			Name,
 			*PropName,
@@ -1655,10 +1645,10 @@ void FNativeClassHeaderGenerator::PropertyNew(FOutputDevice& DeclOut, FOutputDev
 
 	if (FSetProperty* TypedProp = CastField<FSetProperty>(Prop))
 	{
-		DeclOut.Logf(TEXT("%sstatic const UE4CodeGen_Private::FSetPropertyParams %s;\r\n"), DeclSpaces, *NameWithoutScope);
+		DeclOut.Logf(TEXT("%sstatic const UECodeGen_Private::FSetPropertyParams %s;\r\n"), DeclSpaces, *NameWithoutScope);
 
 		Out.Logf(
-			TEXT("%sconst UE4CodeGen_Private::FSetPropertyParams %s = { %s, %s, (EPropertyFlags)0x%016llx, UE4CodeGen_Private::EPropertyGenFlags::Set, %s, %s, %s, %s };%s\r\n"),
+			TEXT("%sconst UECodeGen_Private::FSetPropertyParams %s = { %s, %s, (EPropertyFlags)0x%016llx, UECodeGen_Private::EPropertyGenFlags::Set, %s, %s, %s, %s };%s\r\n"),
 			Spaces,
 			Name,
 			*PropName,
@@ -1676,10 +1666,10 @@ void FNativeClassHeaderGenerator::PropertyNew(FOutputDevice& DeclOut, FOutputDev
 
 	if (FStructProperty* TypedProp = CastField<FStructProperty>(Prop))
 	{
-		DeclOut.Logf(TEXT("%sstatic const UE4CodeGen_Private::FStructPropertyParams %s;\r\n"), DeclSpaces, *NameWithoutScope);
+		DeclOut.Logf(TEXT("%sstatic const UECodeGen_Private::FStructPropertyParams %s;\r\n"), DeclSpaces, *NameWithoutScope);
 
 		Out.Logf(
-			TEXT("%sconst UE4CodeGen_Private::FStructPropertyParams %s = { %s, %s, (EPropertyFlags)0x%016llx, UE4CodeGen_Private::EPropertyGenFlags::Struct, %s, %s, %s, %s, %s };%s\r\n"),
+			TEXT("%sconst UECodeGen_Private::FStructPropertyParams %s = { %s, %s, (EPropertyFlags)0x%016llx, UECodeGen_Private::EPropertyGenFlags::Struct, %s, %s, %s, %s, %s };%s\r\n"),
 			Spaces,
 			Name,
 			*PropName,
@@ -1698,10 +1688,10 @@ void FNativeClassHeaderGenerator::PropertyNew(FOutputDevice& DeclOut, FOutputDev
 
 	if (FDelegateProperty* TypedProp = CastField<FDelegateProperty>(Prop))
 	{
-		DeclOut.Logf(TEXT("%sstatic const UE4CodeGen_Private::FDelegatePropertyParams %s;\r\n"), DeclSpaces, *NameWithoutScope);
+		DeclOut.Logf(TEXT("%sstatic const UECodeGen_Private::FDelegatePropertyParams %s;\r\n"), DeclSpaces, *NameWithoutScope);
 
 		Out.Logf(
-			TEXT("%sconst UE4CodeGen_Private::FDelegatePropertyParams %s = { %s, %s, (EPropertyFlags)0x%016llx, UE4CodeGen_Private::EPropertyGenFlags::Delegate, %s, %s, %s, %s, %s };%s\r\n"),
+			TEXT("%sconst UECodeGen_Private::FDelegatePropertyParams %s = { %s, %s, (EPropertyFlags)0x%016llx, UECodeGen_Private::EPropertyGenFlags::Delegate, %s, %s, %s, %s, %s };%s\r\n"),
 			Spaces,
 			Name,
 			*PropName,
@@ -1720,10 +1710,10 @@ void FNativeClassHeaderGenerator::PropertyNew(FOutputDevice& DeclOut, FOutputDev
 
 	if (FMulticastDelegateProperty* TypedProp = CastField<FMulticastDelegateProperty>(Prop))
 	{
-		DeclOut.Logf(TEXT("%sstatic const UE4CodeGen_Private::FMulticastDelegatePropertyParams %s;\r\n"), DeclSpaces, *NameWithoutScope);
+		DeclOut.Logf(TEXT("%sstatic const UECodeGen_Private::FMulticastDelegatePropertyParams %s;\r\n"), DeclSpaces, *NameWithoutScope);
 
 		Out.Logf(
-			TEXT("%sconst UE4CodeGen_Private::FMulticastDelegatePropertyParams %s = { %s, %s, (EPropertyFlags)0x%016llx, UE4CodeGen_Private::EPropertyGenFlags::%sMulticastDelegate, %s, %s, %s, %s, %s };%s\r\n"),
+			TEXT("%sconst UECodeGen_Private::FMulticastDelegatePropertyParams %s = { %s, %s, (EPropertyFlags)0x%016llx, UECodeGen_Private::EPropertyGenFlags::%sMulticastDelegate, %s, %s, %s, %s, %s };%s\r\n"),
 			Spaces,
 			Name,
 			*PropName,
@@ -1743,10 +1733,10 @@ void FNativeClassHeaderGenerator::PropertyNew(FOutputDevice& DeclOut, FOutputDev
 
 	if (FTextProperty* TypedProp = CastField<FTextProperty>(Prop))
 	{
-		DeclOut.Logf(TEXT("%sstatic const UE4CodeGen_Private::FTextPropertyParams %s;\r\n"), DeclSpaces, *NameWithoutScope);
+		DeclOut.Logf(TEXT("%sstatic const UECodeGen_Private::FTextPropertyParams %s;\r\n"), DeclSpaces, *NameWithoutScope);
 
 		Out.Logf(
-			TEXT("%sconst UE4CodeGen_Private::FTextPropertyParams %s = { %s, %s, (EPropertyFlags)0x%016llx, UE4CodeGen_Private::EPropertyGenFlags::Text, %s, %s, %s, %s };%s\r\n"),
+			TEXT("%sconst UECodeGen_Private::FTextPropertyParams %s = { %s, %s, (EPropertyFlags)0x%016llx, UECodeGen_Private::EPropertyGenFlags::Text, %s, %s, %s, %s };%s\r\n"),
 			Spaces,
 			Name,
 			*PropName,
@@ -1764,10 +1754,10 @@ void FNativeClassHeaderGenerator::PropertyNew(FOutputDevice& DeclOut, FOutputDev
 
 	if (FEnumProperty* TypedProp = CastField<FEnumProperty>(Prop))
 	{
-		DeclOut.Logf(TEXT("%sstatic const UE4CodeGen_Private::FEnumPropertyParams %s;\r\n"), DeclSpaces, *NameWithoutScope);
+		DeclOut.Logf(TEXT("%sstatic const UECodeGen_Private::FEnumPropertyParams %s;\r\n"), DeclSpaces, *NameWithoutScope);
 
 		Out.Logf(
-			TEXT("%sconst UE4CodeGen_Private::FEnumPropertyParams %s = { %s, %s, (EPropertyFlags)0x%016llx, UE4CodeGen_Private::EPropertyGenFlags::Enum, %s, %s, %s, %s, %s };%s\r\n"),
+			TEXT("%sconst UECodeGen_Private::FEnumPropertyParams %s = { %s, %s, (EPropertyFlags)0x%016llx, UECodeGen_Private::EPropertyGenFlags::Enum, %s, %s, %s, %s, %s };%s\r\n"),
 			Spaces,
 			Name,
 			*PropName,
@@ -1786,10 +1776,10 @@ void FNativeClassHeaderGenerator::PropertyNew(FOutputDevice& DeclOut, FOutputDev
 
 	if (FFieldPathProperty* TypedProp = CastField<FFieldPathProperty>(Prop))
 	{
-		DeclOut.Logf(TEXT("%sstatic const UE4CodeGen_Private::FFieldPathPropertyParams %s;\r\n"), DeclSpaces, *NameWithoutScope);
+		DeclOut.Logf(TEXT("%sstatic const UECodeGen_Private::FFieldPathPropertyParams %s;\r\n"), DeclSpaces, *NameWithoutScope);
 
 		Out.Logf(
-			TEXT("%sconst UE4CodeGen_Private::FFieldPathPropertyParams %s = { %s, %s, (EPropertyFlags)0x%016llx, UE4CodeGen_Private::EPropertyGenFlags::FieldPath, %s, %s, %s, %s, %s };%s\r\n"),
+			TEXT("%sconst UECodeGen_Private::FFieldPathPropertyParams %s = { %s, %s, (EPropertyFlags)0x%016llx, UECodeGen_Private::EPropertyGenFlags::FieldPath, %s, %s, %s, %s, %s };%s\r\n"),
 			Spaces,
 			Name,
 			*PropName,
@@ -1854,8 +1844,8 @@ TTuple<FString, FString> FNativeClassHeaderGenerator::OutputProperties(FOutputDe
 
 		WithEditorOnlyMacroEmitter(bHasAllEditorOnlyDataProperties);
 		WithEditorOnlyMacroEmitterDecl(bHasAllEditorOnlyDataProperties);
-		DeclOut.Logf(TEXT("%sstatic const UE4CodeGen_Private::FPropertyParamsBase* const PropPointers[];\r\n"), DeclSpaces);
-		Out.Logf(TEXT("%sconst UE4CodeGen_Private::FPropertyParamsBase* const %sPropPointers[] = {\r\n"), Spaces, Scope);
+		DeclOut.Logf(TEXT("%sstatic const UECodeGen_Private::FPropertyParamsBase* const PropPointers[];\r\n"), DeclSpaces);
+		Out.Logf(TEXT("%sconst UECodeGen_Private::FPropertyParamsBase* const %sPropPointers[] = {\r\n"), Spaces, Scope);
 
 		for (const FPropertyNamePointerPair& PropNameAndPtr : PropertyNamesAndPointers)
 		{
@@ -1863,7 +1853,7 @@ TTuple<FString, FString> FNativeClassHeaderGenerator::OutputProperties(FOutputDe
 
 			WithEditorOnlyMacroEmitter(bRequiresHasEditorOnlyMacro);
 			WithEditorOnlyMacroEmitterDecl(bRequiresHasEditorOnlyMacro);
-			Out.Logf(TEXT("%s\t(const UE4CodeGen_Private::FPropertyParamsBase*)&%s,\r\n"), Spaces, *PropNameAndPtr.Name);
+			Out.Logf(TEXT("%s\t(const UECodeGen_Private::FPropertyParamsBase*)&%s,\r\n"), Spaces, *PropNameAndPtr.Name);
 		}
 
 		WithEditorOnlyMacroEmitter(bHasAllEditorOnlyDataProperties);
@@ -2180,7 +2170,7 @@ void FNativeClassHeaderGenerator::ExportGeneratedPackageInitCode(FOutputDevice&
 		SingletonCount = TEXT("0");
 	}
 
-	Out.Logf(TEXT("\t\t\tstatic const UE4CodeGen_Private::FPackageParams PackageParams = {\r\n"));
+	Out.Logf(TEXT("\t\t\tstatic const UECodeGen_Private::FPackageParams PackageParams = {\r\n"));
 	Out.Logf(TEXT("\t\t\t\t%s,\r\n"), *CreateUTF8LiteralString(InPackage->GetName()));
 	Out.Logf(TEXT("\t\t\t\t%s,\r\n"), SingletonArray);
 	Out.Logf(TEXT("\t\t\t\t%s,\r\n"), SingletonCount);
@@ -2189,7 +2179,7 @@ void FNativeClassHeaderGenerator::ExportGeneratedPackageInitCode(FOutputDevice&
 	Out.Logf(TEXT("\t\t\t\t0x%08X,\r\n"), GenerateTextHash(InDeclarations));
 	Out.Logf(TEXT("\t\t\t\t%s\r\n"), *MetaDataParams);
 	Out.Logf(TEXT("\t\t\t};\r\n"));
-	Out.Logf(TEXT("\t\t\tUE4CodeGen_Private::ConstructUPackage(ReturnPackage, PackageParams);\r\n"));
+	Out.Logf(TEXT("\t\t\tUECodeGen_Private::ConstructUPackage(ReturnPackage, PackageParams);\r\n"));
 	Out.Logf(TEXT("\t\t}\r\n"));
 	Out.Logf(TEXT("\t\treturn ReturnPackage;\r\n"));
 	Out.Logf(TEXT("\t}\r\n"));
@@ -2369,9 +2359,9 @@ void FNativeClassHeaderGenerator::ExportNativeGeneratedInitCode(FOutputDevice& O
 		const TCHAR* InterfaceCount;
 		if (Class->Interfaces.Num() > 0)
 		{
-			GeneratedClassRegisterFunctionText.Log(TEXT("\t\tstatic const UE4CodeGen_Private::FImplementedInterfaceParams InterfaceParams[];\r\n"));
+			GeneratedClassRegisterFunctionText.Log(TEXT("\t\tstatic const UECodeGen_Private::FImplementedInterfaceParams InterfaceParams[];\r\n"));
 
-			StaticDefinitions.Logf(TEXT("\t\tconst UE4CodeGen_Private::FImplementedInterfaceParams %s::InterfaceParams[] = {\r\n"), *StaticsStructName);
+			StaticDefinitions.Logf(TEXT("\t\tconst UECodeGen_Private::FImplementedInterfaceParams %s::InterfaceParams[] = {\r\n"), *StaticsStructName);
 			for (const FImplementedInterface& Inter : Class->Interfaces)
 			{
 				check(Inter.Class);
@@ -2408,7 +2398,7 @@ void FNativeClassHeaderGenerator::ExportNativeGeneratedInitCode(FOutputDevice& O
 		StaticDefinitions.Logf(TEXT("\t\tTCppClassTypeTraits<%s>::IsAbstract,\r\n"), *FNameLookupCPP::GetNameCPP(Class, Class->HasAllClassFlags(CLASS_Interface)));
 		StaticDefinitions.Logf(TEXT("\t};\r\n"));
 
-		GeneratedClassRegisterFunctionText.Log (TEXT("\t\tstatic const UE4CodeGen_Private::FClassParams ClassParams;\r\n"));
+		GeneratedClassRegisterFunctionText.Log (TEXT("\t\tstatic const UECodeGen_Private::FClassParams ClassParams;\r\n"));
 
 		uint32 ClassFlags = (uint32)Class->ClassFlags;
 		if (!bIsNoExport)
@@ -2417,7 +2407,7 @@ void FNativeClassHeaderGenerator::ExportNativeGeneratedInitCode(FOutputDevice& O
 		}
 		ClassFlags = ClassFlags & CLASS_SaveInCompiledInClasses;
 
-		StaticDefinitions.Logf(TEXT("\tconst UE4CodeGen_Private::FClassParams %s::ClassParams = {\r\n"), *StaticsStructName);
+		StaticDefinitions.Logf(TEXT("\tconst UECodeGen_Private::FClassParams %s::ClassParams = {\r\n"), *StaticsStructName);
 		StaticDefinitions.Logf(TEXT("\t\t&%s::StaticClass,\r\n"), *ClassNameCPP);
 		StaticDefinitions.Logf(TEXT("\t\t%s,\r\n"), (Class->ClassConfigName != NAME_None) ? *CreateUTF8LiteralString(Class->ClassConfigName.ToString()) : TEXT("nullptr"));
 		StaticDefinitions.Log (TEXT("\t\t&StaticCppClassTypeInfo,\r\n"));
@@ -2452,7 +2442,7 @@ void FNativeClassHeaderGenerator::ExportNativeGeneratedInitCode(FOutputDevice& O
 		}
 
 		GeneratedClassRegisterFunctionText.Logf(TEXT("\t\t{\r\n"));
-		GeneratedClassRegisterFunctionText.Logf(TEXT("\t\t\tUE4CodeGen_Private::ConstructUClass(OuterClass, %s::ClassParams);\r\n"), *StaticsStructName);
+		GeneratedClassRegisterFunctionText.Logf(TEXT("\t\t\tUECodeGen_Private::ConstructUClass(OuterClass, %s::ClassParams);\r\n"), *StaticsStructName);
 
 		TArray<FString> SparseClassDataTypes;
 		((FClass*)Class)->GetSparseClassDataTypes(SparseClassDataTypes);
@@ -2703,10 +2693,10 @@ void FNativeClassHeaderGenerator::ExportFunction(FOutputDevice& Out, FReferenceG
 
 	FString MetaDataParams = OutputMetaDataCodeForObject(CurrentFunctionText, StaticDefinitions, Function, *FString::Printf(TEXT("%s::Function_MetaDataParams"), *StaticsStructName), TEXT("\t\t"), TEXT("\t"));
 
-	CurrentFunctionText.Log(TEXT("\t\tstatic const UE4CodeGen_Private::FFunctionParams FuncParams;\r\n"));
+	CurrentFunctionText.Log(TEXT("\t\tstatic const UECodeGen_Private::FFunctionParams FuncParams;\r\n"));
 
 	StaticDefinitions.Logf(
-		TEXT("\tconst UE4CodeGen_Private::FFunctionParams %s::FuncParams = { (UObject*(*)())%s, %s, %s, %s, %s, %s, %s, %s, %s, (EFunctionFlags)0x%08X, %d, %d, %s };\r\n"),
+		TEXT("\tconst UECodeGen_Private::FFunctionParams %s::FuncParams = { (UObject*(*)())%s, %s, %s, %s, %s, %s, %s, %s, %s, (EFunctionFlags)0x%08X, %d, %d, %s };\r\n"),
 		*StaticsStructName,
 		*OuterFunc,
 		*GetSingletonNameFuncAddr(SuperFunction, OutReferenceGatherers.UniqueCrossModuleReferences),
@@ -2742,7 +2732,7 @@ void FNativeClassHeaderGenerator::ExportFunction(FOutputDevice& Out, FReferenceG
 
 	CurrentFunctionText.Logf(TEXT("\t\tif (!ReturnFunction)\r\n"));
 	CurrentFunctionText.Logf(TEXT("\t\t{\r\n"));
-	CurrentFunctionText.Logf(TEXT("\t\t\tUE4CodeGen_Private::ConstructUFunction(ReturnFunction, %s::FuncParams, &ReturnFunction);\r\n"), *StaticsStructName);
+	CurrentFunctionText.Logf(TEXT("\t\t\tUECodeGen_Private::ConstructUFunction(ReturnFunction, %s::FuncParams, &ReturnFunction);\r\n"), *StaticsStructName);
 	CurrentFunctionText.Log (TEXT("\t\t}\r\n"));
 	CurrentFunctionText.Log (TEXT("\t\treturn ReturnFunction;\r\n"));
 	CurrentFunctionText.Log (TEXT("\t}\r\n"));
@@ -4040,9 +4030,9 @@ void FNativeClassHeaderGenerator::ExportGeneratedStructBodyMacros(FOutputDevice&
 
 	TTuple<FString, FString> PropertyRange = OutputProperties(GeneratedStructRegisterFunctionText, StaticDefinitions, OutReferenceGatherers, *FString::Printf(TEXT("%s::"), *StaticsStructName), Props, TEXT("\t\t"), TEXT("\t"));
 
-	GeneratedStructRegisterFunctionText.Log (TEXT("\t\tstatic const UE4CodeGen_Private::FStructParams ReturnStructParams;\r\n"));
+	GeneratedStructRegisterFunctionText.Log (TEXT("\t\tstatic const UECodeGen_Private::FStructParams ReturnStructParams;\r\n"));
 
-	StaticDefinitions.Logf(TEXT("\tconst UE4CodeGen_Private::FStructParams %s::ReturnStructParams = {\r\n"), *StaticsStructName);
+	StaticDefinitions.Logf(TEXT("\tconst UECodeGen_Private::FStructParams %s::ReturnStructParams = {\r\n"), *StaticsStructName);
 	StaticDefinitions.Logf(TEXT("\t\t(UObject* (*)())%s,\r\n"), *OuterFunc);
 	StaticDefinitions.Logf(TEXT("\t\t%s,\r\n"), *GetSingletonNameFuncAddr(BaseStruct, OutReferenceGatherers.UniqueCrossModuleReferences));
 	StaticDefinitions.Logf(TEXT("\t\t%s,\r\n"), *NewStructOps);
@@ -4094,7 +4084,7 @@ void FNativeClassHeaderGenerator::ExportGeneratedStructBodyMacros(FOutputDevice&
 	GeneratedStructRegisterFunctionText.Logf(TEXT("\t\tif (!ReturnStruct)\r\n"));
 	GeneratedStructRegisterFunctionText.Log (TEXT("\t\t{\r\n"));
 
-	GeneratedStructRegisterFunctionText.Logf(TEXT("\t\t\tUE4CodeGen_Private::ConstructUScriptStruct(ReturnStruct, %s::ReturnStructParams);\r\n"), *StaticsStructName);
+	GeneratedStructRegisterFunctionText.Logf(TEXT("\t\t\tUECodeGen_Private::ConstructUScriptStruct(ReturnStruct, %s::ReturnStructParams);\r\n"), *StaticsStructName);
 	GeneratedStructRegisterFunctionText.Log (TEXT("\t\t}\r\n"));
 	GeneratedStructRegisterFunctionText.Logf(TEXT("\t\treturn ReturnStruct;\r\n"));
 	GeneratedStructRegisterFunctionText.Log (TEXT("\t}\r\n"));
@@ -4271,8 +4261,8 @@ void FNativeClassHeaderGenerator::ExportGeneratedEnumInitCode(FOutputDevice& Out
 		StaticDeclarations.Logf(TEXT("\tstruct %s\r\n"), *StaticsStructName);
 		StaticDeclarations.Logf(TEXT("\t{\r\n"));
 
-		StaticDeclarations.Logf(TEXT("\t\tstatic const UE4CodeGen_Private::FEnumeratorParam Enumerators[];\r\n"));
-		StaticDefinitions.Logf(TEXT("\tconst UE4CodeGen_Private::FEnumeratorParam %s::Enumerators[] = {\r\n"), *StaticsStructName);
+		StaticDeclarations.Logf(TEXT("\t\tstatic const UECodeGen_Private::FEnumeratorParam Enumerators[];\r\n"));
+		StaticDefinitions.Logf(TEXT("\tconst UECodeGen_Private::FEnumeratorParam %s::Enumerators[] = {\r\n"), *StaticsStructName);
 		for (int32 Index = 0; Index != Enum->NumEnums(); ++Index)
 		{
 			const TCHAR* OverridenNameMetaDatakey = TEXT("OverrideName");
@@ -4284,8 +4274,8 @@ void FNativeClassHeaderGenerator::ExportGeneratedEnumInitCode(FOutputDevice& Out
 		FString MetaDataParamsName = StaticsStructName + TEXT("::Enum_MetaDataParams");
 		FString MetaDataParams = OutputMetaDataCodeForObject(StaticDeclarations, StaticDefinitions, Enum, *MetaDataParamsName, TEXT("\t\t"), TEXT("\t"));
 
-		StaticDeclarations.Logf(TEXT("\t\tstatic const UE4CodeGen_Private::FEnumParams EnumParams;\r\n"));
-		StaticDefinitions.Logf(TEXT("\tconst UE4CodeGen_Private::FEnumParams %s::EnumParams = {\r\n"), *StaticsStructName);
+		StaticDeclarations.Logf(TEXT("\t\tstatic const UECodeGen_Private::FEnumParams EnumParams;\r\n"));
+		StaticDefinitions.Logf(TEXT("\tconst UECodeGen_Private::FEnumParams %s::EnumParams = {\r\n"), *StaticsStructName);
 		StaticDefinitions.Logf(TEXT("\t\t(UObject*(*)())%s,\r\n"), *OuterString.LeftChop(2));
 		StaticDefinitions.Logf(TEXT("\t\t%s,\r\n"), EnumDisplayNameFn.IsEmpty() ? TEXT("nullptr") : *EnumDisplayNameFn);
 		StaticDefinitions.Logf(TEXT("\t\t%s,\r\n"), *CreateUTF8LiteralString(OverriddenEnumNameCpp));
@@ -4294,7 +4284,7 @@ void FNativeClassHeaderGenerator::ExportGeneratedEnumInitCode(FOutputDevice& Out
 		StaticDefinitions.Logf(TEXT("\t\tUE_ARRAY_COUNT(%s::Enumerators),\r\n"), *StaticsStructName);
 		StaticDefinitions.Logf(TEXT("\t\t%s,\r\n"), UEnumObjectFlags);
 		StaticDefinitions.Logf(TEXT("\t\t%s,\r\n"), EnumFlags);
-		StaticDefinitions.Logf(TEXT("\t\tUE4CodeGen_Private::EDynamicType::%s,\r\n"), bIsDynamic ? TEXT("Dynamic") : TEXT("NotDynamic"));
+		StaticDefinitions.Logf(TEXT("\t\tUECodeGen_Private::EDynamicType::%s,\r\n"), bIsDynamic ? TEXT("Dynamic") : TEXT("NotDynamic"));
 		StaticDefinitions.Logf(TEXT("\t\t(uint8)%s,\r\n"), EnumFormStr);
 		StaticDefinitions.Logf(TEXT("\t\t%s\r\n"), *MetaDataParams);
 		StaticDefinitions.Logf(TEXT("\t};\r\n"));
@@ -4325,7 +4315,7 @@ void FNativeClassHeaderGenerator::ExportGeneratedEnumInitCode(FOutputDevice& Out
 	GeneratedEnumRegisterFunctionText.Logf(TEXT("\t\tif (!ReturnEnum)\r\n"));
 	GeneratedEnumRegisterFunctionText.Logf(TEXT("\t\t{\r\n"));
 
-	GeneratedEnumRegisterFunctionText.Logf(TEXT("\t\t\tUE4CodeGen_Private::ConstructUEnum(ReturnEnum, %s::EnumParams);\r\n"), *StaticsStructName);
+	GeneratedEnumRegisterFunctionText.Logf(TEXT("\t\t\tUECodeGen_Private::ConstructUEnum(ReturnEnum, %s::EnumParams);\r\n"), *StaticsStructName);
 	GeneratedEnumRegisterFunctionText.Logf(TEXT("\t\t}\r\n"));
 	GeneratedEnumRegisterFunctionText.Logf(TEXT("\t\treturn ReturnEnum;\r\n"));
 	GeneratedEnumRegisterFunctionText.Logf(TEXT("\t}\r\n"));
diff --git a/Engine/Source/Runtime/CoreUObject/Private/UObject/UObjectGlobals.cpp b/Engine/Source/Runtime/CoreUObject/Private/UObject/UObjectGlobals.cpp
--- a/Engine/Source/Runtime/CoreUObject/Private/UObject/UObjectGlobals.cpp
+++ b/Engine/Source/Runtime/CoreUObject/Private/UObject/UObjectGlobals.cpp
@@ -4176,7 +4176,7 @@ FString FAssetMsg::GetAssetLogString(const UObject* Object, const FString& Messa
 	return ensure(Object) ? GetAssetLogString(*Object->GetOutermost()->GetName(), Message) : FString();
 }
 
-namespace UE4CodeGen_Private
+namespace UECodeGen_Private
 {
 	void ConstructFProperty(FFieldVariant Outer, const FPropertyParamsBase* const*& PropertyArray, int32& NumProperties)
 	{
diff --git a/Engine/Source/Runtime/CoreUObject/Public/UObject/UObjectGlobals.h b/Engine/Source/Runtime/CoreUObject/Public/UObject/UObjectGlobals.h
--- a/Engine/Source/Runtime/CoreUObject/Public/UObject/UObjectGlobals.h
+++ b/Engine/Source/Runtime/CoreUObject/Public/UObject/UObjectGlobals.h
@@ -2578,7 +2578,7 @@ struct FCppClassTypeInfoStatic;
 class FFieldClass;
 
 /// @cond DOXYGEN_IGNORE
-namespace UE4CodeGen_Private
+namespace UECodeGen_Private
 {
 	enum class EPropertyGenFlags : uint32
 	{
