Fixed bad merge

#rnx

diff --git a/Engine/Source/Programs/UnrealHeaderTool/Private/CodeGenerator.cpp b/Engine/Source/Programs/UnrealHeaderTool/Private/CodeGenerator.cpp
--- a/Engine/Source/Programs/UnrealHeaderTool/Private/CodeGenerator.cpp
+++ b/Engine/Source/Programs/UnrealHeaderTool/Private/CodeGenerator.cpp
@@ -83,6 +83,16 @@ struct FPerHeaderData
 	TArray<FSimplifiedParsingClassInfo> ParsedClassArray;
 };
 
+struct FManifestPerHeaderData
+{
+	TArray<FString> PublicUObjectClassesHeadersText;
+	TArray<FPerHeaderData> PublicUObjectClassesHeadersData;
+	TArray<FString> PublicUObjectHeadersText;
+	TArray<FPerHeaderData> PublicUObjectHeadersData;
+	TArray<FString> PrivateUObjectHeadersText;
+	TArray<FPerHeaderData> PrivateUObjectHeadersData;
+};
+
 static void PerformSimplifiedClassParse(UPackage* InParent, const TCHAR* FileName, const TCHAR* Buffer, FPerHeaderData& PerHeaderData);
 static void ProcessInitialClassParse(FPerHeaderData& PerHeaderData);
 
