Client fails to build due to an assert in UHT.

diff --git a/Engine/Source/Programs/UnrealHeaderTool/Private/CodeGenerator.cpp b/Engine/Source/Programs/UnrealHeaderTool/Private/CodeGenerator.cpp
--- a/Engine/Source/Programs/UnrealHeaderTool/Private/CodeGenerator.cpp
+++ b/Engine/Source/Programs/UnrealHeaderTool/Private/CodeGenerator.cpp
@@ -6853,11 +6853,6 @@ ECompilationResult::Type PreparseModules(const FString& ModuleInfoPath, int32& N
 		FDurationTimer PreparseTimer(PreparseTime);
 		PreparseTimer.Start();
 
-		// Create the load tasks
-		FGraphEventRef LastClassParseTask;
-		FGraphEventArray ClassParseTaskPreReqs;
-		ClassParseTaskPreReqs.SetNum(2);
-
 		FGraphEventArray LoadTasks;
 		LoadTasks.SetNum(NumHeaders);
 		for (int32 ModuleIndex = 0, NumModules = GManifest.Modules.Num(); ModuleIndex < NumModules; ++ModuleIndex)
@@ -6958,134 +6953,169 @@ ECompilationResult::Type PreparseModules(const FString& ModuleInfoPath, int32& N
 #endif
 					};
 
-					// Phase #3: Perform initial class parse, can not run until previous header is done
-					auto ClassParseLambda = [&Module, Package, CurrentlyProcessing, &UObjectHeaders, &PerHeaderData, &LogException, &ModuleInfoPath, &Result, Index]()
+					FGraphEventRef LoadTask = FFunctionGraphTask::CreateAndDispatchWhenReady(MoveTemp(LoadLambda), TStatId());
+					FGraphEventRef PreProcessTask = FFunctionGraphTask::CreateAndDispatchWhenReady(MoveTemp(PreProcessLambda), TStatId(), LoadTask);
+					LoadTasks.Add(MoveTemp(PreProcessTask));
+				}
+			}
+		}
+
+		// Wait for all the loading and preparsing to complete
+		FTaskGraphInterface::Get().WaitUntilTasksComplete(LoadTasks);
+
+#if !PLATFORM_EXCEPTIONS_DISABLED
+		FTaskGraphInterface::Get().WaitUntilTasksComplete(ExceptionTasks);
+#endif
+
+		PreparseTimer.Stop();
+		UE_LOG(LogCompile, Log, TEXT("Loaded and preparsed classes in %d modules containing %d files(s) in %.3f secs."), GManifest.Modules.Num(), NumHeaders, PreparseTime);
+	}
+
+	// Class declaration phase
+	if (Result == ECompilationResult::Succeeded)
+	{
+		double PreparseTime = 0.0;
+		FDurationTimer PreparseTimer(PreparseTime);
+		PreparseTimer.Start();
+
+		for (int32 ModuleIndex = 0, NumModules = GManifest.Modules.Num(); ModuleIndex < NumModules; ++ModuleIndex)
+		{
+			FManifestModule& Module = GManifest.Modules[ModuleIndex];
+			FManifestPerHeaderData& ModulePerHeaderData = ManifestPerHeaderDatas[ModuleIndex];
+			UPackage* Package = GetModulePackage(Module);
+
+			// Pre-parse the headers
+			for (int32 PassIndex = 0; PassIndex < FolderType_Count && Result == ECompilationResult::Succeeded; ++PassIndex)
+			{
+				EHeaderFolderTypes CurrentlyProcessing = (EHeaderFolderTypes)PassIndex;
+
+				// We'll make an ordered list of all UObject headers we care about.
+				// @todo uht: Ideally 'dependson' would not be allowed from public -> private, or NOT at all for new style headers
+				const TArray<FString>& UObjectHeaders =
+					(CurrentlyProcessing == PublicClassesHeaders) ? Module.PublicUObjectClassesHeaders :
+					(CurrentlyProcessing == PublicHeaders) ? Module.PublicUObjectHeaders :
+					Module.PrivateUObjectHeaders;
+				if (!UObjectHeaders.Num())
+				{
+					continue;
+				}
+
+				TArray<FString>& HeaderFiles =
+					(CurrentlyProcessing == PublicClassesHeaders) ? ModulePerHeaderData.PublicUObjectClassesHeadersText :
+					(CurrentlyProcessing == PublicHeaders) ? ModulePerHeaderData.PublicUObjectHeadersText :
+					ModulePerHeaderData.PrivateUObjectHeadersText;
+				HeaderFiles.SetNum(UObjectHeaders.Num());
+
+				TArray<FPerHeaderData>& PerHeaderData =
+					(CurrentlyProcessing == PublicClassesHeaders) ? ModulePerHeaderData.PublicUObjectClassesHeadersData :
+					(CurrentlyProcessing == PublicHeaders) ? ModulePerHeaderData.PublicUObjectHeadersData :
+					ModulePerHeaderData.PrivateUObjectHeadersData;
+				PerHeaderData.SetNum(UObjectHeaders.Num());
+
+				for (int32 Index = 0, EIndex = UObjectHeaders.Num(); Index < EIndex; ++Index)
+				{
+					if (Result != ECompilationResult::Succeeded)
 					{
-						if (Result != ECompilationResult::Succeeded)
-						{
-							return;
-						}
+						break;
+					}
 
-						const FString& RawFilename = UObjectHeaders[Index];
+					const FString& RawFilename = UObjectHeaders[Index];
 
 #if !PLATFORM_EXCEPTIONS_DISABLED
-						try
+					try
 #endif
+					{
+						// Import class.
+						const FString FullFilename = FPaths::ConvertRelativePathToFull(ModuleInfoPath, RawFilename);
+
+						ProcessInitialClassParse(PerHeaderData[Index]);
+						TSharedRef<FUnrealSourceFile> UnrealSourceFile = PerHeaderData[Index].UnrealSourceFile.ToSharedRef();
+						FUnrealSourceFile* UnrealSourceFilePtr = &UnrealSourceFile.Get();
+						FString CleanFilename = FPaths::GetCleanFilename(RawFilename);
+						uint32  CleanFilenameHash = GetTypeHash(CleanFilename);
+						if (const TSharedRef<FUnrealSourceFile>* ExistingSourceFile = GUnrealSourceFilesMap.FindByHash(CleanFilenameHash, CleanFilename))
 						{
-							// Import class.
-							const FString FullFilename = FPaths::ConvertRelativePathToFull(ModuleInfoPath, RawFilename);
+							FString NormalizedFullFilename = FullFilename;
+							FString NormalizedExistingFilename = (*ExistingSourceFile)->GetFilename();
+
+							FPaths::NormalizeFilename(NormalizedFullFilename);
+							FPaths::NormalizeFilename(NormalizedExistingFilename);
 
-							ProcessInitialClassParse(PerHeaderData[Index]);
-							TSharedRef<FUnrealSourceFile> UnrealSourceFile = PerHeaderData[Index].UnrealSourceFile.ToSharedRef();
-							FUnrealSourceFile* UnrealSourceFilePtr = &UnrealSourceFile.Get();
-							FString CleanFilename = FPaths::GetCleanFilename(RawFilename);
-							uint32  CleanFilenameHash = GetTypeHash(CleanFilename);
-							if (const TSharedRef<FUnrealSourceFile>* ExistingSourceFile = GUnrealSourceFilesMap.FindByHash(CleanFilenameHash, CleanFilename))
+							if (NormalizedFullFilename != NormalizedExistingFilename)
 							{
-								FString NormalizedFullFilename = FullFilename;
-								FString NormalizedExistingFilename = (*ExistingSourceFile)->GetFilename();
+								FError::Throwf(TEXT("Duplicate leaf header name found: %s (original: %s)"), *NormalizedFullFilename, *NormalizedExistingFilename);
+							}
+						}
+						GUnrealSourceFilesMap.AddByHash(CleanFilenameHash, MoveTemp(CleanFilename), UnrealSourceFile);
 
-								FPaths::NormalizeFilename(NormalizedFullFilename);
-								FPaths::NormalizeFilename(NormalizedExistingFilename);
+						if (CurrentlyProcessing == PublicClassesHeaders)
+						{
+							GPublicSourceFileSet.Add(UnrealSourceFilePtr);
+						}
 
-								if (NormalizedFullFilename != NormalizedExistingFilename)
-								{
-									FError::Throwf(TEXT("Duplicate leaf header name found: %s (original: %s)"), *NormalizedFullFilename, *NormalizedExistingFilename);
-								}
+						// Save metadata for the class path, both for it's include path and relative to the module base directory
+						if (FullFilename.StartsWith(Module.BaseDirectory))
+						{
+							// Get the path relative to the module directory
+							const TCHAR* ModuleRelativePath = *FullFilename + Module.BaseDirectory.Len();
+
+							UnrealSourceFilePtr->SetModuleRelativePath(ModuleRelativePath);
+
+							// Calculate the include path
+							const TCHAR* IncludePath = ModuleRelativePath;
+
+							// Walk over the first potential slash
+							if (*IncludePath == TEXT('/'))
+							{
+								IncludePath++;
 							}
-							GUnrealSourceFilesMap.AddByHash(CleanFilenameHash, MoveTemp(CleanFilename), UnrealSourceFile);
 
-							if (CurrentlyProcessing == PublicClassesHeaders)
+							// Does this module path start with a known include path location? If so, we can cut that part out of the include path
+							static const TCHAR PublicFolderName[] = TEXT("Public/");
+							static const TCHAR PrivateFolderName[] = TEXT("Private/");
+							static const TCHAR ClassesFolderName[] = TEXT("Classes/");
+							if (FCString::Strnicmp(IncludePath, PublicFolderName, UE_ARRAY_COUNT(PublicFolderName) - 1) == 0)
+							{
+								IncludePath += (UE_ARRAY_COUNT(PublicFolderName) - 1);
+							}
+							else if (FCString::Strnicmp(IncludePath, PrivateFolderName, UE_ARRAY_COUNT(PrivateFolderName) - 1) == 0)
 							{
-								GPublicSourceFileSet.Add(UnrealSourceFilePtr);
+								IncludePath += (UE_ARRAY_COUNT(PrivateFolderName) - 1);
+							}
+							else if (FCString::Strnicmp(IncludePath, ClassesFolderName, UE_ARRAY_COUNT(ClassesFolderName) - 1) == 0)
+							{
+								IncludePath += (UE_ARRAY_COUNT(ClassesFolderName) - 1);
 							}
 
-							// Save metadata for the class path, both for it's include path and relative to the module base directory
-							if (FullFilename.StartsWith(Module.BaseDirectory))
+							// Add the include path
+							if (*IncludePath != 0)
 							{
-								// Get the path relative to the module directory
-								const TCHAR* ModuleRelativePath = *FullFilename + Module.BaseDirectory.Len();
-
-								UnrealSourceFilePtr->SetModuleRelativePath(ModuleRelativePath);
-
-								// Calculate the include path
-								const TCHAR* IncludePath = ModuleRelativePath;
-
-								// Walk over the first potential slash
-								if (*IncludePath == TEXT('/'))
-								{
-									IncludePath++;
-								}
-
-								// Does this module path start with a known include path location? If so, we can cut that part out of the include path
-								static const TCHAR PublicFolderName[] = TEXT("Public/");
-								static const TCHAR PrivateFolderName[] = TEXT("Private/");
-								static const TCHAR ClassesFolderName[] = TEXT("Classes/");
-								if (FCString::Strnicmp(IncludePath, PublicFolderName, UE_ARRAY_COUNT(PublicFolderName) - 1) == 0)
-								{
-									IncludePath += (UE_ARRAY_COUNT(PublicFolderName) - 1);
-								}
-								else if (FCString::Strnicmp(IncludePath, PrivateFolderName, UE_ARRAY_COUNT(PrivateFolderName) - 1) == 0)
-								{
-									IncludePath += (UE_ARRAY_COUNT(PrivateFolderName) - 1);
-								}
-								else if (FCString::Strnicmp(IncludePath, ClassesFolderName, UE_ARRAY_COUNT(ClassesFolderName) - 1) == 0)
-								{
-									IncludePath += (UE_ARRAY_COUNT(ClassesFolderName) - 1);
-								}
-
-								// Add the include path
-								if (*IncludePath != 0)
-								{
-									UnrealSourceFilePtr->SetIncludePath(MoveTemp(IncludePath));
-								}
+								UnrealSourceFilePtr->SetIncludePath(MoveTemp(IncludePath));
 							}
 						}
+					}
 #if !PLATFORM_EXCEPTIONS_DISABLED
-						catch (const FFileLineException& Ex)
-						{
-							FString AbsFilename = IFileManager::Get().ConvertToAbsolutePathForExternalAppForRead(*Ex.Filename);
-							LogException(MoveTemp(AbsFilename), Ex.Line, Ex.Message);
-						}
-						catch (TCHAR* ErrorMsg)
-						{
-							FString AbsFilename = IFileManager::Get().ConvertToAbsolutePathForExternalAppForRead(*RawFilename);
-							LogException(MoveTemp(AbsFilename), 1, ErrorMsg);
-						}
-#endif
-					};
-
-					// Only the first ClassParseTask uses only the pre process task.  All others must wait on the previous 
-					// ClassParseTask.  In the future we should modify the code to not require this.
-					FGraphEventRef LoadTask = FFunctionGraphTask::CreateAndDispatchWhenReady(MoveTemp(LoadLambda), TStatId());
-					FGraphEventRef PreProcessTask = FFunctionGraphTask::CreateAndDispatchWhenReady(MoveTemp(PreProcessLambda), TStatId(), LoadTask);
-					FGraphEventRef ClassParseTask;
-					if (LastClassParseTask.IsValid())
+					catch (const FFileLineException& Ex)
 					{
-						ClassParseTaskPreReqs[0] = LastClassParseTask;
-						ClassParseTaskPreReqs[1] = PreProcessTask;
-						ClassParseTask = FFunctionGraphTask::CreateAndDispatchWhenReady(MoveTemp(ClassParseLambda), TStatId(), &ClassParseTaskPreReqs);
+						FString AbsFilename = IFileManager::Get().ConvertToAbsolutePathForExternalAppForRead(*Ex.Filename);
+						LogException(MoveTemp(AbsFilename), Ex.Line, Ex.Message);
 					}
-					else
+					catch (TCHAR* ErrorMsg)
 					{
-						ClassParseTask = FFunctionGraphTask::CreateAndDispatchWhenReady(MoveTemp(ClassParseLambda), TStatId(), PreProcessTask);
+						FString AbsFilename = IFileManager::Get().ConvertToAbsolutePathForExternalAppForRead(*RawFilename);
+						LogException(MoveTemp(AbsFilename), 1, ErrorMsg);
 					}
-					LastClassParseTask = ClassParseTask;
-
-					LoadTasks.Add(MoveTemp(ClassParseTask));
+#endif
 				}
 			}
 		}
 
-		// Wait for all the loading and preparsing to complete
-		FTaskGraphInterface::Get().WaitUntilTasksComplete(LoadTasks);
-
 #if !PLATFORM_EXCEPTIONS_DISABLED
 		FTaskGraphInterface::Get().WaitUntilTasksComplete(ExceptionTasks);
 #endif
 
 		PreparseTimer.Stop();
-		UE_LOG(LogCompile, Log, TEXT("Loaded and preparsed classes in %d modules containing %d files(s) in %.3f secs."), GManifest.Modules.Num(), NumHeaders, PreparseTime);
+		UE_LOG(LogCompile, Log, TEXT("Defined classes in %d modules containing %d files(s) in %.3f secs."), GManifest.Modules.Num(), NumHeaders, PreparseTime);
 	}
 
 	// Perform the second phase of loading
