Fix issue where if in a single include file, we define UMyInterface, implement IMyInterface, and implement the interface in a UObject, UHT would generate a circular dependency error.

The fix is to ignore identifier references that resolve to the same source file.

diff --git a/Engine/Source/Programs/UnrealHeaderTool/Private/CodeGenerator.cpp b/Engine/Source/Programs/UnrealHeaderTool/Private/CodeGenerator.cpp
--- a/Engine/Source/Programs/UnrealHeaderTool/Private/CodeGenerator.cpp
+++ b/Engine/Source/Programs/UnrealHeaderTool/Private/CodeGenerator.cpp
@@ -5758,7 +5758,7 @@ void GetSourceFilesInDependencyOrderRecursive(TArray<FUnrealSourceFile*>& OutTes
 	VisitedSet.Add(SourceFile);
 	for (FHeaderProvider& Include : SourceFile->GetIncludes())
 	{
-		if (FUnrealSourceFile* IncludeFile = Include.Resolve())
+		if (FUnrealSourceFile* IncludeFile = Include.Resolve(*SourceFile))
 		{
 			GetSourceFilesInDependencyOrderRecursive(OutTest, Package, IncludeFile, VisitedSet, bCheckDependenciesOnly, Ignore);
 		}
diff --git a/Engine/Source/Programs/UnrealHeaderTool/Private/HeaderParser.cpp b/Engine/Source/Programs/UnrealHeaderTool/Private/HeaderParser.cpp
--- a/Engine/Source/Programs/UnrealHeaderTool/Private/HeaderParser.cpp
+++ b/Engine/Source/Programs/UnrealHeaderTool/Private/HeaderParser.cpp
@@ -8811,7 +8811,7 @@ ECompilationResult::Type FHeaderParser::ParseHeaders(FClasses& AllClasses, FHead
 			continue;
 		}
 
-		if (FUnrealSourceFile* DepFile = Include.Resolve())
+		if (FUnrealSourceFile* DepFile = Include.Resolve(*SourceFile))
 		{
 			SourceFilesRequired.Add(DepFile);
 		}
diff --git a/Engine/Source/Programs/UnrealHeaderTool/Private/HeaderProvider.cpp b/Engine/Source/Programs/UnrealHeaderTool/Private/HeaderProvider.cpp
--- a/Engine/Source/Programs/UnrealHeaderTool/Private/HeaderProvider.cpp
+++ b/Engine/Source/Programs/UnrealHeaderTool/Private/HeaderProvider.cpp
@@ -13,7 +13,7 @@ FHeaderProvider::FHeaderProvider(EHeaderProviderSourceType InType, FString&& InI
 
 }
 
-FUnrealSourceFile* FHeaderProvider::Resolve()
+FUnrealSourceFile* FHeaderProvider::Resolve(const FUnrealSourceFile& ParentSourceFile)
 {
 	if (Type != EHeaderProviderSourceType::Resolved)
 	{
@@ -23,6 +23,15 @@ FUnrealSourceFile* FHeaderProvider::Resolve()
 			if (TSharedRef<FUnrealTypeDefinitionInfo>* Source = GTypeDefinitionInfoMap.FindByName(IdName))
 			{
 				Cache = &(*Source)->GetUnrealSourceFile();
+
+				// There is an edge case with interfaces.  If you define the UMyInterface and IMyInterface in the same
+				// source file as a class that implements the interface, a HeaderProvider for IMyInterface is added 
+				// at the pre-parse time that later (incorrectly) resolves to UMyInterface.  This results in
+				// the include file thinking that it includes itself.
+				if (Cache == &ParentSourceFile)
+				{
+					Cache = nullptr;
+				}
 			}
 		}
 		else if (const TSharedRef<FUnrealSourceFile>* Source = GUnrealSourceFilesMap.Find(Id))
diff --git a/Engine/Source/Programs/UnrealHeaderTool/Private/HeaderProvider.h b/Engine/Source/Programs/UnrealHeaderTool/Private/HeaderProvider.h
--- a/Engine/Source/Programs/UnrealHeaderTool/Private/HeaderProvider.h
+++ b/Engine/Source/Programs/UnrealHeaderTool/Private/HeaderProvider.h
@@ -20,7 +20,7 @@ class FHeaderProvider
 public:
 	FHeaderProvider(EHeaderProviderSourceType Type, FString&& Id);
 
-	FUnrealSourceFile* Resolve();
+	FUnrealSourceFile* Resolve(const FUnrealSourceFile& ParentSourceFile);
 
 	FString ToString() const;
 
