Minor changes to stablize UHT code generation in preperation for more extensive performance changes.

#rnx

diff --git a/Engine/Source/Programs/UnrealHeaderTool/Private/CodeGenerator.cpp b/Engine/Source/Programs/UnrealHeaderTool/Private/CodeGenerator.cpp
--- a/Engine/Source/Programs/UnrealHeaderTool/Private/CodeGenerator.cpp
+++ b/Engine/Source/Programs/UnrealHeaderTool/Private/CodeGenerator.cpp
@@ -2137,11 +2137,18 @@ void FNativeClassHeaderGenerator::ExportGeneratedPackageInitCode(FOutputDevice&
 	if (SingletonsToOutput)
 	{
 		Algo::Sort(*SingletonsToOutput, [](UField* A, UField* B)
-		{
-			// Structs before delegates then UniqueId order
-			return (uint64(A->IsA<UDelegateFunction>()) << 32) + A->GetUniqueID() <
-			       (uint64(B->IsA<UDelegateFunction>()) << 32) + B->GetUniqueID();
-		});
+			{
+				bool bADel = A->IsA<UDelegateFunction>();
+				bool bBDel = B->IsA<UDelegateFunction>();
+				if (bADel != bBDel)
+				{
+					return !bADel;
+				}
+				const FString& AName = FTypeSingletonCache::Get(A, true).GetName();
+				const FString& BName = FTypeSingletonCache::Get(B, true).GetName();
+				return AName < BName;
+			});
+
 
 		for (UField* ScriptType : *SingletonsToOutput)
 		{
@@ -6050,9 +6057,20 @@ FNativeClassHeaderGenerator::FNativeClassHeaderGenerator(
 
 			RecordPackageSingletons(*SourceFile->GetPackage(), Structs, DelegateFunctions);
 
+			// Sort by line number to help make sure things are in order
+			auto Predicate = [](const UField& Lhs, const UField& Rhs)
+			{
+				TSharedRef<FUnrealTypeDefinitionInfo>* TLhs = GTypeDefinitionInfoMap.Find(&Lhs);
+				check(TLhs);
+				TSharedRef<FUnrealTypeDefinitionInfo>* TRhs = GTypeDefinitionInfoMap.Find(&Rhs);
+				check(TRhs);
+
+				return (*TLhs)->GetLineNumber() < (*TRhs)->GetLineNumber();
+			};
+			Enums.StableSort(Predicate);
+			Structs.StableSort(Predicate);
+
 			// Reverse the containers as they come out in the reverse order of declaration
-			Algo::Reverse(Enums);
-			Algo::Reverse(Structs);
 			Algo::Reverse(DelegateFunctions);
 
 			const FString FileDefineName = SourceFile->GetFileDefineName();
@@ -6066,13 +6084,6 @@ FNativeClassHeaderGenerator::FNativeClassHeaderGenerator(
 				LINE_TERMINATOR,
 				*FileDefineName, *StrippedFilename, *StrippedFilename, *FileDefineName);
 
-			// export delegate definitions
-			for (UDelegateFunction* Func : DelegateFunctions)
-			{
-				GeneratedFunctionDeclarations.Log(FTypeSingletonCache::Get(Func).GetExternDecl());
-				ConstThis->ExportDelegateDeclaration(OutText, ReferenceGatherers, *SourceFile, Func);
-			}
-
 			// Export enums declared in non-UClass headers.
 			for (UEnum* Enum : Enums)
 			{
@@ -6092,6 +6103,13 @@ FNativeClassHeaderGenerator::FNativeClassHeaderGenerator(
 				ConstThis->ExportGeneratedStructBodyMacros(GeneratedHeaderText, OutText, ReferenceGatherers, *SourceFile, Struct);
 			}
 
+			// export delegate definitions
+			for (UDelegateFunction* Func : DelegateFunctions)
+			{
+				GeneratedFunctionDeclarations.Log(FTypeSingletonCache::Get(Func).GetExternDecl());
+				ConstThis->ExportDelegateDeclaration(OutText, ReferenceGatherers, *SourceFile, Func);
+			}
+
 			// export delegate wrapper function implementations
 			for (UDelegateFunction* Func : DelegateFunctions)
 			{
@@ -6157,11 +6175,15 @@ FNativeClassHeaderGenerator::FNativeClassHeaderGenerator(
 		FileInfo.StartLoad(MoveTemp(ClassesHeaderPath));
 	}
 
+	// Create a sorted list of the exported source files so that the generated code is consistent
+	TArray<FUnrealSourceFile*> ExportedSorted(Exported);
+	ExportedSorted.Sort([](const FUnrealSourceFile& Lhs, const FUnrealSourceFile& Rhs) { return Lhs.GetFilename() < Rhs.GetFilename(); });
+
 	// Export an include line for each header
 	TSet<FUnrealSourceFile*> PublicHeaderGroupIncludes;
 	FUHTStringBuilder GeneratedFunctionDeclarations;
 
-	for (FUnrealSourceFile* SourceFile : Exported)
+	for (FUnrealSourceFile* SourceFile : ExportedSorted)
 	{
 		for (const TPair<UClass*, FSimplifiedParsingClassInfo>& ClassDataPair : SourceFile->GetDefinedClassesWithParsingInfo())
 		{
@@ -6179,7 +6201,7 @@ FNativeClassHeaderGenerator::FNativeClassHeaderGenerator(
 	}
 
 	// Add includes for 'Within' classes
-	for (FUnrealSourceFile* SourceFile : Exported)
+	for (FUnrealSourceFile* SourceFile : ExportedSorted)
 	{
 		bool bAddedStructuredArchiveFromArchiveHeader = false;
 		bool bAddedArchiveUObjectFromStructuredArchiveHeader = false;
@@ -6232,9 +6254,18 @@ FNativeClassHeaderGenerator::FNativeClassHeaderGenerator(
 			PublicHeaderGroupIncludes.Append(*SourceFilesForPackage);
 		}
 
+		// Make the public header list stable regardless of the order the files are parsed.
+		TArray<FString> BuildPaths;
+		BuildPaths.Reserve(PublicHeaderGroupIncludes.Num());
 		for (FUnrealSourceFile* SourceFile : PublicHeaderGroupIncludes)
 		{
-			ClassesHText.Logf(TEXT("#include \"%s\"") LINE_TERMINATOR, *GetBuildPath(*SourceFile));
+			BuildPaths.Emplace(GetBuildPath(*SourceFile));
+		}
+		BuildPaths.Sort();
+
+		for (const FString& BuildPath : BuildPaths)
+		{
+			ClassesHText.Logf(TEXT("#include \"%s\"") LINE_TERMINATOR, *BuildPath);
 		}
 
 		ClassesHText.Log(LINE_TERMINATOR);
@@ -6254,9 +6285,10 @@ FNativeClassHeaderGenerator::FNativeClassHeaderGenerator(
 	if (GeneratedFunctionDeclarations.Len())
 	{
 		uint32 CombinedHash = 0;
-		for (const TPair<FUnrealSourceFile*, FGeneratedCPP>& GeneratedCPP : GeneratedCPPs)
+		for (FUnrealSourceFile* SourceFile : ExportedSorted)
 		{
-			uint32 SplitHash = GenerateTextHash(*GeneratedCPP.Value.GeneratedText);
+			const FGeneratedCPP& GeneratedCPP = GeneratedCPPs.FindChecked(SourceFile);
+			uint32 SplitHash = GenerateTextHash(*GeneratedCPP.GeneratedText);
 			if (CombinedHash == 0)
 			{
 				// Don't combine in the first case because it keeps GUID backwards compatibility
diff --git a/Engine/Source/Programs/UnrealHeaderTool/Private/HeaderParser.cpp b/Engine/Source/Programs/UnrealHeaderTool/Private/HeaderParser.cpp
--- a/Engine/Source/Programs/UnrealHeaderTool/Private/HeaderParser.cpp
+++ b/Engine/Source/Programs/UnrealHeaderTool/Private/HeaderParser.cpp
@@ -5205,30 +5205,6 @@ void FHeaderParser::GetVarType(
 	}
 }
 
-/**
- * If the property has already been seen during compilation, then return add. If not,
- * then return replace so that INI files don't mess with header exporting
- *
- * @param PropertyName the string token for the property
- *
- * @return FNAME_Replace_Not_Safe_For_Threading or FNAME_Add
- */
-EFindName FHeaderParser::GetFindFlagForPropertyName(const TCHAR* PropertyName)
-{
-	static TMap<FString,int32> PreviousNames;
-	FString PropertyStr(PropertyName);
-	FString UpperPropertyStr = PropertyStr.ToUpper();
-	// See if it's in the list already
-	if (PreviousNames.Find(UpperPropertyStr))
-	{
-		return FNAME_Add;
-	}
-	// Add it to the list for future look ups
-	PreviousNames.Add(MoveTemp(UpperPropertyStr),1);
-	FName CurrentText(PropertyName,FNAME_Find); // keep generating this FName in case it has been affecting the case of future FNames.
-	return FNAME_Replace_Not_Safe_For_Threading;
-}
-
 FProperty* FHeaderParser::GetVarNameAndDim
 (
 	UStruct*                Scope,
@@ -5507,10 +5483,8 @@ FProperty* FHeaderParser::GetVarNameAndDim
 		}
 	}
 
-	// If this is the first time seeing the property name, then flag it for replace instead of add
-	const EFindName FindFlag = VarProperty.PropertyFlags & CPF_Config ? GetFindFlagForPropertyName(VarProperty.Identifier) : FNAME_Add;
 	// create the FName for the property, splitting (ie Unnamed_3 -> Unnamed,3)
-	FName PropertyName(VarProperty.Identifier, FindFlag);
+	FName PropertyName(VarProperty.Identifier);
 
 	FProperty* Prev = nullptr;
 	for (TFieldIterator<FProperty> It(Scope, EFieldIteratorFlags::ExcludeSuper); It; ++It)
diff --git a/Engine/Source/Programs/UnrealHeaderTool/Private/HeaderParser.h b/Engine/Source/Programs/UnrealHeaderTool/Private/HeaderParser.h
--- a/Engine/Source/Programs/UnrealHeaderTool/Private/HeaderParser.h
+++ b/Engine/Source/Programs/UnrealHeaderTool/Private/HeaderParser.h
@@ -988,16 +988,6 @@ protected:
 	void InitScriptLocation( FScriptLocation& Retry );
 	void ReturnToLocation( const FScriptLocation& Retry, bool Binary=1, bool Text=1 );
 
-	/**
-	 * If the property has already been seen during compilation, then return add. If not,
-	 * then return replace so that INI files don't mess with header exporting
-	 *
-	 * @param PropertyName the string token for the property
-	 *
-	 * @return FNAME_Replace_Not_Safe_For_Threading or FNAME_Add
-	 */
-	EFindName GetFindFlagForPropertyName(const TCHAR* PropertyName);
-
 	static void ValidatePropertyIsDeprecatedIfNecessary(const FPropertyBase& VarProperty, const FToken* OuterPropertyType);
 
 	// Cache of ScriptStructs that have been validated for Net Replication and RPC
