Cleanup the usage of FClasses.  Remove the need to pass around AllClasses when the majority of the time, the methods being invoked could be static.  Removed methods that are currently unused.

#rnx

diff --git a/Engine/Source/Programs/UnrealHeaderTool/Private/ClassDeclarationMetaData.cpp b/Engine/Source/Programs/UnrealHeaderTool/Private/ClassDeclarationMetaData.cpp
--- a/Engine/Source/Programs/UnrealHeaderTool/Private/ClassDeclarationMetaData.cpp
+++ b/Engine/Source/Programs/UnrealHeaderTool/Private/ClassDeclarationMetaData.cpp
@@ -396,7 +396,7 @@ void FClassDeclarationMetaData::MergeClassCategories(FClass* Class)
 	AutoExpandCategories.Append(MoveTemp(ParentAutoExpandCategories));
 }
 
-void FClassDeclarationMetaData::MergeAndValidateClassFlags(const FString& DeclaredClassName, uint32 PreviousClassFlags, FClass* Class, const FClasses& AllClasses)
+void FClassDeclarationMetaData::MergeAndValidateClassFlags(const FString& DeclaredClassName, uint32 PreviousClassFlags, FClass* Class)
 {
 	if (WantsToBePlaceable)
 	{
@@ -412,7 +412,7 @@ void FClassDeclarationMetaData::MergeAndValidateClassFlags(const FString& Declar
 	Class->ClassFlags |= ClassFlags;
 	Class->ClassConfigName = FName(*ConfigName);
 
-	SetAndValidateWithinClass(Class, AllClasses);
+	SetAndValidateWithinClass(Class);
 	SetAndValidateConfigName(Class);
 
 	if (!!(Class->ClassFlags & CLASS_EditInlineNew))
@@ -495,12 +495,12 @@ void FClassDeclarationMetaData::SetAndValidateConfigName(FClass* Class)
 	}
 }
 
-void FClassDeclarationMetaData::SetAndValidateWithinClass(FClass* Class, const FClasses& AllClasses)
+void FClassDeclarationMetaData::SetAndValidateWithinClass(FClass* Class)
 {
 	// Process all of the class specifiers
 	if (ClassWithin.IsEmpty() == false)
 	{
-		UClass* RequiredWithinClass = AllClasses.FindClass(*ClassWithin);
+		UClass* RequiredWithinClass = FClasses::FindClass(*ClassWithin);
 		if (!RequiredWithinClass)
 		{
 			FError::Throwf(TEXT("Within class '%s' not found."), *ClassWithin);
diff --git a/Engine/Source/Programs/UnrealHeaderTool/Private/ClassDeclarationMetaData.h b/Engine/Source/Programs/UnrealHeaderTool/Private/ClassDeclarationMetaData.h
--- a/Engine/Source/Programs/UnrealHeaderTool/Private/ClassDeclarationMetaData.h
+++ b/Engine/Source/Programs/UnrealHeaderTool/Private/ClassDeclarationMetaData.h
@@ -8,7 +8,6 @@
 #include "UObject/Class.h"
 
 class FClass;
-class FClasses;
 struct FPropertySpecifier;
 
 /** Structure that holds class meta data generated from its UCLASS declaration */
@@ -55,13 +54,13 @@ public:
 	* @param	Class Class to merge flags for
 	* @param  AllClasses All known classes
 	*/
-	void MergeAndValidateClassFlags(const FString& DeclaredClassName, uint32 PreviousClassFlags, FClass* Class, const FClasses& AllClasses);
+	void MergeAndValidateClassFlags(const FString& DeclaredClassName, uint32 PreviousClassFlags, FClass* Class);
 private:
 
 	/** Merges all 'show' categories */
 	void MergeShowCategories();
 	/** Sets and validates 'within' property */
-	void SetAndValidateWithinClass(FClass* Class, const FClasses& AllClasses);
+	void SetAndValidateWithinClass(FClass* Class);
 	/** Sets and validates 'ConfigName' property */
 	void SetAndValidateConfigName(FClass* Class);
 
diff --git a/Engine/Source/Programs/UnrealHeaderTool/Private/Classes.cpp b/Engine/Source/Programs/UnrealHeaderTool/Private/Classes.cpp
--- a/Engine/Source/Programs/UnrealHeaderTool/Private/Classes.cpp
+++ b/Engine/Source/Programs/UnrealHeaderTool/Private/Classes.cpp
@@ -62,39 +62,7 @@ FClass* FClasses::GetRootClass() const
 	return UObjectClass;
 }
 
-bool FClasses::IsDependentOn(const FClass* Suspect, const FClass* Source) const
-{
-	check(Suspect != Source);
-	TSet<const FClass*> VisitedDpendencies;
-	return IsDependentOn(Suspect, Source, VisitedDpendencies);
-}
-
-bool FClasses::IsDependentOn(const FClass* Suspect, const FClass* Source, TSet<const FClass*>& VisitedDpendencies) const
-{
-	// Children are all implicitly dependent on their parent, that is, children require their parent
-	// to be compiled first therefore if the source is a parent of the suspect, the suspect is
-	// dependent on the source.
-	if (Suspect->IsChildOf(Source))
-	{
-		return true;
-	}
-
-	// Prevent circular #includes from causing infinite recursion
-	// Note that although it may mean there's a circular dependency somewhere, it does not
-	// necessarily mean it's the one we're looking for
-	if (VisitedDpendencies.Contains(Suspect))
-	{
-		return false;
-	}
-	else
-	{
-		VisitedDpendencies.Add(Suspect);
-	}
-
-	return false;
-}
-
-FClass* FClasses::FindClass(const TCHAR* ClassName) const
+FClass* FClasses::FindClass(const TCHAR* ClassName)
 {
 	check(ClassName);
 
@@ -113,29 +81,7 @@ FClass* FClasses::FindClass(const TCHAR* ClassName) const
 	return nullptr;
 }
 
-TArray<FClass*> FClasses::GetDerivedClasses(FClass* Parent) const
-{
-	const FClassTree* ClassLeaf = ClassTree.FindNode(Parent);
-	TArray<const FClassTree*> ChildLeaves;
-	ClassLeaf->GetChildClasses(ChildLeaves);
-
-	TArray<FClass*> Result;
-	Result.Reserve(ChildLeaves.Num());
-	for (const FClassTree* Node : ChildLeaves)
-	{
-		Result.Add((FClass*)Node->GetClass());
-	}
-
-	return Result;
-}
-
-FClass* FClasses::FindAnyClass(const TCHAR* ClassName) const
-{
-	check(ClassName);
-	return (FClass*)FindObject<UClass>(ANY_PACKAGE, ClassName);
-}
-
-FClass* FClasses::FindScriptClassOrThrow(const FString& InClassName) const
+FClass* FClasses::FindScriptClassOrThrow(const FString& InClassName)
 {
 	FString ErrorMsg;
 	if (FClass* Result = FindScriptClass(InClassName, &ErrorMsg))
@@ -149,7 +95,7 @@ FClass* FClasses::FindScriptClassOrThrow(const FString& InClassName) const
 	return 0;
 }
 
-FClass* FClasses::FindScriptClass(const FString& InClassName, FString* OutErrorMsg) const
+FClass* FClasses::FindScriptClass(const FString& InClassName, FString* OutErrorMsg)
 {
 	// Strip the class name of its prefix and then do a search for the class
 	FString ClassNameStripped = GetClassNameWithPrefixRemoved(InClassName);
@@ -203,16 +149,6 @@ TArray<FClass*> FClasses::GetClassesInPackage(const UPackage* InPackage) const
 
 #if WIP_UHT_REFACTOR
 
-void FClasses::ChangeParentClass(FClass* Class)
-{
-	ClassTree.ChangeParentClass(Class);
-}
-
-bool FClasses::ContainsClass(const FClass* Class) const
-{
-	return !!ClassTree.FindNode(const_cast<FClass*>(Class));
-}
-
 void FClasses::Validate()
 {
 	ClassTree.Validate();
diff --git a/Engine/Source/Programs/UnrealHeaderTool/Private/Classes.h b/Engine/Source/Programs/UnrealHeaderTool/Private/Classes.h
--- a/Engine/Source/Programs/UnrealHeaderTool/Private/Classes.h
+++ b/Engine/Source/Programs/UnrealHeaderTool/Private/Classes.h
@@ -24,21 +24,7 @@ public:
 	 */
 	FClass* GetRootClass() const;
 
-	/**
-	 * Determines whether the class hierarchy rooted at Suspect is
-	 * dependent on the hierarchy rooted at Source.
-	 *
-	 * @param   Suspect  Root of hierarchy for suspect class
-	 * @param   Source   Root of hierarchy for source class
-	 * @return           true if the hierarchy rooted at Suspect is dependent on the one rooted at Source, false otherwise
-	 */
-	bool IsDependentOn(const FClass* Suspect, const FClass* Source) const;
-
-	FClass* FindClass(const TCHAR* ClassName) const;
-
-	TArray<FClass*> GetDerivedClasses(FClass* Parent) const;
-
-	FClass* FindAnyClass(const TCHAR* ClassName) const;
+	static FClass* FindClass(const TCHAR* ClassName);
 
 	/** 
 	 * Attempts to find a script class based on the given name. Will attempt to strip
@@ -48,7 +34,7 @@ public:
 	 * @param   InClassName  Name w/ Unreal prefix to use when searching for a class
 	 * @return               The found class.
 	 */
-	FClass* FindScriptClassOrThrow(const FString& InClassName) const;
+	static FClass* FindScriptClassOrThrow(const FString& InClassName);
 
 	/** 
 	 * Attempts to find a script class based on the given name. Will attempt to strip
@@ -58,7 +44,7 @@ public:
 	 * @param   OutErrorMsg  Error message (if any) giving the caller flexibility in how they present an error
 	 * @return               The found class, or NULL if the class was not found.
 	 */
-	FClass* FindScriptClass(const FString& InClassName, FString* OutErrorMsg = nullptr) const;
+	static FClass* FindScriptClass(const FString& InClassName, FString* OutErrorMsg = nullptr);
 
 	/**
 	 * Returns an array of classes for the given package.
@@ -71,18 +57,6 @@ public:
 // Anything in here should eventually be removed when this class encapsulates its own data structure, rather than being 'poked' by the outside
 #if WIP_UHT_REFACTOR
 
-	/**
-	 * Move a class node in the hierarchy tree after a class has changed its SuperClass
-	 * 
-	 * @param	SearchClass			the class that has changed parents
-	 * @param	InNewParentClass	if non-null force reparenting to this instead of the SuperClass
-	 *
-	 * @return	true if SearchClass was successfully moved to the new location
-	 */
-	void ChangeParentClass(FClass* Class);
-
-	bool ContainsClass(const FClass* Class) const;
-
 	/**
 	 * Validates the state of the tree (shouldn't be needed once this class has well-defined invariants).
 	 */
@@ -98,13 +72,4 @@ public:
 private:
 	FClass*     UObjectClass;
 	FClassTree  ClassTree;
-
-	/**
-	* Determines whether the class hierarchy rooted at Suspect is dependent on the hierarchy rooted at Source.
-	* Used by the public overload of IsDependentOn to recursively track dependencies and handle circular references
-	*/
-	bool IsDependentOn(const FClass* Suspect, const FClass* Source, TSet<const FClass*>& VisitedDpendencies) const;
-
-	friend auto begin(const FClasses& Classes) -> decltype(begin(TObjectRange<FClass>())) { return begin(TObjectRange<FClass>()); }
-	friend auto end  (const FClasses& Classes) -> decltype(end  (TObjectRange<FClass>())) { return end  (TObjectRange<FClass>()); }
 };
diff --git a/Engine/Source/Programs/UnrealHeaderTool/Private/HeaderParser.cpp b/Engine/Source/Programs/UnrealHeaderTool/Private/HeaderParser.cpp
--- a/Engine/Source/Programs/UnrealHeaderTool/Private/HeaderParser.cpp
+++ b/Engine/Source/Programs/UnrealHeaderTool/Private/HeaderParser.cpp
@@ -1378,7 +1378,7 @@ FString FHeaderParser::GetContext()
 //
 // Get a qualified class.
 //
-FClass* FHeaderParser::GetQualifiedClass(const FClasses& AllClasses, const TCHAR* Thing)
+FClass* FHeaderParser::GetQualifiedClass(const TCHAR* Thing)
 {
 	TCHAR ClassName[256]=TEXT("");
 
@@ -1395,7 +1395,7 @@ FClass* FHeaderParser::GetQualifiedClass(const FClasses& AllClasses, const TCHAR
 		FError::Throwf(TEXT("%s: Missing class name"), Thing );
 	}
 
-	return AllClasses.FindScriptClassOrThrow(ClassName);
+	return FClasses::FindScriptClassOrThrow(ClassName);
 }
 
 /*-----------------------------------------------------------------------------
@@ -2394,7 +2394,7 @@ EAccessSpecifier FHeaderParser::ParseAccessProtectionSpecifier(const FToken& Tok
 /**
  * Compile a struct definition.
  */
-UScriptStruct* FHeaderParser::CompileStructDeclaration(FClasses& AllClasses)
+UScriptStruct* FHeaderParser::CompileStructDeclaration()
 {
 	FUnrealSourceFile* CurrentSrcFile = GetCurrentSourceFile();
 	TSharedPtr<FFileScope> Scope = CurrentSrcFile->GetScope();
@@ -2538,11 +2538,11 @@ UScriptStruct* FHeaderParser::CompileStructDeclaration(FClasses& AllClasses)
 
 					ParentStructNameInScript = FString(ParentName.Identifier);
 					FString ParentNameStripped = GetClassNameWithPrefixRemoved(ParentScope.Identifier);
-					FClass* StructClass = AllClasses.FindClass(*ParentNameStripped);
+					FClass* StructClass = FClasses::FindClass(*ParentNameStripped);
 					if( !StructClass )
 					{
 						// If we find the literal class name, the user didn't use a prefix
-						StructClass = AllClasses.FindClass(ParentScope.Identifier);
+						StructClass = FClasses::FindClass(ParentScope.Identifier);
 						if( StructClass )
 						{
 							FError::Throwf(TEXT("'struct': Parent struct class '%s' is missing a prefix, expecting '%s'"), ParentScope.Identifier, *FString::Printf(TEXT("%s%s"),StructClass->GetPrefixCPP(),ParentScope.Identifier) );
@@ -2697,7 +2697,7 @@ UScriptStruct* FHeaderParser::CompileStructDeclaration(FClasses& AllClasses)
 		}
 		else if (Token.Matches(TEXT("UPROPERTY"), ESearchCase::CaseSensitive))
 		{
-			CompileVariableDeclaration(AllClasses, Struct);
+			CompileVariableDeclaration(Struct);
 		}
 		else if (Token.Matches(TEXT("UFUNCTION"), ESearchCase::CaseSensitive))
 		{
@@ -2705,7 +2705,7 @@ UScriptStruct* FHeaderParser::CompileStructDeclaration(FClasses& AllClasses)
 		}
 		else if (Token.Matches(TEXT("RIGVM_METHOD"), ESearchCase::CaseSensitive))
 		{
-			CompileRigVMMethodDeclaration(AllClasses, Struct);
+			CompileRigVMMethodDeclaration(Struct);
 		}
 		else if (Token.Matches(TEXT("GENERATED_USTRUCT_BODY"), ESearchCase::CaseSensitive) || Token.Matches(TEXT("GENERATED_BODY"), ESearchCase::CaseSensitive))
 		{
@@ -3056,7 +3056,7 @@ void FHeaderParser::PopNest(ENestType NestType, const TCHAR* Descr)
 	}
 }
 
-void FHeaderParser::FixupDelegateProperties( FClasses& AllClasses, UStruct* Struct, FScope& Scope, TMap<FName, UFunction*>& DelegateCache )
+void FHeaderParser::FixupDelegateProperties( UStruct* Struct, FScope& Scope, TMap<FName, UFunction*>& DelegateCache )
 {
 	check(Struct);
 
@@ -3130,7 +3130,7 @@ void FHeaderParser::FixupDelegateProperties( FClasses& AllClasses, UStruct* Stru
 						}
 
 						// make sure that the class that contains the delegate can be referenced here
-						UClass* DelegateOwnerClass = AllClasses.FindScriptClassOrThrow(DelegateClassName);
+						UClass* DelegateOwnerClass = FClasses::FindScriptClassOrThrow(DelegateClassName);
 						if (FScope::GetTypeScope(DelegateOwnerClass)->FindTypeByName(*DelegateName) != nullptr)
 						{
 							FError::Throwf(TEXT("Inaccessible type: '%s'"), *DelegateOwnerClass->GetPathName());
@@ -3214,7 +3214,7 @@ void FHeaderParser::FixupDelegateProperties( FClasses& AllClasses, UStruct* Stru
 		UStruct* InternalStruct = Cast<UStruct>(Field);
 		if ( InternalStruct != NULL )
 		{
-			FixupDelegateProperties(AllClasses, InternalStruct, Scope, DelegateCache);
+			FixupDelegateProperties(InternalStruct, Scope, DelegateCache);
 		}
 	}
 
@@ -3500,7 +3500,7 @@ void FHeaderParser::VerifyPropertyMarkups( UClass* TargetClass )
 //
 // Process a compiler directive.
 //
-void FHeaderParser::CompileDirective(FClasses& AllClasses)
+void FHeaderParser::CompileDirective()
 {
 	FUnrealSourceFile* CurrentSourceFilePtr = GetCurrentSourceFile();
 	TSharedRef<FUnrealSourceFile> CurrentSrcFile = CurrentSourceFilePtr->AsShared();
@@ -3642,7 +3642,6 @@ void FHeaderParser::CompileDirective(FClasses& AllClasses)
 -----------------------------------------------------------------------------*/
 
 void FHeaderParser::GetVarType(
-	FClasses&                       AllClasses,
 	FScope*                         Scope,
 	FPropertyBase&                  VarProperty,
 	EPropertyFlags                  Disallow,
@@ -4363,7 +4362,7 @@ void FHeaderParser::GetVarType(
 
 		VarType.PropertyFlags = Flags;
 
-		GetVarType(AllClasses, Scope, VarProperty, Disallow, &VarType, EPropertyDeclarationStyle::None, VariableCategory);
+		GetVarType(Scope, VarProperty, Disallow, &VarType, EPropertyDeclarationStyle::None, VariableCategory);
 		if (VarProperty.IsContainer())
 		{
 			FError::Throwf(TEXT("Nested containers are not supported.") );
@@ -4453,7 +4452,7 @@ void FHeaderParser::GetVarType(
 		VarType.PropertyFlags = Flags;
 
 		FToken MapKeyType;
-		GetVarType(AllClasses, Scope, MapKeyType, Disallow, &VarType, EPropertyDeclarationStyle::None, VariableCategory);
+		GetVarType(Scope, MapKeyType, Disallow, &VarType, EPropertyDeclarationStyle::None, VariableCategory);
 		if (MapKeyType.IsContainer())
 		{
 			FError::Throwf(TEXT("Nested containers are not supported.") );
@@ -4481,7 +4480,7 @@ void FHeaderParser::GetVarType(
 			FError::Throwf(TEXT("Missing value type while parsing TMap."));
 		}
 
-		GetVarType(AllClasses, Scope, VarProperty, Disallow, &VarType, EPropertyDeclarationStyle::None, VariableCategory);
+		GetVarType(Scope, VarProperty, Disallow, &VarType, EPropertyDeclarationStyle::None, VariableCategory);
 		if (VarProperty.IsContainer())
 		{
 			FError::Throwf(TEXT("Nested containers are not supported.") );
@@ -4537,7 +4536,7 @@ void FHeaderParser::GetVarType(
 
 		VarType.PropertyFlags = Flags;
 
-		GetVarType(AllClasses, Scope, VarProperty, Disallow, &VarType, EPropertyDeclarationStyle::None, VariableCategory);
+		GetVarType(Scope, VarProperty, Disallow, &VarType, EPropertyDeclarationStyle::None, VariableCategory);
 		if (VarProperty.IsContainer())
 		{
 			FError::Throwf(TEXT("Nested containers are not supported.") );
@@ -4746,7 +4745,7 @@ void FHeaderParser::GetVarType(
 			FToken DelegateName;
 			if (GetIdentifier(DelegateName))
 			{
-				UClass* LocalOwnerClass = AllClasses.FindClass(*IdentifierStripped);
+				UClass* LocalOwnerClass = FClasses::FindClass(*IdentifierStripped);
 				if (LocalOwnerClass)
 				{
 					TSharedRef<FScope> LocScope = FScope::GetTypeScope(LocalOwnerClass);
@@ -4853,7 +4852,7 @@ void FHeaderParser::GetVarType(
 					{
 						RedirectTypeIdentifier(InnerClass);
 
-						TempClass = AllClasses.FindScriptClass(InnerClass.Identifier);
+						TempClass = FClasses::FindScriptClass(InnerClass.Identifier);
 						if (TempClass == nullptr)
 						{
 							FError::Throwf(TEXT("Unrecognized type '%s' (in expression %s<%s>) - type must be a UCLASS"), InnerClass.Identifier, VarType.Identifier, InnerClass.Identifier);
@@ -4895,7 +4894,7 @@ void FHeaderParser::GetVarType(
 				}
 				else
 				{
-					TempClass = AllClasses.FindScriptClass(VarType.Identifier);
+					TempClass = FClasses::FindScriptClass(VarType.Identifier);
 				}
 			}
 
@@ -4933,7 +4932,7 @@ void FHeaderParser::GetVarType(
 
 						RedirectTypeIdentifier(Limitor);
 
-						VarProperty.MetaClass = AllClasses.FindScriptClassOrThrow(Limitor.Identifier);
+						VarProperty.MetaClass = FClasses::FindScriptClassOrThrow(Limitor.Identifier);
 
 						RequireSymbol( TEXT('>'), TEXT("'class limitor'"), ESymbolParseOption::CloseTemplateBracket );
 					}
@@ -5610,7 +5609,7 @@ FProperty* FHeaderParser::GetVarNameAndDim
 //
 // Compile a declaration in Token. Returns 1 if compiled, 0 if not.
 //
-bool FHeaderParser::CompileDeclaration(FClasses& AllClasses, TArray<UDelegateFunction*>& DelegatesToFixup, FToken& Token)
+bool FHeaderParser::CompileDeclaration(TArray<UDelegateFunction*>& DelegatesToFixup, FToken& Token)
 {
 	EAccessSpecifier AccessSpecifier = ParseAccessProtectionSpecifier(Token);
 	if (AccessSpecifier)
@@ -5636,7 +5635,7 @@ bool FHeaderParser::CompileDeclaration(FClasses& AllClasses, TArray<UDelegateFun
 		bEncounteredNewStyleClass_UnmatchedBrackets = true;
 		CurrentAccessSpecifier = ACCESS_Private;
 
-		if (!TryParseIInterfaceClass(AllClasses))
+		if (!TryParseIInterfaceClass())
 		{
 			bEncounteredNewStyleClass_UnmatchedBrackets = false;
 			UngetToken(Token);
@@ -5751,7 +5750,7 @@ bool FHeaderParser::CompileDeclaration(FClasses& AllClasses, TArray<UDelegateFun
 	{
 		bHaveSeenUClass = true;
 		bEncounteredNewStyleClass_UnmatchedBrackets = true;
-		UClass* Class = CompileClassDeclaration(AllClasses);
+		UClass* Class = CompileClassDeclaration();
 		GStructToSourceLine.Add(Class, MakeTuple(GetCurrentSourceFile()->AsShared(), Token.StartLine));
 		return true;
 	}
@@ -5760,26 +5759,26 @@ bool FHeaderParser::CompileDeclaration(FClasses& AllClasses, TArray<UDelegateFun
 	{
 		bHaveSeenUClass = true;
 		bEncounteredNewStyleClass_UnmatchedBrackets = true;
-		CompileInterfaceDeclaration(AllClasses);
+		CompileInterfaceDeclaration();
 		return true;
 	}
 
 	if (Token.Matches(TEXT("UFUNCTION"), ESearchCase::CaseSensitive))
 	{
-		CompileFunctionDeclaration(AllClasses);
+		CompileFunctionDeclaration();
 		return true;
 	}
 
 	if (Token.Matches(TEXT("UDELEGATE"), ESearchCase::CaseSensitive))
 	{
-		UDelegateFunction* Delegate = CompileDelegateDeclaration(AllClasses, Token.Identifier, EDelegateSpecifierAction::Parse);
+		UDelegateFunction* Delegate = CompileDelegateDeclaration(Token.Identifier, EDelegateSpecifierAction::Parse);
 		DelegatesToFixup.Add(Delegate);
 		return true;
 	}
 
 	if (IsValidDelegateDeclaration(Token)) // Legacy delegate parsing - it didn't need a UDELEGATE
 	{
-		UDelegateFunction* Delegate = CompileDelegateDeclaration(AllClasses, Token.Identifier);
+		UDelegateFunction* Delegate = CompileDelegateDeclaration(Token.Identifier);
 		DelegatesToFixup.Add(Delegate);
 		return true;
 	}
@@ -5789,7 +5788,7 @@ bool FHeaderParser::CompileDeclaration(FClasses& AllClasses, TArray<UDelegateFun
 		CheckAllow(TEXT("'Member variable declaration'"), ENestAllowFlags::VarDecl);
 		check(TopNest->NestType == ENestType::Class);
 
-		CompileVariableDeclaration(AllClasses, GetCurrentClass());
+		CompileVariableDeclaration(GetCurrentClass());
 		return true;
 	}
 
@@ -5803,7 +5802,7 @@ bool FHeaderParser::CompileDeclaration(FClasses& AllClasses, TArray<UDelegateFun
 	if (Token.Matches(TEXT("USTRUCT"), ESearchCase::CaseSensitive))
 	{
 		// Struct definition.
-		UScriptStruct* Struct = CompileStructDeclaration(AllClasses);
+		UScriptStruct* Struct = CompileStructDeclaration();
 		GStructToSourceLine.Add(Struct, MakeTuple(GetCurrentSourceFile()->AsShared(), Token.StartLine));
 		return true;
 	}
@@ -5811,7 +5810,7 @@ bool FHeaderParser::CompileDeclaration(FClasses& AllClasses, TArray<UDelegateFun
 	if (Token.Matches(TEXT('#')))
 	{
 		// Compiler directive.
-		CompileDirective(AllClasses);
+		CompileDirective();
 		return true;
 	}
 
@@ -5835,7 +5834,7 @@ bool FHeaderParser::CompileDeclaration(FClasses& AllClasses, TArray<UDelegateFun
 		{
 			checkf(TopNest->NestType == ENestType::Interface || TopNest->NestType == ENestType::NativeInterface, TEXT("Unexpected end of interface block."));
 			PopNest(TopNest->NestType, TEXT("'Interface'"));
-			PostPopNestInterface(AllClasses, CurrentClass);
+			PostPopNestInterface(CurrentClass);
 
 			// Ensure the UINTERFACE classes have a GENERATED_BODY declaration
 			if (bHaveSeenUClass && !bClassHasGeneratedUInterfaceBody)
@@ -6174,12 +6173,12 @@ bool FHeaderParser::SafeMatchSymbol( const TCHAR Match )
 	return false;
 }
 
-FClass* FHeaderParser::ParseClassNameDeclaration(FClasses& AllClasses, FString& DeclaredClassName, FString& RequiredAPIMacroIfPresent)
+FClass* FHeaderParser::ParseClassNameDeclaration(FString& DeclaredClassName, FString& RequiredAPIMacroIfPresent)
 {
 	FUnrealSourceFile* CurrentSrcFile = GetCurrentSourceFile();
 	ParseNameWithPotentialAPIMacroPrefix(/*out*/ DeclaredClassName, /*out*/ RequiredAPIMacroIfPresent, TEXT("class"));
 
-	FClass* FoundClass = AllClasses.FindClass(*GetClassNameWithPrefixRemoved(*DeclaredClassName));
+	FClass* FoundClass = FClasses::FindClass(*GetClassNameWithPrefixRemoved(*DeclaredClassName));
 	check(FoundClass);
 
 	FClassMetaData* ClassMetaData = GScriptHelper.AddClassData(FoundClass, CurrentSrcFile);
@@ -6202,7 +6201,7 @@ FClass* FHeaderParser::ParseClassNameDeclaration(FClasses& AllClasses, FString&
 	if (bSpecifiesParentClass)
 	{
 		// Set the base class.
-		UClass* TempClass = GetQualifiedClass(AllClasses, TEXT("'extends'"));
+		UClass* TempClass = GetQualifiedClass(TEXT("'extends'"));
 		check(TempClass);
 		// a class cannot 'extends' an interface, use 'implements'
 		if (TempClass->ClassFlags & CLASS_Interface)
@@ -6280,7 +6279,7 @@ FClass* FHeaderParser::ParseClassNameDeclaration(FClasses& AllClasses, FString&
 				break;
 			}
 
-			HandleOneInheritedClass(AllClasses, FoundClass, MoveTemp(InterfaceName));
+			HandleOneInheritedClass(FoundClass, MoveTemp(InterfaceName));
 		}
 	}
 	else if (FoundClass->GetSuperClass())
@@ -6291,11 +6290,11 @@ FClass* FHeaderParser::ParseClassNameDeclaration(FClasses& AllClasses, FString&
 	return FoundClass;
 }
 
-void FHeaderParser::HandleOneInheritedClass(FClasses& AllClasses, UClass* Class, FString&& InterfaceName)
+void FHeaderParser::HandleOneInheritedClass(UClass* Class, FString&& InterfaceName)
 {
 	FUnrealSourceFile* CurrentSrcFile = GetCurrentSourceFile();
 	// Check for UInterface derived interface inheritance
-	if (UClass* Interface = AllClasses.FindScriptClass(InterfaceName))
+	if (UClass* Interface = FClasses::FindScriptClass(InterfaceName))
 	{
 		// Try to find the interface
 		if ( !Interface->HasAnyClassFlags(CLASS_Interface) )
@@ -6371,7 +6370,7 @@ void PostParsingClassSetup(UClass* Class)
 /**
  * Compiles a class declaration.
  */
-UClass* FHeaderParser::CompileClassDeclaration(FClasses& AllClasses)
+UClass* FHeaderParser::CompileClassDeclaration()
 {
 	// Start of a class block.
 	CheckAllow(TEXT("'class'"), ENestAllowFlags::Class);
@@ -6398,7 +6397,7 @@ UClass* FHeaderParser::CompileClassDeclaration(FClasses& AllClasses)
 	FString DeclaredClassName;
 	FString RequiredAPIMacroIfPresent;
 	
-	FClass* Class = ParseClassNameDeclaration(AllClasses, /*out*/ DeclaredClassName, /*out*/ RequiredAPIMacroIfPresent);
+	FClass* Class = ParseClassNameDeclaration(/*out*/ DeclaredClassName, /*out*/ RequiredAPIMacroIfPresent);
 	check(Class);
 	FClassDeclarationMetaData* ClassDeclarationData = &GClassDeclarations.FindChecked(Class->GetFName());
 
@@ -6438,7 +6437,7 @@ UClass* FHeaderParser::CompileClassDeclaration(FClasses& AllClasses)
 	check(ClassData);
 	ClassData->SetPrologLine(PrologFinishLine);
 
-	ClassDeclarationData->MergeAndValidateClassFlags(DeclaredClassName, PrevClassFlags, Class, AllClasses);
+	ClassDeclarationData->MergeAndValidateClassFlags(DeclaredClassName, PrevClassFlags, Class);
 	Class->SetInternalFlags(EInternalObjectFlags::Native);
 
 	// Class metadata
@@ -6499,11 +6498,11 @@ UClass* FHeaderParser::CompileClassDeclaration(FClasses& AllClasses)
 	return Class;
 }
 
-FClass* FHeaderParser::ParseInterfaceNameDeclaration(FClasses& AllClasses, FString& DeclaredInterfaceName, FString& RequiredAPIMacroIfPresent)
+FClass* FHeaderParser::ParseInterfaceNameDeclaration(FString& DeclaredInterfaceName, FString& RequiredAPIMacroIfPresent)
 {
 	ParseNameWithPotentialAPIMacroPrefix(/*out*/ DeclaredInterfaceName, /*out*/ RequiredAPIMacroIfPresent, TEXT("interface"));
 
-	FClass* FoundClass = AllClasses.FindClass(*GetClassNameWithPrefixRemoved(*DeclaredInterfaceName));
+	FClass* FoundClass = FClasses::FindClass(*GetClassNameWithPrefixRemoved(*DeclaredInterfaceName));
 	if (FoundClass == nullptr)
 	{
 		return nullptr;
@@ -6520,7 +6519,7 @@ FClass* FHeaderParser::ParseInterfaceNameDeclaration(FClasses& AllClasses, FStri
 
 	// verify if our super class is an interface class
 	// the super class should have been marked as CLASS_Interface at the importing stage, if it were an interface
-	UClass* TempClass = GetQualifiedClass(AllClasses, TEXT("'extends'"));
+	UClass* TempClass = GetQualifiedClass(TEXT("'extends'"));
 	check(TempClass);
 	if( !(TempClass->ClassFlags & CLASS_Interface) )
 	{
@@ -6542,14 +6541,14 @@ FClass* FHeaderParser::ParseInterfaceNameDeclaration(FClasses& AllClasses, FStri
 	return FoundClass;
 }
 
-bool FHeaderParser::TryParseIInterfaceClass(FClasses& AllClasses)
+bool FHeaderParser::TryParseIInterfaceClass()
 {
 	// 'class' was already matched by the caller
 
 	// Get a class name
 	FString DeclaredInterfaceName;
 	FString RequiredAPIMacroIfPresent;
-	if (ParseInterfaceNameDeclaration(AllClasses, /*out*/ DeclaredInterfaceName, /*out*/ RequiredAPIMacroIfPresent) == nullptr)
+	if (ParseInterfaceNameDeclaration(/*out*/ DeclaredInterfaceName, /*out*/ RequiredAPIMacroIfPresent) == nullptr)
 	{
 		return false;
 	}
@@ -6566,7 +6565,7 @@ bool FHeaderParser::TryParseIInterfaceClass(FClasses& AllClasses)
 	}
 
 	UClass* FoundClass = nullptr;
-	if ((FoundClass = AllClasses.FindClass(*DeclaredInterfaceName.Mid(1))) == nullptr)
+	if ((FoundClass = FClasses::FindClass(*DeclaredInterfaceName.Mid(1))) == nullptr)
 	{
 		return false;
 	}
@@ -6583,7 +6582,7 @@ bool FHeaderParser::TryParseIInterfaceClass(FClasses& AllClasses)
 /**
  *  compiles Java or C# style interface declaration
  */
-void FHeaderParser::CompileInterfaceDeclaration(FClasses& AllClasses)
+void FHeaderParser::CompileInterfaceDeclaration()
 {
 	FUnrealSourceFile* CurrentSrcFile = GetCurrentSourceFile();
 	// Start of an interface block. Since Interfaces and Classes are always at the same nesting level,
@@ -6604,7 +6603,7 @@ void FHeaderParser::CompileInterfaceDeclaration(FClasses& AllClasses)
 
 	// New style files have the interface name / extends afterwards
 	RequireIdentifier(TEXT("class"), ESearchCase::CaseSensitive, TEXT("Interface declaration"));
-	FClass* InterfaceClass = ParseInterfaceNameDeclaration(AllClasses, /*out*/ DeclaredInterfaceName, /*out*/ RequiredAPIMacroIfPresent);
+	FClass* InterfaceClass = ParseInterfaceNameDeclaration(/*out*/ DeclaredInterfaceName, /*out*/ RequiredAPIMacroIfPresent);
 	ClassDefinitionRanges.Add(InterfaceClass, ClassDefinitionRange(&Input[InputPos], nullptr));
 
 	// Record that this interface is RequiredAPI if the CORE_API style macro was present
@@ -6690,7 +6689,7 @@ void FHeaderParser::CompileInterfaceDeclaration(FClasses& AllClasses)
 	PushNest(ENestType::Interface, InterfaceClass);
 }
 
-void FHeaderParser::CompileRigVMMethodDeclaration(FClasses& AllClasses, UStruct* Struct)
+void FHeaderParser::CompileRigVMMethodDeclaration(UStruct* Struct)
 {
 	if (!MatchSymbol(TEXT("(")))
 	{
@@ -6938,7 +6937,7 @@ void FHeaderParser::RedirectTypeIdentifier(FToken& Token) const
 }
 
 // Parse the parameter list of a function or delegate declaration
-void FHeaderParser::ParseParameterList(FClasses& AllClasses, UFunction* Function, bool bExpectCommaBeforeName, TMap<FName, FString>* MetaData)
+void FHeaderParser::ParseParameterList(UFunction* Function, bool bExpectCommaBeforeName, TMap<FName, FString>* MetaData)
 {
 	// Get parameter list.
 	if (MatchSymbol(TEXT(')')))
@@ -6952,7 +6951,7 @@ void FHeaderParser::ParseParameterList(FClasses& AllClasses, UFunction* Function
 		// Get parameter type.
 		FToken Property(CPT_None);
 		EVariableCategory::Type VariableCategory = (Function->FunctionFlags & FUNC_Net) ? EVariableCategory::ReplicatedParameter : EVariableCategory::RegularParameter;
-		GetVarType(AllClasses, GetCurrentScope(), Property, ~(CPF_ParmFlags | CPF_AutoWeak | CPF_RepSkip | CPF_UObjectWrapper | CPF_NativeAccessSpecifiers), NULL, EPropertyDeclarationStyle::None, VariableCategory);
+		GetVarType(GetCurrentScope(), Property, ~(CPF_ParmFlags | CPF_AutoWeak | CPF_RepSkip | CPF_UObjectWrapper | CPF_NativeAccessSpecifiers), NULL, EPropertyDeclarationStyle::None, VariableCategory);
 		Property.PropertyFlags |= CPF_Parm;
 
 		if (bExpectCommaBeforeName)
@@ -7121,7 +7120,7 @@ void FHeaderParser::ParseParameterList(FClasses& AllClasses, UFunction* Function
 	} while( MatchSymbol(TEXT(',')) );
 	RequireSymbol( TEXT(')'), TEXT("parameter list") );
 }
-UDelegateFunction* FHeaderParser::CompileDelegateDeclaration(FClasses& AllClasses, const TCHAR* DelegateIdentifier, EDelegateSpecifierAction::Type SpecifierAction)
+UDelegateFunction* FHeaderParser::CompileDelegateDeclaration(const TCHAR* DelegateIdentifier, EDelegateSpecifierAction::Type SpecifierAction)
 {
 	const TCHAR* CurrentScopeName = TEXT("Delegate Declaration");
 
@@ -7205,7 +7204,7 @@ UDelegateFunction* FHeaderParser::CompileDelegateDeclaration(FClasses& AllClasse
 
 	if (bHasReturnValue)
 	{
-		GetVarType(AllClasses, GetCurrentScope(), ReturnType, CPF_None, nullptr, EPropertyDeclarationStyle::None, EVariableCategory::Return);
+		GetVarType(GetCurrentScope(), ReturnType, CPF_None, nullptr, EPropertyDeclarationStyle::None, EVariableCategory::Return);
 		RequireSymbol(TEXT(','), CurrentScopeName);
 	}
 
@@ -7291,7 +7290,7 @@ UDelegateFunction* FHeaderParser::CompileDelegateDeclaration(FClasses& AllClasse
 	{
 		RequireSymbol(TEXT(','), CurrentScopeName);
 
-		ParseParameterList(AllClasses, DelegateSignatureFunction, /*bExpectCommaBeforeName=*/ true);
+		ParseParameterList(DelegateSignatureFunction, /*bExpectCommaBeforeName=*/ true);
 
 		// Check the expected versus actual number of parameters
 		int32 ParamCount = UE_PTRDIFF_TO_INT32(FoundParamCount - DelegateParameterCountStrings.GetData()) + 1;
@@ -7336,7 +7335,7 @@ UDelegateFunction* FHeaderParser::CompileDelegateDeclaration(FClasses& AllClasse
 	DelegateSignatureFunction->Bind();
 
 	// End the nesting
-	PostPopFunctionDeclaration(AllClasses, DelegateSignatureFunction);
+	PostPopFunctionDeclaration(DelegateSignatureFunction);
 
 	// Don't allow delegate signatures to be redefined.
 	auto FunctionIterator = GetCurrentScope()->GetTypeIterator<UFunction>();
@@ -7429,7 +7428,7 @@ bool AreFunctionSignaturesEqual(const UFunction* Lhs, const UFunction* Rhs)
 /**
  * Parses and compiles a function declaration
  */
-void FHeaderParser::CompileFunctionDeclaration(FClasses& AllClasses)
+void FHeaderParser::CompileFunctionDeclaration()
 {
 	CheckAllow(TEXT("'Function'"), ENestAllowFlags::Function);
 
@@ -7680,7 +7679,7 @@ void FHeaderParser::CompileFunctionDeclaration(FClasses& AllClasses)
 	FToken ReturnType( CPT_None );
 
 	// C++ style functions always have a return value type, even if it's void
-	GetVarType(AllClasses, GetCurrentScope(), ReturnType, CPF_None, nullptr, EPropertyDeclarationStyle::None, EVariableCategory::Return);
+	GetVarType(GetCurrentScope(), ReturnType, CPF_None, nullptr, EPropertyDeclarationStyle::None, EVariableCategory::Return);
 	bool bHasReturnValue = ReturnType.Type != CPT_None;
 
 	// Skip whitespaces to get InputPos exactly on beginning of function name.
@@ -7756,7 +7755,7 @@ void FHeaderParser::CompileFunctionDeclaration(FClasses& AllClasses)
 	}
 
 	// Get parameter list.
-	ParseParameterList(AllClasses, TopFunction, false, &MetaData);
+	ParseParameterList(TopFunction, false, &MetaData);
 
 	// Get return type, if any.
 	if (bHasReturnValue)
@@ -7919,7 +7918,7 @@ void FHeaderParser::CompileFunctionDeclaration(FClasses& AllClasses)
 	}
 
 	// Just declaring a function, so end the nesting.
-	PostPopFunctionDeclaration(AllClasses, TopFunction);
+	PostPopFunctionDeclaration(TopFunction);
 
 	// See what's coming next
 	FToken Token;
@@ -8169,7 +8168,7 @@ struct FExposeOnSpawnValidator
 	}
 };
 
-void FHeaderParser::CompileVariableDeclaration(FClasses& AllClasses, UStruct* Struct)
+void FHeaderParser::CompileVariableDeclaration(UStruct* Struct)
 {
 	EPropertyFlags DisallowFlags = CPF_ParmFlags;
 	EPropertyFlags EdFlags       = CPF_None;
@@ -8178,7 +8177,7 @@ void FHeaderParser::CompileVariableDeclaration(FClasses& AllClasses, UStruct* St
 	FPropertyBase OriginalProperty(CPT_None);
 	FIndexRange TypeRange;
 	ELayoutMacroType LayoutMacroType = ELayoutMacroType::None;
-	GetVarType( AllClasses, &FScope::GetTypeScope(Struct).Get(), OriginalProperty, DisallowFlags, /*OuterPropertyType=*/ NULL, EPropertyDeclarationStyle::UPROPERTY, EVariableCategory::Member, &TypeRange, &LayoutMacroType);
+	GetVarType( &FScope::GetTypeScope(Struct).Get(), OriginalProperty, DisallowFlags, /*OuterPropertyType=*/ NULL, EPropertyDeclarationStyle::UPROPERTY, EVariableCategory::Member, &TypeRange, &LayoutMacroType);
 	OriginalProperty.PropertyFlags |= EdFlags;
 
 	FString* Category = OriginalProperty.MetaData.Find(NAME_Category);
@@ -8418,7 +8417,7 @@ void FHeaderParser::CompileVariableDeclaration(FClasses& AllClasses, UStruct* St
 // Compile a statement: Either a declaration or a command.
 // Returns 1 if success, 0 if end of file.
 //
-bool FHeaderParser::CompileStatement(FClasses& AllClasses, TArray<UDelegateFunction*>& DelegatesToFixup)
+bool FHeaderParser::CompileStatement(TArray<UDelegateFunction*>& DelegatesToFixup)
 {
 	// Get a token and compile it.
 	FToken Token;
@@ -8427,7 +8426,7 @@ bool FHeaderParser::CompileStatement(FClasses& AllClasses, TArray<UDelegateFunct
 		// End of file.
 		return false;
 	}
-	else if (!CompileDeclaration(AllClasses, DelegatesToFixup, Token))
+	else if (!CompileDeclaration(DelegatesToFixup, Token))
 	{
 		FError::Throwf(TEXT("'%s': Bad command or expression"), Token.Identifier );
 	}
@@ -8565,7 +8564,7 @@ void FHeaderParser::FinalizeScriptExposedFunctions(UClass* Class)
 // Parses the header associated with the specified class.
 // Returns result enumeration.
 //
-ECompilationResult::Type FHeaderParser::ParseHeader(FClasses& AllClasses, FUnrealSourceFile* SourceFile)
+ECompilationResult::Type FHeaderParser::ParseHeader(FUnrealSourceFile* SourceFile)
 {
 	SetCurrentSourceFile(SourceFile);
 	FUnrealSourceFile* CurrentSrcFile = SourceFile;
@@ -8618,7 +8617,7 @@ ECompilationResult::Type FHeaderParser::ParseHeader(FClasses& AllClasses, FUnrea
 	{
 		// Parse entire program.
 		TArray<UDelegateFunction*> DelegatesToFixup;
-		while (CompileStatement(AllClasses, DelegatesToFixup))
+		while (CompileStatement(DelegatesToFixup))
 		{
 			bEmptyFile = false;
 
@@ -8643,7 +8642,7 @@ ECompilationResult::Type FHeaderParser::ParseHeader(FClasses& AllClasses, FUnrea
 
 			// now validate all delegate variables declared in the class
 			TMap<FName, UFunction*> DelegateCache;
-			FixupDelegateProperties(AllClasses, Struct, FScope::GetTypeScope(Struct).Get(), DelegateCache);
+			FixupDelegateProperties(Struct, FScope::GetTypeScope(Struct).Get(), DelegateCache);
 		}
 
 		// Fix up any delegates themselves, if they refer to other delegates
@@ -8651,7 +8650,7 @@ ECompilationResult::Type FHeaderParser::ParseHeader(FClasses& AllClasses, FUnrea
 			TMap<FName, UFunction*> DelegateCache;
 			for (UDelegateFunction* Delegate : DelegatesToFixup)
 			{
-				FixupDelegateProperties(AllClasses, Delegate, CurrentSrcFile->GetScope().Get(), DelegateCache);
+				FixupDelegateProperties(Delegate, CurrentSrcFile->GetScope().Get(), DelegateCache);
 			}
 		}
 
@@ -8742,7 +8741,7 @@ ECompilationResult::Type FHeaderParser::ParseHeader(FClasses& AllClasses, FUnrea
 	Global functions.
 -----------------------------------------------------------------------------*/
 
-ECompilationResult::Type FHeaderParser::ParseRestOfModulesSourceFiles(FClasses& AllClasses, UPackage* ModulePackage, FHeaderParser& HeaderParser)
+ECompilationResult::Type FHeaderParser::ParseRestOfModulesSourceFiles(UPackage* ModulePackage, FHeaderParser& HeaderParser)
 {
 	if (const TArray<FUnrealSourceFile*>* SourceFiles = GUnrealSourceFilesMap.FindFilesForPackage(ModulePackage))
 	{
@@ -8751,7 +8750,7 @@ ECompilationResult::Type FHeaderParser::ParseRestOfModulesSourceFiles(FClasses&
 			if (!SourceFile->IsParsed() || SourceFile->GetDefinedClassesCount() == 0)
 			{
 				ECompilationResult::Type Result;
-				if ((Result = ParseHeaders(AllClasses, HeaderParser, SourceFile)) != ECompilationResult::Succeeded)
+				if ((Result = ParseHeaders(HeaderParser, SourceFile)) != ECompilationResult::Succeeded)
 				{
 					return Result;
 				}
@@ -8765,7 +8764,7 @@ ECompilationResult::Type FHeaderParser::ParseRestOfModulesSourceFiles(FClasses&
 // Parse Class's annotated headers and optionally its child classes.
 static const FString ObjectHeader(TEXT("NoExportTypes.h"));
 
-ECompilationResult::Type FHeaderParser::ParseHeaders(FClasses& AllClasses, FHeaderParser& HeaderParser, FUnrealSourceFile* SourceFile)
+ECompilationResult::Type FHeaderParser::ParseHeaders(FHeaderParser& HeaderParser, FUnrealSourceFile* SourceFile)
 {
 	ECompilationResult::Type Result = ECompilationResult::Succeeded;
 
@@ -8804,7 +8803,7 @@ ECompilationResult::Type FHeaderParser::ParseHeaders(FClasses& AllClasses, FHead
 	{
 		SourceFile->GetScope()->IncludeScope(&RequiredFile->GetScope().Get());
 
-		ECompilationResult::Type ParseResult = ParseHeaders(AllClasses, HeaderParser, RequiredFile);
+		ECompilationResult::Type ParseResult = ParseHeaders(HeaderParser, RequiredFile);
 
 		if (ParseResult != ECompilationResult::Succeeded)
 		{
@@ -8814,7 +8813,7 @@ ECompilationResult::Type FHeaderParser::ParseHeaders(FClasses& AllClasses, FHead
 
 	// Parse the file
 	{
-		ECompilationResult::Type OneFileResult = HeaderParser.ParseHeader(AllClasses, SourceFile);
+		ECompilationResult::Type OneFileResult = HeaderParser.ParseHeader(SourceFile);
 
 		for (const TPair<UClass*, FSimplifiedParsingClassInfo>& ClassDataPair : SourceFile->GetDefinedClassesWithParsingInfo())
 		{
@@ -9113,7 +9112,7 @@ ECompilationResult::Type FHeaderParser::ParseAllHeadersInside(
 			{
 				if ((!SourceFile->IsParsed() || SourceFile->GetDefinedClassesCount() == 0))
 				{
-					Result = ParseHeaders(ModuleClasses, HeaderParser, SourceFile);
+					Result = ParseHeaders(HeaderParser, SourceFile);
 					if (Result != ECompilationResult::Succeeded)
 					{
 						return Result;
@@ -9123,7 +9122,7 @@ ECompilationResult::Type FHeaderParser::ParseAllHeadersInside(
 		}
 		if (Result == ECompilationResult::Succeeded)
 		{
-			Result = FHeaderParser::ParseRestOfModulesSourceFiles(ModuleClasses, CurrentPackage, HeaderParser);
+			Result = FHeaderParser::ParseRestOfModulesSourceFiles(CurrentPackage, HeaderParser);
 		}
 
 		if (Result == ECompilationResult::Succeeded)
@@ -10542,25 +10541,25 @@ void FHeaderParser::PostPopNestClass(UClass* CurrentClass)
 	}
 }
 
-void FHeaderParser::PostPopFunctionDeclaration(FClasses& AllClasses, UFunction* PoppedFunction)
+void FHeaderParser::PostPopFunctionDeclaration(UFunction* PoppedFunction)
 {
 	//@TODO: UCREMOVAL: Move this code to occur at delegate var declaration, and force delegates to be declared before variables that use them
 	if (!GetCurrentScope()->IsFileScope() && GetCurrentClassData()->ContainsDelegates())
 	{
 		// now validate all delegate variables declared in the class
 		TMap<FName, UFunction*> DelegateCache;
-		FixupDelegateProperties(AllClasses, PoppedFunction, *GetCurrentScope(), DelegateCache);
+		FixupDelegateProperties(PoppedFunction, *GetCurrentScope(), DelegateCache);
 	}
 }
 
-void FHeaderParser::PostPopNestInterface(FClasses& AllClasses, UClass* CurrentInterface)
+void FHeaderParser::PostPopNestInterface(UClass* CurrentInterface)
 {
 	FClassMetaData* ClassData = GScriptHelper.FindClassData(CurrentInterface);
 	check(ClassData);
 	if (ClassData->ContainsDelegates())
 	{
 		TMap<FName, UFunction*> DelegateCache;
-		FixupDelegateProperties(AllClasses, CurrentInterface, FScope::GetTypeScope(ExactCast<UClass>(CurrentInterface)).Get(), DelegateCache);
+		FixupDelegateProperties(CurrentInterface, FScope::GetTypeScope(ExactCast<UClass>(CurrentInterface)).Get(), DelegateCache);
 	}
 }
 
diff --git a/Engine/Source/Programs/UnrealHeaderTool/Private/HeaderParser.h b/Engine/Source/Programs/UnrealHeaderTool/Private/HeaderParser.h
--- a/Engine/Source/Programs/UnrealHeaderTool/Private/HeaderParser.h
+++ b/Engine/Source/Programs/UnrealHeaderTool/Private/HeaderParser.h
@@ -494,13 +494,12 @@ public:
 	/**
 	 * Parse Class's annotated headers and optionally its child classes.  Marks the class as CLASS_Parsed.
 	 *
-	 * @param	AllClasses			the class tree containing all classes in the current package
 	 * @param	HeaderParser		the header parser
 	 * @param	SourceFile			Source file info.
 	 *
 	 * @return	Result enumeration.
 	 */
-	static ECompilationResult::Type ParseHeaders(FClasses& AllClasses, FHeaderParser& HeaderParser, FUnrealSourceFile* SourceFile);
+	static ECompilationResult::Type ParseHeaders(FHeaderParser& HeaderParser, FUnrealSourceFile* SourceFile);
 
 protected:
 	friend struct FScriptLocation;
@@ -702,7 +701,7 @@ protected:
 	bool IsBitfieldProperty(ELayoutMacroType LayoutMacroType);
 
 	// Parse the parameter list of a function or delegate declaration
-	void ParseParameterList(FClasses& AllClasses, UFunction* Function, bool bExpectCommaBeforeName = false, TMap<FName, FString>* MetaData = NULL);
+	void ParseParameterList(UFunction* Function, bool bExpectCommaBeforeName = false, TMap<FName, FString>* MetaData = NULL);
 
 	// Modify token to fix redirected types if needed
 	void RedirectTypeIdentifier(FToken& Token) const;
@@ -735,13 +734,12 @@ protected:
 	/**
 	 * Parse rest of the module's source files.
 	 *
-	 * @param	AllClasses The class tree containing all classes in the current package.
 	 * @param	ModulePackage Current package.
 	 * @param	HeaderParser The header parser.
 	 *
 	 * @return	Result enumeration.
 	 */
-	static ECompilationResult::Type ParseRestOfModulesSourceFiles(FClasses& AllClasses, UPackage* ModulePackage, FHeaderParser& HeaderParser);
+	static ECompilationResult::Type ParseRestOfModulesSourceFiles(UPackage* ModulePackage, FHeaderParser& HeaderParser);
 
 	//@TODO: Remove this method
 	static void ParseClassName(const TCHAR* Temp, FString& ClassName);
@@ -783,24 +781,23 @@ protected:
 	/**
 	 * Parses given source file.
 	 *
-	 * @param AllClasses The class tree for current package.
 	 * @param SourceFile Source file to parse.
 	 *
 	 * @returns Compilation result enum.
 	 */
-	ECompilationResult::Type ParseHeader(FClasses& AllClasses, FUnrealSourceFile* SourceFile);
-	void CompileDirective(FClasses& AllClasses);
+	ECompilationResult::Type ParseHeader(FUnrealSourceFile* SourceFile);
+	void CompileDirective();
 	void FinalizeScriptExposedFunctions(UClass* Class);
 	UEnum* CompileEnum();
-	UScriptStruct* CompileStructDeclaration(FClasses& AllClasses);
-	bool CompileDeclaration(FClasses& AllClasses, TArray<UDelegateFunction*>& DelegatesToFixup, FToken& Token);
+	UScriptStruct* CompileStructDeclaration();
+	bool CompileDeclaration(TArray<UDelegateFunction*>& DelegatesToFixup, FToken& Token);
 
 	/** Skip C++ (noexport) declaration. */
 	bool SkipDeclaration(FToken& Token);
 	/** Similar to MatchSymbol() but will return to the exact location as on entry if the symbol was not found. */
 	bool SafeMatchSymbol(const TCHAR Match);
-	void HandleOneInheritedClass(FClasses& AllClasses, UClass* Class, FString&& InterfaceName);
-	FClass* ParseClassNameDeclaration(FClasses& AllClasses, FString& DeclaredClassName, FString& RequiredAPIMacroIfPresent);
+	void HandleOneInheritedClass(UClass* Class, FString&& InterfaceName);
+	FClass* ParseClassNameDeclaration(FString& DeclaredClassName, FString& RequiredAPIMacroIfPresent);
 
 	/** The property style of a variable declaration being parsed */
 	struct EPropertyDeclarationStyle
@@ -828,18 +825,18 @@ protected:
 	template<typename T>
 	UDelegateFunction* CreateDelegateFunction(const FFuncInfo &FuncInfo) const;	
 
-	UClass* CompileClassDeclaration(FClasses& AllClasses);
-	UDelegateFunction* CompileDelegateDeclaration(FClasses& AllClasses, const TCHAR* DelegateIdentifier, EDelegateSpecifierAction::Type SpecifierAction = EDelegateSpecifierAction::DontParse);
-	void CompileFunctionDeclaration(FClasses& AllClasses);
-	void CompileVariableDeclaration (FClasses& AllClasses, UStruct* Struct);
-	void CompileInterfaceDeclaration(FClasses& AllClasses);
-	void CompileRigVMMethodDeclaration(FClasses& AllClasses, UStruct* Struct);
+	UClass* CompileClassDeclaration();
+	UDelegateFunction* CompileDelegateDeclaration(const TCHAR* DelegateIdentifier, EDelegateSpecifierAction::Type SpecifierAction = EDelegateSpecifierAction::DontParse);
+	void CompileFunctionDeclaration();
+	void CompileVariableDeclaration (UStruct* Struct);
+	void CompileInterfaceDeclaration();
+	void CompileRigVMMethodDeclaration(UStruct* Struct);
 	void ParseRigVMMethodParameters(UStruct* Struct);
 
-	FClass* ParseInterfaceNameDeclaration(FClasses& AllClasses, FString& DeclaredInterfaceName, FString& RequiredAPIMacroIfPresent);
-	bool TryParseIInterfaceClass(FClasses& AllClasses);
+	FClass* ParseInterfaceNameDeclaration(FString& DeclaredInterfaceName, FString& RequiredAPIMacroIfPresent);
+	bool TryParseIInterfaceClass();
 
-	bool CompileStatement(FClasses& AllClasses, TArray<UDelegateFunction*>& DelegatesToFixup);
+	bool CompileStatement(TArray<UDelegateFunction*>& DelegatesToFixup);
 
 	// Checks to see if a particular kind of command is allowed on this nesting level.
 	bool IsAllowedInThisNesting(ENestAllowFlags AllowFlags);
@@ -856,7 +853,6 @@ protected:
 	/**
 	 * Parses a variable or return value declaration and determines the variable type and property flags.
 	 *
-	 * @param   AllClasses                the class tree for CurrentPackage
 	 * @param   Scope                     struct to create the property in
 	 * @param   VarProperty               will be filled in with type and property flag data for the property declaration that was parsed
 	 * @param   Disallow                  contains a mask of variable modifiers that are disallowed in this context
@@ -866,7 +862,6 @@ protected:
 	 * @param   ParsedVarIndexRange       The source text [Start, End) index range for the parsed type.
 	 */
 	void GetVarType(
-		FClasses&                       AllClasses,
 		FScope*							Scope,
 		FPropertyBase&                  VarProperty,
 		EPropertyFlags                  Disallow,
@@ -932,7 +927,7 @@ protected:
 
 	const TCHAR* NestTypeName( ENestType NestType );
 
-	FClass* GetQualifiedClass(const FClasses& AllClasses, const TCHAR* Thing);
+	FClass* GetQualifiedClass(const TCHAR* Thing);
 
 	/**
 	 * Increase the nesting level, setting the new top nesting level to
@@ -949,19 +944,17 @@ protected:
 	 * Tasks that need to be done after popping function declaration
 	 * from parsing stack.
 	 *
-	 * @param AllClasses The class tree for current package.
 	 * @param PoppedFunction Function that have just been popped.
 	 */
-	void PostPopFunctionDeclaration(FClasses& AllClasses, UFunction* PoppedFunction);
+	void PostPopFunctionDeclaration(UFunction* PoppedFunction);
 
 	/**
 	 * Tasks that need to be done after popping interface definition
 	 * from parsing stack.
 	 *
-	 * @param AllClasses The class tree for current package.
 	 * @param CurrentInterface Interface that have just been popped.
 	 */
-	void PostPopNestInterface(FClasses& AllClasses, UClass* CurrentInterface);
+	void PostPopNestInterface(UClass* CurrentInterface);
 
 	/**
 	 * Tasks that need to be done after popping class definition
@@ -977,12 +970,11 @@ protected:
 	 *
 	 * @todo: this function will no longer be required once the post-parse fixup phase is added (TTPRO #13256)
 	 *
-	 * @param	AllClasses			the class tree for CurrentPackage
 	 * @param	Struct				the struct to validate delegate properties for
 	 * @param	Scope				the current scope
 	 * @param	DelegateCache		cached map of delegates that have already been found; used for faster lookup.
 	 */
-	void FixupDelegateProperties(FClasses& AllClasses, UStruct* ValidationScope, FScope& Scope, TMap<FName, UFunction*>& DelegateCache);
+	void FixupDelegateProperties(UStruct* ValidationScope, FScope& Scope, TMap<FName, UFunction*>& DelegateCache);
 
 	// Retry functions.
 	void InitScriptLocation( FScriptLocation& Retry );
