Move the initialization of constant configuration data into a class explicitly initialized at startup.
This avoids a race condition with lazy initialization.

#rnx

diff --git a/Engine/Source/Programs/UnrealHeaderTool/Private/HeaderParser.cpp b/Engine/Source/Programs/UnrealHeaderTool/Private/HeaderParser.cpp
--- a/Engine/Source/Programs/UnrealHeaderTool/Private/HeaderParser.cpp
+++ b/Engine/Source/Programs/UnrealHeaderTool/Private/HeaderParser.cpp
@@ -106,14 +106,7 @@ const FName FHeaderParserNames::NAME_SparseClassDataTypes(TEXT("SparseClassDataT
 const FName FHeaderParserNames::NAME_IsConversionRoot(TEXT("IsConversionRoot"));
 const FName FHeaderParserNames::NAME_AdvancedClassDisplay(TEXT("AdvancedClassDisplay"));
 
-EGeneratedCodeVersion FHeaderParser::DefaultGeneratedCodeVersion = EGeneratedCodeVersion::V1;
-EPointerMemberBehavior FHeaderParser::NativePointerMemberBehavior = EPointerMemberBehavior::AllowSilently;
-EPointerMemberBehavior FHeaderParser::ObjectPtrMemberBehavior = EPointerMemberBehavior::AllowSilently;
-TArray<FString> FHeaderParser::StructsWithNoPrefix;
-TArray<FString> FHeaderParser::StructsWithTPrefix;
 FRigVMStructMap FHeaderParser::StructRigVMMap;
-TArray<FString> FHeaderParser::DelegateParameterCountStrings;
-TMap<FString, FString> FHeaderParser::TypeRedirectMap;
 TArray<FString> FHeaderParser::PropertyCPPTypesRequiringUIRanges = { TEXT("float"), TEXT("double") };
 TArray<FString> FHeaderParser::ReservedTypeNames = { TEXT("none") };
 TMap<UClass*, ClassDefinitionRange> ClassDefinitionRanges;
@@ -1353,6 +1346,69 @@ namespace
 	};
 }
 
+FUHTConfig::FUHTConfig()
+{
+	// Read Ini options, GConfig must exist by this point
+	check(GConfig);
+
+	const FName TypeRedirectsKey(TEXT("TypeRedirects"));
+	const FName StructsWithNoPrefixKey(TEXT("StructsWithNoPrefix"));
+	const FName StructsWithTPrefixKey(TEXT("StructsWithTPrefix"));
+	const FName DelegateParameterCountStringsKey(TEXT("DelegateParameterCountStrings"));
+	const FName GeneratedCodeVersionKey(TEXT("GeneratedCodeVersion"));
+	const FName NativePointerMemberBehaviorKey(TEXT("NativePointerMemberBehavior"));
+	const FName ObjectPtrMemberBehaviorKey(TEXT("ObjectPtrMemberBehavior"));
+
+	FConfigSection* ConfigSection = GConfig->GetSectionPrivate(TEXT("UnrealHeaderTool"), false, true, GEngineIni);
+	if (ConfigSection)
+	{
+		for (FConfigSection::TIterator It(*ConfigSection); It; ++It)
+		{
+			if (It.Key() == TypeRedirectsKey)
+			{
+				FString OldType;
+				FString NewType;
+
+				FParse::Value(*It.Value().GetValue(), TEXT("OldType="), OldType);
+				FParse::Value(*It.Value().GetValue(), TEXT("NewType="), NewType);
+
+				TypeRedirectMap.Add(MoveTemp(OldType), MoveTemp(NewType));
+			}
+			else if (It.Key() == StructsWithNoPrefixKey)
+			{
+				StructsWithNoPrefix.Add(It.Value().GetValue());
+			}
+			else if (It.Key() == StructsWithTPrefixKey)
+			{
+				StructsWithTPrefix.Add(It.Value().GetValue());
+			}
+			else if (It.Key() == DelegateParameterCountStringsKey)
+			{
+				DelegateParameterCountStrings.Add(It.Value().GetValue());
+			}
+			else if (It.Key() == GeneratedCodeVersionKey)
+			{
+				DefaultGeneratedCodeVersion = ToGeneratedCodeVersion(It.Value().GetValue());
+			}
+			else if (It.Key() == NativePointerMemberBehaviorKey)
+			{
+				NativePointerMemberBehavior = ToPointerMemberBehavior(It.Value().GetValue());
+			}
+			else if (It.Key() == ObjectPtrMemberBehaviorKey)
+			{
+				ObjectPtrMemberBehavior = ToPointerMemberBehavior(It.Value().GetValue());
+			}
+		}
+	}
+}
+
+const FUHTConfig& FUHTConfig::Get()
+{
+	static FUHTConfig UHTConfig;
+	return UHTConfig;
+}
+
+
 /////////////////////////////////////////////////////
 // FHeaderParser
 
@@ -2565,7 +2621,7 @@ UScriptStruct* FHeaderParser::CompileStructDeclaration()
 			const UField* Type = nullptr;
 			bool bOverrideParentStructName = false;
 
-			if( !StructsWithNoPrefix.Contains(ParentStructNameInScript) )
+			if( !UHTConfig.StructsWithNoPrefix.Contains(ParentStructNameInScript) )
 			{
 				bOverrideParentStructName = true;
 				ParentStructNameStripped = GetClassNameWithPrefixRemoved(ParentStructNameInScript);
@@ -2608,7 +2664,7 @@ UScriptStruct* FHeaderParser::CompileStructDeclaration()
 				BaseStruct = ((UScriptStruct*)Type);
 				if( bOverrideParentStructName )
 				{
-					const TCHAR* PrefixCPP = StructsWithTPrefix.Contains(ParentStructNameStripped) ? TEXT("T") : BaseStruct->GetPrefixCPP();
+					const TCHAR* PrefixCPP = UHTConfig.StructsWithTPrefix.Contains(ParentStructNameStripped) ? TEXT("T") : BaseStruct->GetPrefixCPP();
 					if( ParentStructNameInScript != FString::Printf(TEXT("%s%s"), PrefixCPP, *ParentStructNameStripped) )
 					{
 						BaseStruct = nullptr;
@@ -2643,7 +2699,7 @@ UScriptStruct* FHeaderParser::CompileStructDeclaration()
 	if( DeclaredPrefix == Struct->GetPrefixCPP() || DeclaredPrefix == TEXT("T") )
 	{
 		// Found a prefix, do a basic check to see if it's valid
-		const TCHAR* ExpectedPrefixCPP = StructsWithTPrefix.Contains(StructNameStripped) ? TEXT("T") : Struct->GetPrefixCPP();
+		const TCHAR* ExpectedPrefixCPP = UHTConfig.StructsWithTPrefix.Contains(StructNameStripped) ? TEXT("T") : Struct->GetPrefixCPP();
 		FString ExpectedStructName = FString::Printf(TEXT("%s%s"), ExpectedPrefixCPP, *StructNameStripped);
 		if (StructNameInScript != ExpectedStructName)
 		{
@@ -2652,7 +2708,7 @@ UScriptStruct* FHeaderParser::CompileStructDeclaration()
 	}
 	else
 	{
-		const TCHAR* ExpectedPrefixCPP = StructsWithTPrefix.Contains(StructNameInScript) ? TEXT("T") : Struct->GetPrefixCPP();
+		const TCHAR* ExpectedPrefixCPP = UHTConfig.StructsWithTPrefix.Contains(StructNameInScript) ? TEXT("T") : Struct->GetPrefixCPP();
 		FString ExpectedStructName = FString::Printf(TEXT("%s%s"), ExpectedPrefixCPP, *StructNameInScript);
 		FError::Throwf(TEXT("Struct '%s' is missing a valid Unreal prefix, expecting '%s'"), *StructNameInScript, *ExpectedStructName);
 	}
@@ -4770,16 +4826,16 @@ void FHeaderParser::GetVarType(
 		{
 			if (bStripped)
 			{
-				const TCHAR* PrefixCPP = StructsWithTPrefix.Contains(IdentifierStripped) ? TEXT("T") : Struct->GetPrefixCPP();
+				const TCHAR* PrefixCPP = UHTConfig.StructsWithTPrefix.Contains(IdentifierStripped) ? TEXT("T") : Struct->GetPrefixCPP();
 				FString ExpectedStructName = FString::Printf(TEXT("%s%s"), PrefixCPP, *Struct->GetName() );
 				if( FString(VarType.Identifier) != ExpectedStructName )
 				{
 					FError::Throwf( TEXT("Struct '%s' is missing or has an incorrect prefix, expecting '%s'"), VarType.Identifier, *ExpectedStructName );
 				}
 			}
-			else if( !StructsWithNoPrefix.Contains(VarType.Identifier) )
+			else if( !UHTConfig.StructsWithNoPrefix.Contains(VarType.Identifier) )
 			{
-				const TCHAR* PrefixCPP = StructsWithTPrefix.Contains(VarType.Identifier) ? TEXT("T") : Struct->GetPrefixCPP();
+				const TCHAR* PrefixCPP = UHTConfig.StructsWithTPrefix.Contains(VarType.Identifier) ? TEXT("T") : Struct->GetPrefixCPP();
 				FError::Throwf(TEXT("Struct '%s' is missing a prefix, expecting '%s'"), VarType.Identifier, *FString::Printf(TEXT("%s%s"), PrefixCPP, *Struct->GetName()) );
 			}
 
@@ -4973,7 +5029,7 @@ void FHeaderParser::GetVarType(
 					// Optionally emit messages about native pointer members and swallow trailing 'const' after pointer properties
 					if (VariableCategory == EVariableCategory::Member)
 					{
-						ConditionalLogPointerUsage(NativePointerMemberBehavior, TEXT("Native pointer"), FString(InputPos - VarStartPos, Input + VarStartPos).TrimStartAndEnd().ReplaceCharWithEscapedChar());
+						ConditionalLogPointerUsage(UHTConfig.NativePointerMemberBehavior, TEXT("Native pointer"), FString(InputPos - VarStartPos, Input + VarStartPos).TrimStartAndEnd().ReplaceCharWithEscapedChar());
 
 						MatchIdentifier(TEXT("const"), ESearchCase::CaseSensitive);
 					}
@@ -4982,7 +5038,7 @@ void FHeaderParser::GetVarType(
 				}
 				else if ((PropertyType == CPT_ObjectPtrReference) && (VariableCategory == EVariableCategory::Member))
 				{
-					ConditionalLogPointerUsage(ObjectPtrMemberBehavior, TEXT("ObjectPtr"), FString(InputPos - VarStartPos, Input + VarStartPos).TrimStartAndEnd().ReplaceCharWithEscapedChar());
+					ConditionalLogPointerUsage(UHTConfig.ObjectPtrMemberBehavior, TEXT("ObjectPtr"), FString(InputPos - VarStartPos, Input + VarStartPos).TrimStartAndEnd().ReplaceCharWithEscapedChar());
 				}
 
 				// Imply const if it's a parameter that is a pointer to a const class
@@ -6929,7 +6985,7 @@ void FHeaderParser::RedirectTypeIdentifier(FToken& Token) const
 {
 	check(Token.TokenType == TOKEN_Identifier);
 
-	FString* FoundRedirect = TypeRedirectMap.Find(Token.Identifier);
+	const FString* FoundRedirect = UHTConfig.TypeRedirectMap.Find(Token.Identifier);
 	if (FoundRedirect)
 	{
 		Token.SetIdentifier(**FoundRedirect);
@@ -7167,7 +7223,7 @@ UDelegateFunction* FHeaderParser::CompileDelegateDeclaration(const TCHAR* Delega
 	const bool bIsSparse       = DelegateMacro.Contains(TEXT("_SPARSE"), ESearchCase::CaseSensitive);
 
 	// Determine the parameter count
-	const FString* FoundParamCount = DelegateParameterCountStrings.FindByPredicate([&](const FString& Str){ return DelegateMacro.Contains(Str); });
+	const FString* FoundParamCount = UHTConfig.DelegateParameterCountStrings.FindByPredicate([&](const FString& Str){ return DelegateMacro.Contains(Str); });
 
 	// Try reconstructing the string to make sure it matches our expectations
 	FString ExpectedOriginalString = FString::Printf(TEXT("DECLARE_DYNAMIC%s%s_DELEGATE%s%s%s"),
@@ -7293,7 +7349,7 @@ UDelegateFunction* FHeaderParser::CompileDelegateDeclaration(const TCHAR* Delega
 		ParseParameterList(DelegateSignatureFunction, /*bExpectCommaBeforeName=*/ true);
 
 		// Check the expected versus actual number of parameters
-		int32 ParamCount = UE_PTRDIFF_TO_INT32(FoundParamCount - DelegateParameterCountStrings.GetData()) + 1;
+		int32 ParamCount = UE_PTRDIFF_TO_INT32(FoundParamCount - UHTConfig.DelegateParameterCountStrings.GetData()) + 1;
 		if (DelegateSignatureFunction->NumParms != ParamCount)
 		{
 			FError::Throwf(TEXT("Expected %d parameters but found %d parameters"), ParamCount, DelegateSignatureFunction->NumParms);
@@ -8927,6 +8983,7 @@ FHeaderParser::FHeaderParser(FFeedbackContext* InWarn, const FManifestModule& In
 	, bSpottedAutogeneratedHeaderInclude(false)
 	, NestLevel(0)
 	, TopNest(nullptr)
+	, UHTConfig(FUHTConfig::Get())
 	, CurrentlyParsedModule(&InModule)
 {
 	// Determine if the current module is part of the engine or a game (we are more strict about things for Engine modules)
@@ -8957,65 +9014,6 @@ FHeaderParser::FHeaderParser(FFeedbackContext* InWarn, const FManifestModule& In
 		bIsCurrentModulePartOfEngine = true;
 		check(false);
 	}
-
-	static bool bConfigOptionsInitialized = false;
-
-	if (!bConfigOptionsInitialized)
-	{
-		// Read Ini options, GConfig must exist by this point
-		check(GConfig);
-
-		const FName TypeRedirectsKey(TEXT("TypeRedirects"));
-		const FName StructsWithNoPrefixKey(TEXT("StructsWithNoPrefix"));
-		const FName StructsWithTPrefixKey(TEXT("StructsWithTPrefix"));
-		const FName DelegateParameterCountStringsKey(TEXT("DelegateParameterCountStrings"));
-		const FName GeneratedCodeVersionKey(TEXT("GeneratedCodeVersion"));
-		const FName NativePointerMemberBehaviorKey(TEXT("NativePointerMemberBehavior"));
-		const FName ObjectPtrMemberBehaviorKey(TEXT("ObjectPtrMemberBehavior"));
-
-		FConfigSection* ConfigSection = GConfig->GetSectionPrivate(TEXT("UnrealHeaderTool"), false, true, GEngineIni);
-		if (ConfigSection)
-		{
-			for (FConfigSection::TIterator It(*ConfigSection); It; ++It)
-			{
-				if (It.Key() == TypeRedirectsKey)
-				{
-					FString OldType;
-					FString NewType;
-
-					FParse::Value(*It.Value().GetValue(), TEXT("OldType="), OldType);
-					FParse::Value(*It.Value().GetValue(), TEXT("NewType="), NewType);
-
-					TypeRedirectMap.Add(MoveTemp(OldType), MoveTemp(NewType));
-				}
-				else if (It.Key() == StructsWithNoPrefixKey)
-				{
-					StructsWithNoPrefix.Add(It.Value().GetValue());
-				}
-				else if (It.Key() == StructsWithTPrefixKey)
-				{
-					StructsWithTPrefix.Add(It.Value().GetValue());
-				}
-				else if (It.Key() == DelegateParameterCountStringsKey)
-				{
-					DelegateParameterCountStrings.Add(It.Value().GetValue());
-				}
-				else if (It.Key() == GeneratedCodeVersionKey)
-				{
-					DefaultGeneratedCodeVersion = ToGeneratedCodeVersion(It.Value().GetValue());
-				}
-				else if (It.Key() == NativePointerMemberBehaviorKey)
-				{
-					NativePointerMemberBehavior = ToPointerMemberBehavior(It.Value().GetValue());
-				}
-				else if (It.Key() == ObjectPtrMemberBehaviorKey)
-				{
-					ObjectPtrMemberBehavior = ToPointerMemberBehavior(It.Value().GetValue());
-				}
-			}
-		}
-		bConfigOptionsInitialized = true;
-	}
 }
 
 // Throws if a specifier value wasn't provided
@@ -10387,7 +10385,7 @@ void FHeaderParser::CompileVersionDeclaration(UStruct* Struct)
 	}
 
 	// Default version based on config file.
-	EGeneratedCodeVersion Version = DefaultGeneratedCodeVersion;
+	EGeneratedCodeVersion Version = UHTConfig.DefaultGeneratedCodeVersion;
 
 	// Overwrite with module-specific value if one was specified.
 	if (CurrentlyParsedModule->GeneratedCodeVersion != EGeneratedCodeVersion::None)
@@ -10870,7 +10868,7 @@ bool FHeaderParser::CheckUIMinMaxRangeFromMetaData(const FString& UIMin, const F
 
 void FHeaderParser::ConditionalLogPointerUsage(EPointerMemberBehavior PointerMemberBehavior, const TCHAR* PointerTypeDesc, FString&& PointerTypeDecl)
 {
-	switch (NativePointerMemberBehavior)
+	switch (UHTConfig.NativePointerMemberBehavior)
 	{
 		case EPointerMemberBehavior::Disallow:
 		{
diff --git a/Engine/Source/Programs/UnrealHeaderTool/Private/HeaderParser.h b/Engine/Source/Programs/UnrealHeaderTool/Private/HeaderParser.h
--- a/Engine/Source/Programs/UnrealHeaderTool/Private/HeaderParser.h
+++ b/Engine/Source/Programs/UnrealHeaderTool/Private/HeaderParser.h
@@ -426,6 +426,36 @@ enum class EPointerMemberBehavior
 	AllowAndLog,
 };
 
+/////////////////////////////////////////////////////
+// UHTConfig
+
+struct FUHTConfig
+{
+	FUHTConfig();
+
+	static const FUHTConfig& Get();
+
+	// Types that have been renamed, treat the old deprecated name as the new name for code generation
+	TMap<FString, FString> TypeRedirectMap;
+
+	// Special parsed struct names that do not require a prefix
+	TArray<FString> StructsWithNoPrefix;
+
+	// Special parsed struct names that have a 'T' prefix
+	TArray<FString> StructsWithTPrefix;
+
+	// Mapping from 'human-readable' macro substring to # of parameters for delegate declarations
+	// Index 0 is 1 parameter, Index 1 is 2, etc...
+	TArray<FString> DelegateParameterCountStrings;
+
+	// Default version of generated code. Defaults to oldest possible, unless specified otherwise in config.
+	EGeneratedCodeVersion DefaultGeneratedCodeVersion = EGeneratedCodeVersion::V1;
+
+	EPointerMemberBehavior NativePointerMemberBehavior = EPointerMemberBehavior::AllowSilently;
+
+	EPointerMemberBehavior ObjectPtrMemberBehavior = EPointerMemberBehavior::AllowSilently;
+};
+
 /////////////////////////////////////////////////////
 // FHeaderParser
 
@@ -436,12 +466,6 @@ enum class EPointerMemberBehavior
 class FHeaderParser : public FBaseParser, public FContextSupplier
 {
 public:
-	// Default version of generated code. Defaults to oldest possible, unless specified otherwise in config.
-	static EGeneratedCodeVersion DefaultGeneratedCodeVersion;
-
-	static EPointerMemberBehavior NativePointerMemberBehavior;
-	static EPointerMemberBehavior ObjectPtrMemberBehavior;
-
 	// Compute the function parameter size and save the return offset
 	static void ComputeFunctionParametersSize(UClass* InClass);
 
@@ -670,18 +694,8 @@ protected:
 
 	////////////////////////////////////////////////////
 
-	// Special parsed struct names that do not require a prefix
-	static TArray<FString> StructsWithNoPrefix;
-	
-	// Special parsed struct names that have a 'T' prefix
-	static TArray<FString> StructsWithTPrefix;
-
-	// Mapping from 'human-readable' macro substring to # of parameters for delegate declarations
-	// Index 0 is 1 parameter, Index 1 is 2, etc...
-	static TArray<FString> DelegateParameterCountStrings;
-
-	// Types that have been renamed, treat the old deprecated name as the new name for code generation
-	static TMap<FString, FString> TypeRedirectMap;
+	// UHTConfig data
+	const FUHTConfig& UHTConfig;
 
 	// List of all used identifiers for net service function declarations (every function must be unique)
 	TMap<int32, FString> UsedRPCIds;
diff --git a/Engine/Source/Programs/UnrealHeaderTool/Private/UnrealSourceFile.cpp b/Engine/Source/Programs/UnrealHeaderTool/Private/UnrealSourceFile.cpp
--- a/Engine/Source/Programs/UnrealHeaderTool/Private/UnrealSourceFile.cpp
+++ b/Engine/Source/Programs/UnrealHeaderTool/Private/UnrealSourceFile.cpp
@@ -116,7 +116,7 @@ EGeneratedCodeVersion FUnrealSourceFile::GetGeneratedCodeVersionForStruct(UStruc
 		return *Version;
 	}
 
-	return FHeaderParser::DefaultGeneratedCodeVersion;
+	return FUHTConfig::Get().DefaultGeneratedCodeVersion;
 }
 
 void FUnrealSourceFile::MarkDependenciesResolved()
