Provide a more general exception class.
New mechanism for maintaining the overall result.
Helper methods for making try blocks.

Work in progress

#rnx

diff --git a/Engine/Source/Programs/UnrealHeaderTool/Private/CodeGenerator.cpp b/Engine/Source/Programs/UnrealHeaderTool/Private/CodeGenerator.cpp
--- a/Engine/Source/Programs/UnrealHeaderTool/Private/CodeGenerator.cpp
+++ b/Engine/Source/Programs/UnrealHeaderTool/Private/CodeGenerator.cpp
@@ -58,7 +58,7 @@
 #include "Misc/ScopeExit.h"
 #include "UnrealTypeDefinitionInfo.h"
 
-#include "FileLineException.h"
+#include "Exceptions.h"
 #include "UObject/FieldIterator.h"
 #include "UObject/FieldPath.h"
 
@@ -6445,11 +6445,6 @@ void FNativeClassHeaderGenerator::DeleteUnusedGeneratedHeaders(TSet<FString>&& P
 	GAsyncFileTasks.Add(FFunctionGraphTask::CreateAndDispatchWhenReady(MoveTemp(DeleteUnusedGeneratedHeadersTask), TStatId()));
 }
 
-/**
- * Dirty hack global variable to allow different result codes passed through
- * exceptions. Needs to be fixed in future versions of UHT.
- */
-ECompilationResult::Type GCompilationResult = ECompilationResult::OtherCompilationError;
 FCriticalSection TestCommandLineCS;
 
 bool FNativeClassHeaderGenerator::SaveHeaderIfChanged(FReferenceGatherers& OutReferenceGatherers, const FPreloadHeaderFileInfo& FileInfo, FString&& InNewHeaderContents, FGraphEventRef& OutSaveTaskRef) const
@@ -6546,7 +6541,7 @@ bool FNativeClassHeaderGenerator::SaveHeaderIfChanged(FReferenceGatherers& OutRe
 			FString ConflictPath = HeaderPathStr + TEXT(".conflict");
 			FFileHelper::SaveStringToFile(InNewHeaderContents, *ConflictPath);
 
-			GCompilationResult = ECompilationResult::FailedDueToHeaderChange;
+			FResults::SetResult(ECompilationResult::FailedDueToHeaderChange);
 			FError::Throwf(TEXT("ERROR: '%s': Changes to generated code are not allowed - conflicts written to '%s'"), *HeaderPathStr, *ConflictPath);
 		}
 
@@ -6804,7 +6799,7 @@ UPackage* GetModulePackage(FManifestModule& Module)
 
 //FGraphEventRef CreateModuleParse
 
-ECompilationResult::Type PreparseModules(const FString& ModuleInfoPath, int32& NumFailures)
+void PreparseModules(const FString& ModuleInfoPath)
 {
 	// Three passes.  1) Public 'Classes' headers (legacy)  2) Public headers   3) Private headers
 	enum EHeaderFolderTypes
@@ -6819,41 +6814,6 @@ ECompilationResult::Type PreparseModules(const FString& ModuleInfoPath, int32& N
 	// The meta data keywords must be initialized prior to going wide
 	FBaseParser::InitMetadataKeywords();
 
-	ECompilationResult::Type Result = ECompilationResult::Succeeded;
-
-#if !PLATFORM_EXCEPTIONS_DISABLED
-	FGraphEventArray ExceptionTasks;
-
-	auto LogException = [&Result, &NumFailures, &ExceptionTasks](FString&& Filename, int32 Line, const FString& Message)
-	{
-		auto LogExceptionTask = [&Result, &NumFailures, Filename = MoveTemp(Filename), Line, Message]()
-		{
-			TGuardValue<ELogTimes::Type> DisableLogTimes(GPrintLogTimes, ELogTimes::None);
-
-			FString FormattedErrorMessage = FString::Printf(TEXT("%s(%d): Error: %s\r\n"), *Filename, Line, *Message);
-			Result = ECompilationResult::OtherCompilationError;
-
-			UE_LOG(LogCompile, Log, TEXT("%s"), *FormattedErrorMessage);
-			GWarn->Log(ELogVerbosity::Error, FormattedErrorMessage);
-
-			++NumFailures;
-		};
-
-		if (IsInGameThread())
-		{
-			LogExceptionTask();
-		}
-		else
-		{
-			FGraphEventRef EventRef = FFunctionGraphTask::CreateAndDispatchWhenReady(MoveTemp(LogExceptionTask), TStatId(), nullptr, ENamedThreads::GameThread);
-
-			static FCriticalSection ExceptionCS;
-			FScopeLock Lock(&ExceptionCS);
-			ExceptionTasks.Add(EventRef);
-		}
-	};
-#endif
-
 	// Pre-initialization
 	int32 NumHeaders = 0;
 	for (FManifestModule& Module : GManifest.Modules)
@@ -6863,7 +6823,7 @@ ECompilationResult::Type PreparseModules(const FString& ModuleInfoPath, int32& N
 		GetModulePackage(Module);
 
 		// Count the headers
-		for (int32 PassIndex = 0; PassIndex < FolderType_Count && Result == ECompilationResult::Succeeded; ++PassIndex)
+		for (int32 PassIndex = 0; PassIndex < FolderType_Count; ++PassIndex)
 		{
 			EHeaderFolderTypes CurrentlyProcessing = (EHeaderFolderTypes)PassIndex;
 			const TArray<FString>& UObjectHeaders =
@@ -6892,7 +6852,7 @@ ECompilationResult::Type PreparseModules(const FString& ModuleInfoPath, int32& N
 			UPackage* Package = GetModulePackage(Module);
 
 			// Pre-parse the headers
-			for (int32 PassIndex = 0; PassIndex < FolderType_Count && Result == ECompilationResult::Succeeded; ++PassIndex)
+			for (int32 PassIndex = 0; PassIndex < FolderType_Count; ++PassIndex)
 			{
 				EHeaderFolderTypes CurrentlyProcessing = (EHeaderFolderTypes)PassIndex;
 
@@ -6923,64 +6883,31 @@ ECompilationResult::Type PreparseModules(const FString& ModuleInfoPath, int32& N
 				{
 
 					// Phase #1: Load the file
-					auto LoadLambda = [&UObjectHeaders, &HeaderFiles, &LogException, &ModuleInfoPath, &Result, Index]()
+					auto LoadLambda = [&UObjectHeaders, &HeaderFiles, &PerHeaderData, &ModuleInfoPath, Index]()
 					{
-						if (Result != ECompilationResult::Succeeded)
-						{
-							return;
-						}
-
-						const FString& RawFilename = UObjectHeaders[Index];
-
-#if !PLATFORM_EXCEPTIONS_DISABLED
-						try
-#endif
-						{
-							const FString FullFilename = FPaths::ConvertRelativePathToFull(ModuleInfoPath, RawFilename);
-
-							if (!FFileHelper::LoadFileToString(HeaderFiles[Index], *FullFilename))
+						FResults::Try([&UObjectHeaders, &HeaderFiles, &PerHeaderData, &ModuleInfoPath, Index]()
 							{
-								FError::Throwf(TEXT("UnrealHeaderTool was unable to load source file '%s'"), *FullFilename);
+								const FString& RawFilename = UObjectHeaders[Index];
+								const FString FullFilename = FPaths::ConvertRelativePathToFull(ModuleInfoPath, RawFilename);
+
+								if (!FFileHelper::LoadFileToString(HeaderFiles[Index], *FullFilename))
+								{
+									FUHTException::Throwf(*PerHeaderData[Index].UnrealSourceFile, 1, TEXT("Unable to load source file"));
+								}
 							}
-						}
-#if !PLATFORM_EXCEPTIONS_DISABLED
-						catch (TCHAR* ErrorMsg)
-						{
-							FString AbsFilename = IFileManager::Get().ConvertToAbsolutePathForExternalAppForRead(*RawFilename);
-							LogException(MoveTemp(AbsFilename), 1, ErrorMsg);
-						}
-#endif
+						);
 					};
 
-					// Phase #2: Perform simplified class parse (can run concurrenrtly)
-					auto PreProcessLambda = [Package, &UObjectHeaders, &HeaderFiles, &PerHeaderData, &LogException, &ModuleInfoPath, &Result, Index]()
+					// Phase #2: Perform simplified class parse (can run concurrently)
+					auto PreProcessLambda = [Package, &UObjectHeaders, &HeaderFiles, &PerHeaderData, Index]()
 					{
-						if (Result != ECompilationResult::Succeeded)
-						{
-							return;
-						}
-
-						const FString& RawFilename = UObjectHeaders[Index];
-
-#if !PLATFORM_EXCEPTIONS_DISABLED
-						try
-#endif
-						{
-							PerformSimplifiedClassParse(Package, *RawFilename, *HeaderFiles[Index], PerHeaderData[Index]);
-							HeaderFiles[Index].Empty();
-						}
-#if !PLATFORM_EXCEPTIONS_DISABLED
-						catch (const FFileLineException& Ex)
-						{
-							FString AbsFilename = IFileManager::Get().ConvertToAbsolutePathForExternalAppForRead(*Ex.Filename);
-							LogException(MoveTemp(AbsFilename), Ex.Line, Ex.Message);
-						}
-						catch (TCHAR* ErrorMsg)
-						{
-							FString AbsFilename = IFileManager::Get().ConvertToAbsolutePathForExternalAppForRead(*RawFilename);
-							LogException(MoveTemp(AbsFilename), 1, ErrorMsg);
-						}
-#endif
+						FResults::Try([Package, &UObjectHeaders, &HeaderFiles, &PerHeaderData, Index]()
+							{
+								const FString& RawFilename = UObjectHeaders[Index];
+								PerformSimplifiedClassParse(Package, *RawFilename, *HeaderFiles[Index], PerHeaderData[Index]);
+								HeaderFiles[Index].Empty();
+							}
+						);
 					};
 
 					FGraphEventRef LoadTask = FFunctionGraphTask::CreateAndDispatchWhenReady(MoveTemp(LoadLambda), TStatId());
@@ -6992,17 +6919,14 @@ ECompilationResult::Type PreparseModules(const FString& ModuleInfoPath, int32& N
 
 		// Wait for all the loading and preparsing to complete
 		FTaskGraphInterface::Get().WaitUntilTasksComplete(LoadTasks);
-
-#if !PLATFORM_EXCEPTIONS_DISABLED
-		FTaskGraphInterface::Get().WaitUntilTasksComplete(ExceptionTasks);
-#endif
+		FResults::WaitForErrorTasks();
 
 		PreparseTimer.Stop();
 		UE_LOG(LogCompile, Log, TEXT("Loaded and preparsed classes in %d modules containing %d files(s) in %.3f secs."), GManifest.Modules.Num(), NumHeaders, PreparseTime);
 	}
 
 	// Class declaration phase
-	if (Result == ECompilationResult::Succeeded)
+	if (FResults::IsSucceeding())
 	{
 		double PreparseTime = 0.0;
 		FDurationTimer PreparseTimer(PreparseTime);
@@ -7015,7 +6939,7 @@ ECompilationResult::Type PreparseModules(const FString& ModuleInfoPath, int32& N
 			UPackage* Package = GetModulePackage(Module);
 
 			// Pre-parse the headers
-			for (int32 PassIndex = 0; PassIndex < FolderType_Count && Result == ECompilationResult::Succeeded; ++PassIndex)
+			for (int32 PassIndex = 0; PassIndex < FolderType_Count && FResults::IsSucceeding(); ++PassIndex)
 			{
 				EHeaderFolderTypes CurrentlyProcessing = (EHeaderFolderTypes)PassIndex;
 
@@ -7042,114 +6966,93 @@ ECompilationResult::Type PreparseModules(const FString& ModuleInfoPath, int32& N
 					ModulePerHeaderData.PrivateUObjectHeadersData;
 				PerHeaderData.SetNum(UObjectHeaders.Num());
 
-				for (int32 Index = 0, EIndex = UObjectHeaders.Num(); Index < EIndex; ++Index)
+				for (int32 Index = 0, EIndex = UObjectHeaders.Num(); Index < EIndex && FResults::IsSucceeding(); ++Index)
 				{
-					if (Result != ECompilationResult::Succeeded)
-					{
-						break;
-					}
-
-					const FString& RawFilename = UObjectHeaders[Index];
-
-#if !PLATFORM_EXCEPTIONS_DISABLED
-					try
-#endif
-					{
-						// Import class.
-						const FString FullFilename = FPaths::ConvertRelativePathToFull(ModuleInfoPath, RawFilename);
-
-						ProcessInitialClassParse(PerHeaderData[Index]);
-						TSharedRef<FUnrealSourceFile> UnrealSourceFile = PerHeaderData[Index].UnrealSourceFile.ToSharedRef();
-						FUnrealSourceFile* UnrealSourceFilePtr = &UnrealSourceFile.Get();
-						FString CleanFilename = FPaths::GetCleanFilename(RawFilename);
-						uint32  CleanFilenameHash = GetTypeHash(CleanFilename);
-						if (const TSharedRef<FUnrealSourceFile>* ExistingSourceFile = GUnrealSourceFilesMap.FindByHash(CleanFilenameHash, CleanFilename))
+					FResults::Try([&Module, &UObjectHeaders, &HeaderFiles, &PerHeaderData, &ModuleInfoPath, CurrentlyProcessing, Index]
 						{
-							FString NormalizedFullFilename = FullFilename;
-							FString NormalizedExistingFilename = (*ExistingSourceFile)->GetFilename();
-
-							FPaths::NormalizeFilename(NormalizedFullFilename);
-							FPaths::NormalizeFilename(NormalizedExistingFilename);
+							// Import class.
+							const FString& RawFilename = UObjectHeaders[Index];
+							const FString FullFilename = FPaths::ConvertRelativePathToFull(ModuleInfoPath, RawFilename);
 
-							if (NormalizedFullFilename != NormalizedExistingFilename)
+							ProcessInitialClassParse(PerHeaderData[Index]);
+							TSharedRef<FUnrealSourceFile> UnrealSourceFile = PerHeaderData[Index].UnrealSourceFile.ToSharedRef();
+							FUnrealSourceFile* UnrealSourceFilePtr = &UnrealSourceFile.Get();
+							FString CleanFilename = FPaths::GetCleanFilename(RawFilename);
+							uint32  CleanFilenameHash = GetTypeHash(CleanFilename);
+							if (const TSharedRef<FUnrealSourceFile>* ExistingSourceFile = GUnrealSourceFilesMap.FindByHash(CleanFilenameHash, CleanFilename))
 							{
-								FError::Throwf(TEXT("Duplicate leaf header name found: %s (original: %s)"), *NormalizedFullFilename, *NormalizedExistingFilename);
-							}
-						}
-						GUnrealSourceFilesMap.AddByHash(CleanFilenameHash, MoveTemp(CleanFilename), UnrealSourceFile);
-
-						if (CurrentlyProcessing == PublicClassesHeaders)
-						{
-							GPublicSourceFileSet.Add(UnrealSourceFilePtr);
-						}
+								FString NormalizedFullFilename = FullFilename;
+								FString NormalizedExistingFilename = (*ExistingSourceFile)->GetFilename();
 
-						// Save metadata for the class path, both for it's include path and relative to the module base directory
-						if (FullFilename.StartsWith(Module.BaseDirectory))
-						{
-							// Get the path relative to the module directory
-							const TCHAR* ModuleRelativePath = *FullFilename + Module.BaseDirectory.Len();
-
-							UnrealSourceFilePtr->SetModuleRelativePath(ModuleRelativePath);
+								FPaths::NormalizeFilename(NormalizedFullFilename);
+								FPaths::NormalizeFilename(NormalizedExistingFilename);
 
-							// Calculate the include path
-							const TCHAR* IncludePath = ModuleRelativePath;
-
-							// Walk over the first potential slash
-							if (*IncludePath == TEXT('/'))
-							{
-								IncludePath++;
+								if (NormalizedFullFilename != NormalizedExistingFilename)
+								{
+									FUHTException::Throwf(*UnrealSourceFilePtr, 1, TEXT("Duplicate leaf header name found: %s (original: %s)"), *NormalizedFullFilename, *NormalizedExistingFilename);
+								}
 							}
+							GUnrealSourceFilesMap.AddByHash(CleanFilenameHash, MoveTemp(CleanFilename), UnrealSourceFile);
 
-							// Does this module path start with a known include path location? If so, we can cut that part out of the include path
-							static const TCHAR PublicFolderName[] = TEXT("Public/");
-							static const TCHAR PrivateFolderName[] = TEXT("Private/");
-							static const TCHAR ClassesFolderName[] = TEXT("Classes/");
-							if (FCString::Strnicmp(IncludePath, PublicFolderName, UE_ARRAY_COUNT(PublicFolderName) - 1) == 0)
+							if (CurrentlyProcessing == PublicClassesHeaders)
 							{
-								IncludePath += (UE_ARRAY_COUNT(PublicFolderName) - 1);
-							}
-							else if (FCString::Strnicmp(IncludePath, PrivateFolderName, UE_ARRAY_COUNT(PrivateFolderName) - 1) == 0)
-							{
-								IncludePath += (UE_ARRAY_COUNT(PrivateFolderName) - 1);
-							}
-							else if (FCString::Strnicmp(IncludePath, ClassesFolderName, UE_ARRAY_COUNT(ClassesFolderName) - 1) == 0)
-							{
-								IncludePath += (UE_ARRAY_COUNT(ClassesFolderName) - 1);
+								GPublicSourceFileSet.Add(UnrealSourceFilePtr);
 							}
 
-							// Add the include path
-							if (*IncludePath != 0)
+							// Save metadata for the class path, both for it's include path and relative to the module base directory
+							if (FullFilename.StartsWith(Module.BaseDirectory))
 							{
-								UnrealSourceFilePtr->SetIncludePath(MoveTemp(IncludePath));
+								// Get the path relative to the module directory
+								const TCHAR* ModuleRelativePath = *FullFilename + Module.BaseDirectory.Len();
+
+								UnrealSourceFilePtr->SetModuleRelativePath(ModuleRelativePath);
+
+								// Calculate the include path
+								const TCHAR* IncludePath = ModuleRelativePath;
+
+								// Walk over the first potential slash
+								if (*IncludePath == TEXT('/'))
+								{
+									IncludePath++;
+								}
+
+								// Does this module path start with a known include path location? If so, we can cut that part out of the include path
+								static const TCHAR PublicFolderName[] = TEXT("Public/");
+								static const TCHAR PrivateFolderName[] = TEXT("Private/");
+								static const TCHAR ClassesFolderName[] = TEXT("Classes/");
+								if (FCString::Strnicmp(IncludePath, PublicFolderName, UE_ARRAY_COUNT(PublicFolderName) - 1) == 0)
+								{
+									IncludePath += (UE_ARRAY_COUNT(PublicFolderName) - 1);
+								}
+								else if (FCString::Strnicmp(IncludePath, PrivateFolderName, UE_ARRAY_COUNT(PrivateFolderName) - 1) == 0)
+								{
+									IncludePath += (UE_ARRAY_COUNT(PrivateFolderName) - 1);
+								}
+								else if (FCString::Strnicmp(IncludePath, ClassesFolderName, UE_ARRAY_COUNT(ClassesFolderName) - 1) == 0)
+								{
+									IncludePath += (UE_ARRAY_COUNT(ClassesFolderName) - 1);
+								}
+
+								// Add the include path
+								if (*IncludePath != 0)
+								{
+									UnrealSourceFilePtr->SetIncludePath(MoveTemp(IncludePath));
+								}
 							}
 						}
-					}
-#if !PLATFORM_EXCEPTIONS_DISABLED
-					catch (const FFileLineException& Ex)
-					{
-						FString AbsFilename = IFileManager::Get().ConvertToAbsolutePathForExternalAppForRead(*Ex.Filename);
-						LogException(MoveTemp(AbsFilename), Ex.Line, Ex.Message);
-					}
-					catch (TCHAR* ErrorMsg)
-					{
-						FString AbsFilename = IFileManager::Get().ConvertToAbsolutePathForExternalAppForRead(*RawFilename);
-						LogException(MoveTemp(AbsFilename), 1, ErrorMsg);
-					}
-#endif
+					);
 				}
 			}
 		}
 
-#if !PLATFORM_EXCEPTIONS_DISABLED
-		FTaskGraphInterface::Get().WaitUntilTasksComplete(ExceptionTasks);
-#endif
+		FResults::WaitForErrorTasks();
 
 		PreparseTimer.Stop();
 		UE_LOG(LogCompile, Log, TEXT("Defined classes in %d modules containing %d files(s) in %.3f secs."), GManifest.Modules.Num(), NumHeaders, PreparseTime);
 	}
 
 	// Perform the second phase of loading
-	if (Result == ECompilationResult::Succeeded)
+	if (FResults::IsSucceeding())
 	{
 		double ResolveTime = 0.0;
 		FDurationTimer ResolveTimer(ResolveTime);
@@ -7161,38 +7064,18 @@ ECompilationResult::Type PreparseModules(const FString& ModuleInfoPath, int32& N
 
 			// Don't resolve superclasses for module when loading from makefile.
 			// Data is only partially loaded at this point.
-#if !PLATFORM_EXCEPTIONS_DISABLED
-			try
-#endif
-			{
-				ResolveSuperClasses(Package);
-			}
-#if !PLATFORM_EXCEPTIONS_DISABLED
-			catch (TCHAR* ErrorMsg)
-			{
-				TGuardValue<ELogTimes::Type> DisableLogTimes(GPrintLogTimes, ELogTimes::None);
-
-				FString FormattedErrorMessage = FString::Printf(TEXT("Error: %s\r\n"), ErrorMsg);
-
-				Result = GCompilationResult;
-
-				UE_LOG(LogCompile, Log, TEXT("%s"), *FormattedErrorMessage);
-				GWarn->Log(ELogVerbosity::Error, FormattedErrorMessage);
-
-				++NumFailures;
-			}
-#endif
+			FResults::Try([Package] 
+				{
+					ResolveSuperClasses(Package);
+				}
+			);
 		}
 
-#if !PLATFORM_EXCEPTIONS_DISABLED
-		FTaskGraphInterface::Get().WaitUntilTasksComplete(ExceptionTasks);
-#endif
+		FResults::WaitForErrorTasks();
 
 		ResolveTimer.Stop();
 		UE_LOG(LogCompile, Log, TEXT("Resolving classes in %d modules containing %d files(s) in %.3f secs."), GManifest.Modules.Num(), NumHeaders, ResolveTime);
 	}
-
-	return Result;
 }
 
 ECompilationResult::Type UnrealHeaderTool_Main(const FString& ModuleInfoFilename)
@@ -7202,7 +7085,6 @@ ECompilationResult::Type UnrealHeaderTool_Main(const FString& ModuleInfoFilename
 	MainTimer.Start();
 
 	check(GIsUCCMakeStandaloneHeaderGenerator);
-	ECompilationResult::Type Result = ECompilationResult::Succeeded;
 
 	FString ModuleInfoPath = FPaths::GetPath(ModuleInfoFilename);
 
@@ -7216,24 +7098,23 @@ ECompilationResult::Type UnrealHeaderTool_Main(const FString& ModuleInfoFilename
 #if !PLATFORM_EXCEPTIONS_DISABLED
 	catch (const TCHAR* Ex)
 	{
-		UE_LOG(LogCompile, Error, TEXT("Failed to load manifest file '%s': %s"), *ModuleInfoFilename, Ex);
-		return GCompilationResult;
+		FResults::LogError(*FString::Printf(TEXT("Failed to load manifest file '%s': %s"), *ModuleInfoFilename, Ex));
+		return FResults::GetOverallResults();
 	}
 #endif
 
 	// Counters.
-	int32 NumFailures = 0;
 	double TotalModulePreparseTime = 0.0;
 	double TotalParseAndCodegenTime = 0.0;
 
 	{
 		FDurationTimer TotalModulePreparseTimer(TotalModulePreparseTime);
 		TotalModulePreparseTimer.Start();
-		Result = PreparseModules(ModuleInfoPath, NumFailures);
+		PreparseModules(ModuleInfoPath);
 		TotalModulePreparseTimer.Stop();
 	}
 	// Do the actual parse of the headers and generate for them
-	if (Result == ECompilationResult::Succeeded)
+	if (FResults::IsSucceeding())
 	{
 		FScopedDurationTimer ParseAndCodeGenTimer(TotalParseAndCodegenTime);
 
@@ -7249,37 +7130,11 @@ ECompilationResult::Type UnrealHeaderTool_Main(const FString& ModuleInfoFilename
 
 			if (ScriptSuperClass && !ScriptSuperClass->HasAnyClassFlags(CLASS_Intrinsic) && GTypeDefinitionInfoMap.Contains(ScriptClass) && !GTypeDefinitionInfoMap.Contains(ScriptSuperClass))
 			{
-				class FSuperClassContextSupplier : public FContextSupplier
-				{
-				public:
-					FSuperClassContextSupplier(const UClass* Class)
-						: DefinitionInfo(GTypeDefinitionInfoMap[Class])
-					{ }
-
-					virtual FString GetContext() override
-					{
-						FString Filename = IFileManager::Get().ConvertToAbsolutePathForExternalAppForRead(*DefinitionInfo->GetUnrealSourceFile().GetFilename());
-						int32 LineNumber = DefinitionInfo->GetLineNumber();
-						return FString::Printf(TEXT("%s(%i)"), *Filename, LineNumber);
-					}
-				private:
-					TSharedRef<FUnrealTypeDefinitionInfo> DefinitionInfo;
-				} ContextSupplier(ScriptClass);
-
-				FContextSupplier* OldContext = GWarn->GetContext();
-
-				TGuardValue<ELogTimes::Type> DisableLogTimes(GPrintLogTimes, ELogTimes::None);
-
-				GWarn->SetContext(&ContextSupplier);
-				GWarn->Log(ELogVerbosity::Error, FString::Printf(TEXT("Error: Superclass %s of class %s not found"), *ScriptSuperClass->GetName(), *ScriptClass->GetName()));
-				GWarn->SetContext(OldContext);
-
-				Result = ECompilationResult::OtherCompilationError;
-				++NumFailures;
+				FResults::LogError(*ScriptClass, *FString::Printf(TEXT("Error: Superclass %s of class %s not found"), *ScriptSuperClass->GetName(), *ScriptClass->GetName()));
 			}
 		}
 
-		if (Result == ECompilationResult::Succeeded)
+		if (FResults::IsSucceeding())
 		{
 			TArray<IScriptGeneratorPluginInterface*> ScriptPlugins;
 			// Can only export scripts for game targets
@@ -7295,10 +7150,10 @@ ECompilationResult::Type UnrealHeaderTool_Main(const FString& ModuleInfoFilename
 					FClasses AllClasses(ClassesByPackageMap.Find(Package));
 					AllClasses.Validate();
 
-					Result = FHeaderParser::ParseAllHeadersInside(AllClasses, GWarn, Package, Module, ScriptPlugins);
+					ECompilationResult::Type Result = FHeaderParser::ParseAllHeadersInside(AllClasses, GWarn, Package, Module, ScriptPlugins);
 					if (Result != ECompilationResult::Succeeded)
 					{
-						++NumFailures;
+						FResults::SetResult(Result);
 						break;
 					}
 				}
@@ -7373,13 +7228,7 @@ ECompilationResult::Type UnrealHeaderTool_Main(const FString& ModuleInfoFilename
 	}
 
 	RequestEngineExit(TEXT("UnrealHeaderTool finished"));
-
-	if (Result != ECompilationResult::Succeeded || NumFailures > 0)
-	{
-		return ECompilationResult::OtherCompilationError;
-	}
-
-	return Result;
+	return FResults::GetOverallResults();
 }
 
 UClass* ProcessParsedClass(bool bClassIsAnInterface, TArray<FHeaderProvider>& DependentOn, const FString& ClassName, const FString& BaseClassName, UObject* InParent, EObjectFlags Flags)
diff --git a/Engine/Source/Programs/UnrealHeaderTool/Private/Exceptions.cpp b/Engine/Source/Programs/UnrealHeaderTool/Private/Exceptions.cpp
new file mode 100644
--- /dev/null
+++ b/Engine/Source/Programs/UnrealHeaderTool/Private/Exceptions.cpp
@@ -0,0 +1,202 @@
+// Copyright Epic Games, Inc. All Rights Reserved.
+
+#include "Exceptions.h"
+
+#include "ClassMaps.h"
+#include "UnrealHeaderTool.h"
+#include "UnrealHeaderToolGlobals.h"
+#include "UnrealSourceFile.h"
+#include "UnrealTypeDefinitionInfo.h"
+
+#include "CoreGlobals.h"
+#include "Async/TaskGraphInterfaces.h"
+#include "HAL/FileManager.h"
+#include "Misc/FeedbackContext.h"
+#include "Misc/ScopeLock.h"
+
+#include <atomic>
+
+FUHTException::FUHTException(ECompilationResult::Type InResult, FString&& InFilename, int32 InLine, FString&& InMessage)
+	: Result(InResult)
+	, Message(MoveTemp(InMessage))
+	, Filename(MoveTemp(InFilename))
+	, Line(InLine)
+{
+}
+
+FUHTException::FUHTException(ECompilationResult::Type InResult, const FUnrealSourceFile& SourceFile, int32 InLine, FString&& InMessage)
+	: Result(InResult)
+	, Message(MoveTemp(InMessage))
+	, Filename(SourceFile.GetFilename())
+	, Line(InLine)
+{
+}
+
+FUHTException::FUHTException(ECompilationResult::Type InResult, FUnrealTypeDefinitionInfo& TypeDef, FString&& InMessage)
+	: Result(InResult)
+	, Message(MoveTemp(InMessage))
+	, Filename(TypeDef.GetUnrealSourceFile().GetFilename())
+	, Line(TypeDef.GetLineNumber())
+{
+}
+
+FUHTException::FUHTException(ECompilationResult::Type InResult, const UObject* InObject, FString&& InMessage)
+	: Result(InResult)
+	, Message(MoveTemp(InMessage))
+	, Line(1)
+{
+	if (const UField* Field = Cast<UField>(InObject))
+	{
+		if (TSharedRef<FUnrealTypeDefinitionInfo>* TypeDef = GTypeDefinitionInfoMap.Find(Field))
+		{
+			Filename = (*TypeDef)->GetUnrealSourceFile().GetFilename();
+			Line = (*TypeDef)->GetLineNumber();
+		}
+	}
+}
+
+namespace UE::UnrealHeaderTool::Exceptions::Private
+{
+	std::atomic<ECompilationResult::Type> OverallResults = ECompilationResult::Succeeded;
+	std::atomic<int32> NumFailures = 0;
+	FGraphEventArray ErrorTasks;
+	FCriticalSection ErrorTasksCS;
+
+	void LogErrorInternal(ECompilationResult::Type InResult, const FString& Filename, int32 Line, const FString& Message)
+	{
+		TGuardValue<ELogTimes::Type> DisableLogTimes(GPrintLogTimes, ELogTimes::None);
+
+		FString FormattedErrorMessage;
+		if (Filename.IsEmpty())
+		{
+			FormattedErrorMessage = FString::Printf(TEXT("Error: %s\r\n"), *Message);
+		}
+		else
+		{
+			FormattedErrorMessage = FString::Printf(TEXT("%s(%d): Error: %s\r\n"), *Filename, Line, *Message);
+		}
+
+		UE_LOG(LogCompile, Log, TEXT("%s"), *FormattedErrorMessage);
+		GWarn->Log(ELogVerbosity::Error, FormattedErrorMessage);
+
+		FResults::SetResult(InResult);
+	}
+}
+
+void FResults::LogError(FString&& Filename, int32 Line, const FString& Message, ECompilationResult::Type InResult/* = ECompilationResult::OtherCompilationError*/)
+{
+	using namespace UE::UnrealHeaderTool::Exceptions::Private;
+
+	if (IsInGameThread())
+	{
+		LogErrorInternal(InResult, Filename, Line, Message);
+	}
+	else
+	{
+		auto LogExceptionTask = [Filename = MoveTemp(Filename), Line, Message, InResult]()
+		{
+			LogErrorInternal(InResult, Filename, Line, Message);
+		};
+
+		FGraphEventRef EventRef = FFunctionGraphTask::CreateAndDispatchWhenReady(MoveTemp(LogExceptionTask), TStatId(), nullptr, ENamedThreads::GameThread);
+
+		FScopeLock Lock(&ErrorTasksCS);
+		ErrorTasks.Add(EventRef);
+	}
+}
+
+void FResults::LogError(const FUHTException& Ex)
+{
+	FString AbsFilename;
+	const TCHAR* Filename = Ex.GetFilename();
+	if (Filename != nullptr)
+	{
+		AbsFilename = IFileManager::Get().ConvertToAbsolutePathForExternalAppForRead(Ex.GetFilename());
+	}
+	LogError(MoveTemp(AbsFilename), Ex.GetLine(), Ex.GetMessage());
+}
+
+void FResults::LogError(const FUnrealSourceFile& SourceFile, const FUHTException& Ex)
+{
+	FString AbsFilename;
+	const TCHAR* Filename = Ex.GetFilename();
+	if (Filename == nullptr)
+	{
+		Filename = *SourceFile.GetFilename();
+	}
+	if (Filename != nullptr)
+	{
+		AbsFilename = IFileManager::Get().ConvertToAbsolutePathForExternalAppForRead(Ex.GetFilename());
+	}
+	LogError(MoveTemp(AbsFilename), Ex.GetLine(), Ex.GetMessage());
+}
+
+void FResults::LogError(const FUnrealSourceFile& SourceFile, int32 Line, const TCHAR* ErrorMsg, ECompilationResult::Type InResult/* = ECompilationResult::OtherCompilationError*/)
+{
+	FString AbsFilename = IFileManager::Get().ConvertToAbsolutePathForExternalAppForRead(*SourceFile.GetFilename());
+	LogError(MoveTemp(AbsFilename), Line, ErrorMsg, InResult);
+}
+
+void FResults::LogError(const UObject& Object, const TCHAR* ErrorMsg, ECompilationResult::Type InResult/* = ECompilationResult::OtherCompilationError*/)
+{
+	if (const UField* Field = Cast<UField>(&Object))
+	{
+		if (TSharedRef<FUnrealTypeDefinitionInfo>* TypeDef = GTypeDefinitionInfoMap.Find(Field))
+		{
+			LogError((*TypeDef)->GetUnrealSourceFile(), (*TypeDef)->GetLineNumber(), ErrorMsg, InResult);
+		}
+	}
+	LogError(ErrorMsg, InResult);
+}
+
+void FResults::LogError(const TCHAR* ErrorMsg, ECompilationResult::Type InResult/* = ECompilationResult::OtherCompilationError*/)
+{
+	LogError(TEXT(""), 1, ErrorMsg, InResult);
+}
+
+void FResults::WaitForErrorTasks()
+{
+	using namespace UE::UnrealHeaderTool::Exceptions::Private;
+
+	FGraphEventArray LocalExceptionTasks;
+	{
+		FScopeLock Lock(&ErrorTasksCS);
+		LocalExceptionTasks = MoveTemp(ErrorTasks);
+	}
+	FTaskGraphInterface::Get().WaitUntilTasksComplete(LocalExceptionTasks);
+}
+
+bool FResults::IsSucceeding()
+{
+	using namespace UE::UnrealHeaderTool::Exceptions::Private;
+
+	return OverallResults == ECompilationResult::Succeeded;
+}
+
+void FResults::SetResult(ECompilationResult::Type InResult)
+{
+	using namespace UE::UnrealHeaderTool::Exceptions::Private;
+
+	checkf(InResult != ECompilationResult::Succeeded, TEXT("The results can't be set to succeeded."));
+	OverallResults = InResult;
+	++NumFailures;
+}
+
+ECompilationResult::Type FResults::GetResults()
+{
+	using namespace UE::UnrealHeaderTool::Exceptions::Private;
+
+	return OverallResults;
+}
+
+ECompilationResult::Type FResults::GetOverallResults()
+{
+	using namespace UE::UnrealHeaderTool::Exceptions::Private;
+
+	// For some legacy reason, we don't actually return the result
+	if (OverallResults != ECompilationResult::Succeeded || NumFailures > 0)
+	{
+		return ECompilationResult::OtherCompilationError;
+	}
+	return OverallResults;
+}
diff --git a/Engine/Source/Programs/UnrealHeaderTool/Private/Exceptions.h b/Engine/Source/Programs/UnrealHeaderTool/Private/Exceptions.h
new file mode 100644
--- /dev/null
+++ b/Engine/Source/Programs/UnrealHeaderTool/Private/Exceptions.h
@@ -0,0 +1,324 @@
+// Copyright Epic Games, Inc. All Rights Reserved.
+
+#pragma once
+
+#include "CoreMinimal.h"
+#include "Containers/UnrealString.h"
+#include "Misc/CompilationResult.h"
+#include "ProfilingDebugging/ScopedTimers.h"
+
+class FUnrealSourceFile;
+class FUnrealTypeDefinitionInfo;
+
+class FUHTException
+{
+public:
+
+	/**
+	 * Generate an exception
+	 * @param InFilename The file generating the exception
+	 * @param InLine The line number in the file
+	 * @param InFmt The format string
+	 * @param InArgs Arguments supplied to the format string
+	 */
+	template <typename FmtType, typename... Types>
+	UE_NORETURN static void VARARGS Throwf(FString&& InFilename, int32 InLine, const FmtType& InFmt, Types... InArgs)
+	{
+		FString ResultString = FString::Printf(InFmt, InArgs ...);
+		throw FUHTException(ECompilationResult::OtherCompilationError, MoveTemp(InFilename), InLine, MoveTemp(ResultString));
+	}
+
+	/**
+	 * Generate an exception
+	 * @param InResult Result code to be returned as the overall result of the compilation process
+	 * @param InFilename The file generating the exception
+	 * @param InLine The line number in the file
+	 * @param InFmt The format string
+	 * @param InArgs Arguments supplied to the format string
+	 */
+	template <typename FmtType, typename... Types>
+	UE_NORETURN static void VARARGS Throwf(ECompilationResult::Type InResult, FString&& InFilename, int32 InLine, const FmtType& InFmt, Types... InArgs)
+	{
+		FString ResultString = FString::Printf(InFmt, InArgs ...);
+		throw FUHTException(InResult, MoveTemp(InFilename), InLine, MoveTemp(ResultString));
+	}
+
+	/**
+	 * Generate an exception
+	 * @param InSourceFile The source file generating the exception
+	 * @param InLine The line number in the file
+	 * @param InFmt The format string
+	 * @param InArgs Arguments supplied to the format string
+	 */
+	template <typename FmtType, typename... Types>
+	UE_NORETURN static void VARARGS Throwf(const FUnrealSourceFile& InSourceFile, int32 InLine, const FmtType& InFmt, Types... InArgs)
+	{
+		FString ResultString = FString::Printf(InFmt, InArgs ...);
+		throw FUHTException(ECompilationResult::OtherCompilationError, InSourceFile, InLine, MoveTemp(ResultString));
+	}
+
+	/**
+	 * Generate an exception
+	 * @param InResult Result code to be returned as the overall result of the compilation process
+	 * @param InSourceFile The source file generating the exception
+	 * @param InLine The line number in the file
+	 * @param InFmt The format string
+	 * @param InArgs Arguments supplied to the format string
+	 */
+	template <typename FmtType, typename... Types>
+	UE_NORETURN static void VARARGS Throwf(ECompilationResult::Type InResult, const FUnrealSourceFile& InSourceFile, int32 InLine, const FmtType& InFmt, Types... InArgs)
+	{
+		FString ResultString = FString::Printf(InFmt, InArgs ...);
+		throw FUHTException(InResult, InSourceFile, InLine, MoveTemp(ResultString));
+	}
+
+	/**
+	 * Generate an exception
+	 * @param InTypeDef The type definition generating the exception.  The filename and line number will be retrieved from the type definition if possible.
+	 * @param InFmt The format string
+	 * @param InArgs Arguments supplied to the format string
+	 */
+	template <typename FmtType, typename... Types>
+	UE_NORETURN static void VARARGS Throwf(FUnrealTypeDefinitionInfo& InTypeDef, const FmtType& InFmt, Types... InArgs)
+	{
+		FString ResultString = FString::Printf(InFmt, InArgs ...);
+		throw FUHTException(ECompilationResult::OtherCompilationError, InTypeDef, MoveTemp(ResultString));
+	}
+
+	/**
+	 * Generate an exception
+	 * @param InResult Result code to be returned as the overall result of the compilation process
+	 * @param InTypeDef The type definition generating the exception.  The filename and line number will be retrieved from the type definition if possible.
+	 * @param InFmt The format string
+	 * @param InArgs Arguments supplied to the format string
+	 */
+	template <typename FmtType, typename... Types>
+	UE_NORETURN static void VARARGS Throwf(ECompilationResult::Type InResult, FUnrealTypeDefinitionInfo& InTypeDef, const FmtType& InFmt, Types... InArgs)
+	{
+		FString ResultString = FString::Printf(InFmt, InArgs ...);
+		throw FUHTException(InResult, InTypeDef, MoveTemp(ResultString));
+	}
+
+	/**
+	 * Generate an exception
+	 * @param InObject The type object generating the exception.  The filename and line number will be retrieved from the object's type definition if possible.
+	 * @param InFmt The format string
+	 * @param InArgs Arguments supplied to the format string
+	 */
+	template <typename FmtType, typename... Types>
+	UE_NORETURN static void VARARGS Throwf(const UObject* InObject, const FmtType& InFmt, Types... InArgs)
+	{
+		FString ResultString = FString::Printf(InFmt, InArgs ...);
+		throw FUHTException(ECompilationResult::OtherCompilationError, InObject, MoveTemp(ResultString));
+	}
+
+	/**
+	 * Generate an exception
+	 * @param Result Result code to be returned as the overall result of the compilation process
+	 * @param Object The type object generating the exception.  The filename and line number will be retrieved from the object's type definition if possible.
+	 * @param Fmt The format string
+	 * @param Args Arguments supplied to the format string
+	 */
+	template <typename FmtType, typename... Types>
+	UE_NORETURN static void VARARGS Throwf(ECompilationResult::Type InResult, const UObject* InObject, const FmtType& InFmt, Types... InArgs)
+	{
+		FString ResultString = FString::Printf(InFmt, InArgs ...);
+		throw FUHTException(InResult, InObject, MoveTemp(ResultString));
+	}
+
+	/**
+	 * Return the result code of the exception
+	 */
+	ECompilationResult::Type GetResult() const 
+	{ 
+		return Result; 
+	}
+
+	/**
+	 * Return the filename of the exception.
+	 */
+	const TCHAR* GetFilename() const 
+	{
+		return *Filename; 
+	}
+
+	/**
+	 * Return the line number in the file of the exception
+	 */
+	int32 GetLine() const
+	{
+		return Line;
+	}
+
+	/**
+	 * Return the message of the exception
+	 */
+	const FString& GetMessage() const 
+	{ 
+		return Message; 
+	}
+
+private:
+	ECompilationResult::Type Result = ECompilationResult::OtherCompilationError;
+	FString Message;
+	FString Filename;
+	int32 Line;
+
+	FUHTException(ECompilationResult::Type InResult, FString&& InFilename, int32 InLine, FString&& InMessage);
+	FUHTException(ECompilationResult::Type InResult, const FUnrealSourceFile& SourceFile, int32 InLine, FString&& InMessage);
+	FUHTException(ECompilationResult::Type InResult, FUnrealTypeDefinitionInfo& TypeDef, FString&& InMessage);
+	FUHTException(ECompilationResult::Type InResult, const UObject* Object, FString&& InMessage);
+};
+
+/** Helper methods for working with exceptions and compilation results */
+struct FResults
+{
+
+	/**
+	 * Wait for any pending error tasks to complete.
+	 * 
+	 * When job threads are used to log errors or throw exceptions, those errors and exceptions are collected by the main game thread.
+	 * After waiting for all the pending jobs to complete, invoke this method to ensure that all pending errors and exceptions have
+	 * been collected.
+	 */
+	static void WaitForErrorTasks();
+
+	/**
+	 * Test to see if no errors or exceptions have been posted.
+	 */
+	static bool IsSucceeding();
+
+	/**
+	 * Set the overall results.
+	 */
+	static void SetResult(ECompilationResult::Type InResult);
+
+	/**
+	 * Get the current results without processing for overall result
+	 */
+	static ECompilationResult::Type GetResults();
+
+	/**
+	 * Get the overall results to be returned from compilation.
+	 */
+	static ECompilationResult::Type GetOverallResults();
+
+	/**
+	 * Log an error.
+	 * @param Filename The filename generating the error.  If empty, then no file and line number are included in the error.
+	 * @param Line Line number of the error
+	 * @param Message Message body of the error
+	 * @param Result Compilation result of the error
+	 */
+	static void LogError(FString&& Filename, int32 Line, const FString& Message, ECompilationResult::Type Result = ECompilationResult::OtherCompilationError);
+
+	/**
+	 * Log an error from an exception with possible override of the source file.
+	 * @param SourceFile The source file currently being processed.  If the exception does not include filename information, the source file will be used.
+	 * @param Ex The exception generating the error.
+	 */
+	static void LogError(const FUnrealSourceFile& SourceFile, const FUHTException& Ex);
+
+	/**
+	 * Log an error from an exception.
+	 * @param Ex The exception generating the error.
+	 */
+	static void LogError(const FUHTException& Ex);
+
+	/**
+	 * Log an error for the given source file
+	 * @param SourceFile The source file generating the error
+	 * @param Line The line number generating the error
+	 * @param ErrorMsg The text of the error
+	 * @param Result Compilation result of the error
+	 */
+	static void LogError(const FUnrealSourceFile& SourceFile, int32 Line, const TCHAR* ErrorMsg, ECompilationResult::Type Result = ECompilationResult::OtherCompilationError);
+
+	/**
+	 * Log an error for the given source file where the object is defined
+	 * @param Object The source file and line number will be determined from the object
+	 * @param ErrorMsg The text of the error
+	 * @param Result Compilation result of the error
+	 */
+	static void LogError(const UObject& Object, const TCHAR* ErrorMsg, ECompilationResult::Type Result = ECompilationResult::OtherCompilationError);
+
+	/**
+	 * Log an error without any source file information
+	 * @param ErrorMsg The text of the error
+	 * @param Result Compilation result of the error
+	 */
+	static void LogError(const TCHAR* ErrorMsg, ECompilationResult::Type Result = ECompilationResult::OtherCompilationError);
+
+	/**
+	 * Invoke the given lambda in a try block catching all supported exception types.
+	 * @param SourceFile The source file being processed
+	 * @param InLambda The code to be executed in the try block
+	 */
+	template<typename Lambda>
+	static void Try(FUnrealSourceFile& SourceFile, Lambda&& InLambda)
+	{
+		if (IsSucceeding())
+		{
+#if !PLATFORM_EXCEPTIONS_DISABLED
+			try
+#endif
+			{
+				InLambda();
+			}
+#if !PLATFORM_EXCEPTIONS_DISABLED
+			catch (const FUHTException& Ex)
+			{
+				LogError(SourceFile, Ex);
+			}
+			catch (const TCHAR* ErrorMsg)
+			{
+				LogError(SourceFile, 1, ErrorMsg);
+			}
+#endif
+		}
+	}
+
+	/**
+	 * Invoke the given lambda in a try block catching all supported exception types.
+	 * @param InLambda The code to be executed in the try block
+	 */
+	template<typename Lambda>
+	static void Try(Lambda&& InLambda)
+	{
+		if (IsSucceeding())
+		{
+#if !PLATFORM_EXCEPTIONS_DISABLED
+			try
+#endif
+			{
+				InLambda();
+			}
+#if !PLATFORM_EXCEPTIONS_DISABLED
+			catch (const FUHTException& Ex)
+			{
+				LogError(Ex);
+			}
+			catch (const TCHAR* ErrorMsg)
+			{
+				LogError(ErrorMsg);
+			}
+#endif
+		}
+	}
+
+	/**
+	 * Invoke the given lambda in a try block catching all supported exception types.
+	 * @param InLambda The code to be executed in the try block
+	 * @return The time in seconds it took to execute the lambda
+	 */
+	template<typename Lambda>
+	static double TimedTry(Lambda&& InLambda)
+	{
+		double DeltaTime = 0.0;
+		{
+			FScopedDurationTimer Timer(DeltaTime);
+			Try(InLambda);
+		}
+		return DeltaTime;
+	}
+};
diff --git a/Engine/Source/Programs/UnrealHeaderTool/Private/FileLineException.cpp b/Engine/Source/Programs/UnrealHeaderTool/Private/FileLineException.cpp
deleted file mode 100644
diff --git a/Engine/Source/Programs/UnrealHeaderTool/Private/FileLineException.h b/Engine/Source/Programs/UnrealHeaderTool/Private/FileLineException.h
deleted file mode 100644
diff --git a/Engine/Source/Programs/UnrealHeaderTool/Private/HeaderParser.cpp b/Engine/Source/Programs/UnrealHeaderTool/Private/HeaderParser.cpp
--- a/Engine/Source/Programs/UnrealHeaderTool/Private/HeaderParser.cpp
+++ b/Engine/Source/Programs/UnrealHeaderTool/Private/HeaderParser.cpp
@@ -18,7 +18,7 @@
 #include "Misc/DefaultValueHelper.h"
 #include "Manifest.h"
 #include "Math/UnitConversion.h"
-#include "FileLineException.h"
+#include "Exceptions.h"
 #include "UnrealTypeDefinitionInfo.h"
 #include "Containers/EnumAsByte.h"
 #include "Algo/AllOf.h"
@@ -111,12 +111,6 @@ TArray<FString> FHeaderParser::PropertyCPPTypesRequiringUIRanges = { TEXT("float
 TArray<FString> FHeaderParser::ReservedTypeNames = { TEXT("none") };
 TMap<UClass*, ClassDefinitionRange> ClassDefinitionRanges;
 
-/**
- * Dirty hack global variable to allow different result codes passed through
- * exceptions. Needs to be fixed in future versions of UHT.
- */
-extern ECompilationResult::Type GCompilationResult;
-
 /*-----------------------------------------------------------------------------
 	Utility functions.
 -----------------------------------------------------------------------------*/
@@ -8786,7 +8780,7 @@ ECompilationResult::Type FHeaderParser::ParseHeader(FUnrealSourceFile* SourceFil
 		}
 
 		FailedFilesAnnotation.Set(CurrentSrcFile);
-		Result = GCompilationResult;
+		Result = FResults::GetResults();
 	}
 #endif
 
@@ -9165,7 +9159,7 @@ ECompilationResult::Type FHeaderParser::ParseAllHeadersInside(
 	catch (TCHAR* ErrorMsg)
 	{
 		Warn->Log(ELogVerbosity::Error, ErrorMsg);
-		Result = GCompilationResult;
+		Result = FResults::GetResults();
 	}
 #endif
 	// Unregister the header parser from the feedback context
@@ -9450,7 +9444,7 @@ void FHeaderParser::SimplifiedClassParse(const TCHAR* Filename, const TCHAR* InB
 					bool bKeepOldDirective = ShouldKeepDirective(OldDirective);
 					if (bKeepNewDirective != bKeepOldDirective)
 					{
-						FFileLineException::Throwf(
+						FUHTException::Throwf(
 							Filename,
 							CurrentLine,
 							TEXT("Mixing %s with %s in an #elif preprocessor block is not supported"),
@@ -9485,7 +9479,7 @@ void FHeaderParser::SimplifiedClassParse(const TCHAR* Filename, const TCHAR* InB
 							break;
 
 						case EBlockDirectiveType::WithHotReload:
-							FFileLineException::Throwf(Filename, CurrentLine, TEXT("Bad preprocessor directive in metadata declaration: %s; Only 'CPP', '1' and '0' can have #else directives"), *ClassName);
+							FUHTException::Throwf(Filename, CurrentLine, TEXT("Bad preprocessor directive in metadata declaration: %s; Only 'CPP', '1' and '0' can have #else directives"), *ClassName);
 
 						case EBlockDirectiveType::UnrecognizedBlock:
 						case EBlockDirectiveType::WithEditor:
@@ -9526,7 +9520,7 @@ void FHeaderParser::SimplifiedClassParse(const TCHAR* Filename, const TCHAR* InB
 							|| FindInitialStr(FoundSubstr, TrimmedStrLine, TEXT("UDELEGATE"))
 							|| FindInitialStr(FoundSubstr, TrimmedStrLine, TEXT("UFUNCTION")))
 						{
-							FFileLineException::Throwf(Filename, CurrentLine, TEXT("%s must not be inside preprocessor blocks, except for WITH_EDITORONLY_DATA"), FoundSubstr);
+							FUHTException::Throwf(Filename, CurrentLine, TEXT("%s must not be inside preprocessor blocks, except for WITH_EDITORONLY_DATA"), FoundSubstr);
 						}
 
 						// Try and determine if this line contains something like a serialize function
@@ -9548,7 +9542,7 @@ void FHeaderParser::SimplifiedClassParse(const TCHAR* Filename, const TCHAR* InB
 									if (((TrimmedStrLine.Find(Str_FArchive, ESearchCase::CaseSensitive, ESearchDir::FromStart, Pos)) != -1) ||
 										((TrimmedStrLine.Find(Str_FStructuredArchive, ESearchCase::CaseSensitive, ESearchDir::FromStart, Pos)) != -1))
 									{
-										FFileLineException::Throwf(Filename, CurrentLine, TEXT("'%s' must not be inside preprocessor blocks, except for WITH_EDITORONLY_DATA"), *TrimmedStrLine);
+										FUHTException::Throwf(Filename, CurrentLine, TEXT("'%s' must not be inside preprocessor blocks, except for WITH_EDITORONLY_DATA"), *TrimmedStrLine);
 									}
 								}
 							}
@@ -9605,7 +9599,7 @@ void FHeaderParser::SimplifiedClassParse(const TCHAR* Filename, const TCHAR* InB
 			{
 				if (bFoundGeneratedInclude)
 				{
-					FFileLineException::Throwf(Filename, CurrentLine, TEXT("#include found after .generated.h file - the .generated.h file should always be the last #include in a header"));
+					FUHTException::Throwf(Filename, CurrentLine, TEXT("#include found after .generated.h file - the .generated.h file should always be the last #include in a header"));
 				}
 
 				bFoundGeneratedInclude = DependsOnHeaderName.Contains(TEXT(".generated.h"));
@@ -9739,7 +9733,7 @@ void FHeaderParser::SimplifiedClassParse(const TCHAR* Filename, const TCHAR* InB
 				{
 					if (UInterfaceMacroDecl[10] != TEXT('('))
 					{
-						FFileLineException::Throwf(Filename, CurrentLine, TEXT("Missing open parenthesis after UINTERFACE"));
+						FUHTException::Throwf(Filename, CurrentLine, TEXT("Missing open parenthesis after UINTERFACE"));
 					}
 
 					FName StrippedInterfaceName;
@@ -9761,7 +9755,7 @@ void FHeaderParser::SimplifiedClassParse(const TCHAR* Filename, const TCHAR* InB
 				{
 					if (UClassMacroDecl[6] != TEXT('('))
 					{
-						FFileLineException::Throwf(Filename, CurrentLine, TEXT("Missing open parenthesis after UCLASS"));
+						FUHTException::Throwf(Filename, CurrentLine, TEXT("Missing open parenthesis after UCLASS"));
 					}
 
 					FName StrippedClassName;
@@ -9856,7 +9850,7 @@ void FHeaderPreParser::ParseClassDeclaration(const TCHAR* Filename, const TCHAR*
 
 				if (ClassNameWithoutPrefixStr == DependencyClassNameWithoutPrefixStr)
 				{
-					FFileLineException::Throwf(Filename, InputLineLocal, TEXT("A class cannot inherit itself or a type with the same name but a different prefix"));
+					FUHTException::Throwf(Filename, InputLineLocal, TEXT("A class cannot inherit itself or a type with the same name but a different prefix"));
 				}
 
 				FString StrippedDependencyName = DependencyClassName.Mid(1);
@@ -9881,7 +9875,7 @@ void FHeaderPreParser::ParseClassDeclaration(const TCHAR* Filename, const TCHAR*
 			FToken InterfaceClassNameToken;
 			if (!GetIdentifier(InterfaceClassNameToken, true))
 			{
-				FFileLineException::Throwf(Filename, InputLine, TEXT("Expected an interface class name"));
+				FUHTException::Throwf(Filename, InputLine, TEXT("Expected an interface class name"));
 			}
 
 			AddDependencyIfNeeded(FString(InterfaceClassNameToken.Identifier));
