Update the header provider to support refering a type definition info.

This new functionallity will be used in future additions.

#rnx

diff --git a/Engine/Source/Programs/UnrealHeaderTool/Private/HeaderProvider.cpp b/Engine/Source/Programs/UnrealHeaderTool/Private/HeaderProvider.cpp
--- a/Engine/Source/Programs/UnrealHeaderTool/Private/HeaderProvider.cpp
+++ b/Engine/Source/Programs/UnrealHeaderTool/Private/HeaderProvider.cpp
@@ -4,26 +4,37 @@
 #include "UnrealHeaderTool.h"
 #include "UnrealTypeDefinitionInfo.h"
 #include "ClassMaps.h"
+#include "HeaderParser.h"
+#include "StringUtils.h"
 
-FHeaderProvider::FHeaderProvider(EHeaderProviderSourceType InType, FString&& InId)//, bool bInAutoInclude/* = false*/)
+FHeaderProvider::FHeaderProvider(EHeaderProviderSourceType InType, FString&& InId)
 	: Type(InType)
 	, Id(MoveTemp(InId))
-	, Cache(nullptr)
 {
+}
 
+FHeaderProvider::FHeaderProvider(FUnrealTypeDefinitionInfo& InTypeDef)
+	: Type(EHeaderProviderSourceType::TypeDef)
+	, Id(InTypeDef.GetNameCPP())
+	, TypeDef(&InTypeDef)
+{
+	check(TypeDef->HasSource());
 }
 
 FUnrealSourceFile* FHeaderProvider::Resolve(const FUnrealSourceFile& ParentSourceFile)
 {
-	if (Type != EHeaderProviderSourceType::Resolved)
+	if (!bResolved)
 	{
-		if (Type == EHeaderProviderSourceType::ClassName)
+		bResolved = true;
+
+		switch (Type)
+		{
+		case EHeaderProviderSourceType::ClassName:
 		{
 			FName IdName(*Id, FNAME_Find);
 			if (TSharedRef<FUnrealTypeDefinitionInfo>* Source = GTypeDefinitionInfoMap.FindByName(IdName))
 			{
 				Cache = &(*Source)->GetUnrealSourceFile();
-
 				// There is an edge case with interfaces.  If you define the UMyInterface and IMyInterface in the same
 				// source file as a class that implements the interface, a HeaderProvider for IMyInterface is added 
 				// at the pre-parse time that later (incorrectly) resolves to UMyInterface.  This results in
@@ -33,13 +44,55 @@ FUnrealSourceFile* FHeaderProvider::Resolve(const FUnrealSourceFile& ParentSourc
 					Cache = nullptr;
 				}
 			}
+			break;
+		}
+
+		case EHeaderProviderSourceType::ScriptStructName:
+		{
+			if (!FUHTConfig::Get().StructsWithNoPrefix.Contains(Id))
+			{
+				FName IdName(*GetClassNameWithPrefixRemoved(Id), FNAME_Find);
+				if (TSharedRef<FUnrealTypeDefinitionInfo>* Source = GTypeDefinitionInfoMap.FindByName(IdName))
+				{
+					Cache = &(*Source)->GetUnrealSourceFile();
+				}
+			}
+			if (Cache == nullptr)
+			{
+				FName IdName(*Id, FNAME_Find);
+				if (TSharedRef<FUnrealTypeDefinitionInfo>* Source = GTypeDefinitionInfoMap.FindByName(IdName))
+				{
+					Cache = &(*Source)->GetUnrealSourceFile();
+				}
+			}
+			break;
 		}
-		else if (const TSharedRef<FUnrealSourceFile>* Source = GUnrealSourceFilesMap.Find(Id))
+
+		case EHeaderProviderSourceType::TypeDef:
 		{
-			Cache = &Source->Get();
+			Cache = &TypeDef->GetUnrealSourceFile();
+			break;
+		}
+
+		case EHeaderProviderSourceType::FileName:
+		{
+			if (const TSharedRef<FUnrealSourceFile>* Source = GUnrealSourceFilesMap.Find(Id))
+			{
+				Cache = &Source->Get();
+			}
+			break;
+		}
+
+		default:
+			check(false);
 		}
 
-		Type = EHeaderProviderSourceType::Resolved;
+		// There is questionable compatibility hack where a source file will always be exported
+		// regardless of having types when it is being included by the SAME package.
+		if (Cache && Cache->GetPackage() == ParentSourceFile.GetPackage())
+		{
+			Cache->MarkReferenced();
+		}
 	}
 
 	return Cache;
@@ -47,7 +100,20 @@ FUnrealSourceFile* FHeaderProvider::Resolve(const FUnrealSourceFile& ParentSourc
 
 FString FHeaderProvider::ToString() const
 {
-	return FString::Printf(TEXT("%s %s"), Type == EHeaderProviderSourceType::ClassName ? TEXT("class") : TEXT("file"), *Id);
+	switch (Type)
+	{
+	case EHeaderProviderSourceType::ClassName:
+		return FString::Printf(TEXT("%s %s"), TEXT("class"), *Id);
+	case EHeaderProviderSourceType::ScriptStructName:
+		return FString::Printf(TEXT("%s %s"), TEXT("struct"), *Id);
+	case EHeaderProviderSourceType::TypeDef:
+		return FString::Printf(TEXT("%s %s"), TEXT("property type"), *Id);
+	case EHeaderProviderSourceType::FileName:
+		return FString::Printf(TEXT("%s %s"), TEXT("file"), *Id);
+	default:
+		check(false);
+		return FString::Printf(TEXT("%s %s"), TEXT("unknown"), *Id);
+	}
 }
 
 const FString& FHeaderProvider::GetId() const
@@ -57,5 +123,5 @@ const FString& FHeaderProvider::GetId() const
 
 bool operator==(const FHeaderProvider& A, const FHeaderProvider& B)
 {
-	return A.Type == B.Type && A.Id == B.Id;
+	return A.Type == B.Type && A.Id == B.Id && A.TypeDef == B.TypeDef;
 }
diff --git a/Engine/Source/Programs/UnrealHeaderTool/Private/HeaderProvider.h b/Engine/Source/Programs/UnrealHeaderTool/Private/HeaderProvider.h
--- a/Engine/Source/Programs/UnrealHeaderTool/Private/HeaderProvider.h
+++ b/Engine/Source/Programs/UnrealHeaderTool/Private/HeaderProvider.h
@@ -6,12 +6,14 @@
 #include "Containers/UnrealString.h"
 
 class FUnrealSourceFile;
+class FUnrealTypeDefinitionInfo;
 
 enum class EHeaderProviderSourceType
 {
 	ClassName,
+	ScriptStructName,
+	TypeDef,
 	FileName,
-	Resolved
 };
 
 class FHeaderProvider
@@ -19,6 +21,7 @@ class FHeaderProvider
 	friend bool operator==(const FHeaderProvider& A, const FHeaderProvider& B);
 public:
 	FHeaderProvider(EHeaderProviderSourceType Type, FString&& Id);
+	explicit FHeaderProvider(FUnrealTypeDefinitionInfo& InTypeDef);
 
 	FUnrealSourceFile* Resolve(const FUnrealSourceFile& ParentSourceFile);
 
@@ -29,7 +32,9 @@ public:
 private:
 	EHeaderProviderSourceType Type;
 	FString Id;
-	FUnrealSourceFile* Cache;
+	FUnrealTypeDefinitionInfo* TypeDef = nullptr;
+	FUnrealSourceFile* Cache = nullptr;
+	bool bResolved = false;
 };
 
 bool operator==(const FHeaderProvider& A, const FHeaderProvider& B);
diff --git a/Engine/Source/Programs/UnrealHeaderTool/Private/UnrealSourceFile.h b/Engine/Source/Programs/UnrealHeaderTool/Private/UnrealSourceFile.h
--- a/Engine/Source/Programs/UnrealHeaderTool/Private/UnrealSourceFile.h
+++ b/Engine/Source/Programs/UnrealHeaderTool/Private/UnrealSourceFile.h
@@ -255,6 +255,22 @@ public:
 	 */
 	void MarkDependenciesResolved();
 
+	/**
+	 * Mark this source file has being referenced
+	 */
+	void MarkReferenced()
+	{
+		bIsReferenced = true;
+	}
+
+	/**
+	 * Return true if this source file should be exported
+	 */
+	bool ShouldExport() const
+	{
+		return bIsReferenced || GetScope()->ContainsTypes();
+	}
+
 	/**
 	 * Checks if dependencies has been resolved.
 	 */
@@ -294,6 +310,9 @@ private:
 	// Source file content.
 	FString Content;
 
+	// Tells if this file is referenced by another
+	bool bIsReferenced = false;
+
 	// Tells if generated header file was changed.
 	bool bHasChanged;
 
diff --git a/Engine/Source/Programs/UnrealHeaderTool/Private/UnrealTypeDefinitionInfo.h b/Engine/Source/Programs/UnrealHeaderTool/Private/UnrealTypeDefinitionInfo.h
--- a/Engine/Source/Programs/UnrealHeaderTool/Private/UnrealTypeDefinitionInfo.h
+++ b/Engine/Source/Programs/UnrealHeaderTool/Private/UnrealTypeDefinitionInfo.h
@@ -20,6 +20,27 @@ public:
 		, LineNumber(InLineNumber)
 	{ }
 
+	/**
+	 * Return the CPP version of the name
+	 */
+	const FString& GetNameCPP() const
+	{
+		// TEMPORARY STUB 
+		static FString EmptyString;
+		return EmptyString;
+		//return NameCPP;
+	}
+
+	/**
+	 * Return true if this type has source information
+	 */
+	bool HasSource() const
+	{
+		// TEMPORARY STUB
+		return true;
+		//return SourceFile != nullptr;
+	}
+
 	/**
 	 * Gets the line number in source file this type was defined in.
 	 */
