Adding referenced class include file.

#rnx

diff --git a/Engine/Source/Programs/UnrealHeaderTool/Private/ClassDeclarationMetaData.cpp b/Engine/Source/Programs/UnrealHeaderTool/Private/ClassDeclarationMetaData.cpp
--- a/Engine/Source/Programs/UnrealHeaderTool/Private/ClassDeclarationMetaData.cpp
+++ b/Engine/Source/Programs/UnrealHeaderTool/Private/ClassDeclarationMetaData.cpp
@@ -8,6 +8,7 @@
 #include "BaseParser.h"
 #include "HeaderParser.h"
 #include "Classes.h"
+#include "ParserClass.h"
 #include "Specifiers/ClassMetadataSpecifiers.h"
 #include "Algo/FindSortedStringCaseInsensitive.h"
 
