Implement the expanded type definitions.
Functionallity will be used in a future CL.

#rnx

diff --git a/Engine/Source/Programs/UnrealHeaderTool/Private/ClassMaps.h b/Engine/Source/Programs/UnrealHeaderTool/Private/ClassMaps.h
--- a/Engine/Source/Programs/UnrealHeaderTool/Private/ClassMaps.h
+++ b/Engine/Source/Programs/UnrealHeaderTool/Private/ClassMaps.h
@@ -5,10 +5,10 @@
 #include "CoreMinimal.h"
 #include "UObject/Stack.h"
 #include "UObject/ErrorException.h"
-#include "UnderlyingEnumType.h"
 
 #include "UnrealSourceFile.h"
 #include "UObject/ErrorException.h"
+#include "UnrealTypeDefinitionInfo.h"
 
 class UField;
 class UClass;
@@ -21,15 +21,6 @@ struct FManifestModule;
 class FUnrealSourceFile;
 class FUnrealTypeDefinitionInfo;
 
-enum class ESerializerArchiveType
-{
-	None = 0,
-
-	Archive                 = 1,
-	StructuredArchiveRecord = 2
-};
-ENUM_CLASS_FLAGS(ESerializerArchiveType)
-
 struct FArchiveTypeDefinePair
 {
 	ESerializerArchiveType ArchiveType = ESerializerArchiveType::None;
diff --git a/Engine/Source/Programs/UnrealHeaderTool/Private/ParserHelper.h b/Engine/Source/Programs/UnrealHeaderTool/Private/ParserHelper.h
--- a/Engine/Source/Programs/UnrealHeaderTool/Private/ParserHelper.h
+++ b/Engine/Source/Programs/UnrealHeaderTool/Private/ParserHelper.h
@@ -13,6 +13,7 @@
 #include "ClassMaps.h"
 #include "Templates/UniqueObj.h"
 #include "Templates/UniquePtr.h"
+#include "UnrealTypeDefinitionInfo.h"
 
 class UEnum;
 class UScriptStruct;
@@ -69,12 +70,6 @@ struct EArrayType
 	};
 };
 
-enum class EAllocatorType
-{
-	Default,
-	MemoryImage
-};
-
 struct ERefQualifier
 {
 	enum Type
diff --git a/Engine/Source/Programs/UnrealHeaderTool/Private/UnderlyingEnumType.h b/Engine/Source/Programs/UnrealHeaderTool/Private/UnderlyingEnumType.h
deleted file mode 100644
diff --git a/Engine/Source/Programs/UnrealHeaderTool/Private/UnrealTypeDefinitionInfo.cpp b/Engine/Source/Programs/UnrealHeaderTool/Private/UnrealTypeDefinitionInfo.cpp
new file mode 100644
--- /dev/null
+++ b/Engine/Source/Programs/UnrealHeaderTool/Private/UnrealTypeDefinitionInfo.cpp
@@ -0,0 +1,138 @@
+// Copyright Epic Games, Inc. All Rights Reserved.
+
+#include "UnrealTypeDefinitionInfo.h"
+#include "NativeClassExporter.h"
+#include "Scope.h"
+#include "UnrealSourceFile.h"
+#include "Misc/PackageName.h"
+#include "UObject/ErrorException.h"
+
+FUnrealPropertyDefinitionInfo* FUnrealTypeDefinitionInfo::AsProperty()
+{
+	return nullptr;
+}
+
+FUnrealObjectDefinitionInfo* FUnrealTypeDefinitionInfo::AsObject()
+{
+	return nullptr;
+}
+
+FUnrealPackageDefinitionInfo* FUnrealTypeDefinitionInfo::AsPackage()
+{
+	return nullptr;
+}
+
+FUnrealFieldDefinitionInfo* FUnrealTypeDefinitionInfo::AsField()
+{
+	return nullptr;
+}
+
+FUnrealEnumDefinitionInfo* FUnrealTypeDefinitionInfo::AsEnum()
+{
+	return nullptr;
+}
+
+FUnrealStructDefinitionInfo* FUnrealTypeDefinitionInfo::AsStruct()
+{
+	return nullptr;
+}
+
+FUnrealScriptStructDefinitionInfo* FUnrealTypeDefinitionInfo::AsScriptStruct()
+{
+	return nullptr;
+}
+
+FUnrealClassDefinitionInfo* FUnrealTypeDefinitionInfo::AsClass()
+{
+	return nullptr;
+}
+
+TSharedRef<FScope> FUnrealTypeDefinitionInfo::GetScope()
+{
+	if (!HasSource())
+	{
+		FError::Throwf(TEXT("Attempt to fetch the scope for type \"%s\" when it doesn't implement the method or there is no source file associated with the type."), *GetNameCPP());
+	}
+	return GetUnrealSourceFile().GetScope();
+}
+
+// Constructor
+FUnrealPackageDefinitionInfo::FUnrealPackageDefinitionInfo(const FManifestModule& InModule, UPackage* InPackage)
+	: FUnrealObjectDefinitionInfo(FString())
+	, Module(InModule)
+	, ShortUpperName(FPackageName::GetShortName(InPackage).ToUpper())
+	, API(FString::Printf(TEXT("%s_API "), *ShortUpperName))
+{
+	SetObject(InPackage);
+}
+
+void FUnrealTypeDefinitionInfo::SetHash(uint32 InHash)
+{
+	Hash = InHash;
+}
+
+uint32 FUnrealTypeDefinitionInfo::GetHash(bool bIncludeNoExport) const
+{
+	if (Hash == 0)
+	{
+		FError::Throwf(TEXT("Attempt to fetch the generated hash for type \"%s\" before it has been generated.  Include dependencies, topological sort, or job graph is in error."), *GetNameCPP());
+	}
+	return Hash;
+}
+
+void FUnrealTypeDefinitionInfo::GetHashTag(FUHTStringBuilder& Out) const
+{
+	uint32 TempHash = GetHash(false);
+	if (TempHash != 0)
+	{
+		if (Out.IsEmpty())
+		{
+			Out.Appendf(TEXT(" // %u"), TempHash);
+		}
+		else
+		{
+			Out.Appendf(TEXT(" %u"), TempHash);
+		}
+	}
+}
+
+uint32 FUnrealScriptStructDefinitionInfo::GetHash(bool bIncludeNoExport) const
+{
+	if (!bIncludeNoExport)
+	{
+		UScriptStruct* Struct = GetScriptStruct();
+		if ((Struct->StructFlags & STRUCT_NoExport) != 0)
+		{
+			return 0;
+		}
+	}
+	return FUnrealStructDefinitionInfo::GetHash(bIncludeNoExport);
+}
+
+TSharedRef<FScope> FUnrealStructDefinitionInfo::GetScope()
+{
+	return StructScope.ToSharedRef();
+}
+
+void FUnrealStructDefinitionInfo::SetObject(UObject* InObject)
+{
+	check(InObject != nullptr);
+
+	FUnrealFieldDefinitionInfo::SetObject(InObject);
+
+	StructScope = MakeShared<FStructScope>(static_cast<UStruct*>(InObject), &GetUnrealSourceFile().GetScope().Get());
+}
+
+uint32 FUnrealClassDefinitionInfo::GetHash(bool bIncludeNoExport) const
+{
+	if (!bIncludeNoExport)
+	{
+		UClass* Class = GetClass();
+		if (Class->HasAnyClassFlags(CLASS_NoExport))
+		{
+			return 0;
+		}
+	}
+	return FUnrealStructDefinitionInfo::GetHash(bIncludeNoExport);
+}
+
diff --git a/Engine/Source/Programs/UnrealHeaderTool/Private/UnrealTypeDefinitionInfo.h b/Engine/Source/Programs/UnrealHeaderTool/Private/UnrealTypeDefinitionInfo.h
--- a/Engine/Source/Programs/UnrealHeaderTool/Private/UnrealTypeDefinitionInfo.h
+++ b/Engine/Source/Programs/UnrealHeaderTool/Private/UnrealTypeDefinitionInfo.h
@@ -4,31 +4,217 @@
 
 #include "CoreMinimal.h"
 #include "Templates/SharedPointer.h"
+#include "ClassDeclarationMetaData.h"
+#include <atomic>
 
 // Forward declarations.
+class FClass;
+class FScope;
 class FUnrealSourceFile;
+struct FManifestModule;
+
+// These are declared in this way to allow swapping out the classes for something more optimized in the future
+typedef FStringOutputDevice FUHTStringBuilder;
+
+// The FUnrealTypeDefinitionInfo represents most all types required during parsing.  At this time, these structures
+// have a 1-1 correspondence with Engine types such as UClass and UProperty.  The goal is to provide a universal mechanism
+// of associating compiler data with given types without needing to add extra data to the engine types.  They are also
+// intended to eliminate all other UHT specific containers that map extra data with Engine classes.
+
+// The following types are supported represented in hierarchal form below:
+class FUnrealTypeDefinitionInfo;							// Base for all types, provides virtual methods to cast between all types
+	class FUnrealPropertyDefinitionInfo;					// Represents properties (FField)
+	class FUnrealObjectDefinitionInfo;						// Represents UObject
+		class FUnrealPackageDefinitionInfo;					// Represents UPackage
+		class FUnrealFieldDefinitionInfo;					// Represents UField
+			class FUnrealEnumDefinitionInfo;				// Represents UEnum
+			class FUnrealStructDefinitionInfo;				// Represents UStruct
+				class FUnrealScriptStructDefinitionInfo;	// Represents UScriptStruct
+				class FUnrealClassDefinitionInfo;			// Represents UClass
+
+enum class EUnderlyingEnumType
+{
+	Unspecified,
+	uint8,
+	uint16,
+	uint32,
+	uint64,
+	int8,
+	int16,
+	int32,
+	int64
+};
+
+enum class EAllocatorType
+{
+	Default,
+	MemoryImage
+};
+
+enum class ESerializerArchiveType
+{
+	None = 0,
+
+	Archive = 1,
+	StructuredArchiveRecord = 2
+};
+ENUM_CLASS_FLAGS(ESerializerArchiveType)
 
 /**
- * Class that stores information about type (USTRUCT/UCLASS) definition.
+ * Class that stores information about type definitions.
  */
 class FUnrealTypeDefinitionInfo : public TSharedFromThis<FUnrealTypeDefinitionInfo>
 {
 public:
-	// Constructor
-	FUnrealTypeDefinitionInfo(FUnrealSourceFile& InSourceFile, int32 InLineNumber)
-		: SourceFile(InSourceFile)
-		, LineNumber(InLineNumber)
-	{ }
+	virtual ~FUnrealTypeDefinitionInfo() = default;
+
+	/**
+	 * If this is a property, return the property version of the object
+	 * @return Pointer to the definition info or nullptr if not of that type.
+	 */
+	virtual FUnrealPropertyDefinitionInfo* AsProperty();
+
+	/**
+	 * If this is a property, return the property version of the object
+	 * @return Reference to the definition.  Will assert if not of that type.
+	 */
+	FUnrealPropertyDefinitionInfo& AsPropertyChecked()
+	{
+		FUnrealPropertyDefinitionInfo* Info = AsProperty();
+		check(Info);
+		return *Info;
+	}
+
+	/**
+	 * If this is an object, return the object version of the object
+	 * @return Pointer to the definition info or nullptr if not of that type.
+	 */
+	virtual FUnrealObjectDefinitionInfo* AsObject();
+
+	/**
+	 * If this is a object, return the object version of the object
+	 * @return Reference to the definition.  Will assert if not of that type.
+	 */
+	FUnrealObjectDefinitionInfo& AsObjectChecked()
+	{
+		FUnrealObjectDefinitionInfo* Info = AsObject();
+		check(Info);
+		return *Info;
+	}
+
+	/**
+	 * If this is a package, return the package version of the object
+	 * @return Pointer to the definition info or nullptr if not of that type.
+	 */
+	virtual FUnrealPackageDefinitionInfo* AsPackage();
+
+	/**
+	 * If this is a package, return the package version of the object
+	 * @return Reference to the definition.  Will assert if not of that type.
+	 */
+	FUnrealPackageDefinitionInfo& AsPackageChecked()
+	{
+		FUnrealPackageDefinitionInfo* Info = AsPackage();
+		check(Info);
+		return *Info;
+	}
+
+	/**
+	 * If this is an field, return the field version of the object
+	 * @return Pointer to the definition info or nullptr if not of that type.
+	 */
+	virtual FUnrealFieldDefinitionInfo* AsField();
+
+	/**
+	 * If this is an field, return the field version of the object
+	 * @return Reference to the definition.  Will assert if not of that type.
+	 */
+	FUnrealFieldDefinitionInfo& AsFieldChecked()
+	{
+		FUnrealFieldDefinitionInfo* Info = AsField();
+		check(Info);
+		return *Info;
+	}
+
+	/**
+	 * If this is an enumeration, return the enumeration version of the object
+	 * @return Pointer to the definition info or nullptr if not of that type.
+	 */
+	virtual FUnrealEnumDefinitionInfo* AsEnum();
+
+	/**
+	 * If this is an enumeration, return the enumeration version of the object
+	 * @return Reference to the definition.  Will assert if not of that type.
+	 */
+	FUnrealEnumDefinitionInfo& AsEnumChecked()
+	{
+		FUnrealEnumDefinitionInfo* Info = AsEnum();
+		check(Info);
+		return *Info;
+	}
+
+	/**
+	 * If this is a struct, return the struct version of the object
+	 * @return Pointer to the definition info or nullptr if not of that type.
+	 */
+	virtual FUnrealStructDefinitionInfo* AsStruct();
+
+	/**
+	 * If this is a struct, return the struct version of the object
+	 * @return Reference to the definition.  Will assert if not of that type.
+	 */
+	FUnrealStructDefinitionInfo& AsStructChecked()
+	{
+		FUnrealStructDefinitionInfo* Info = AsStruct();
+		check(Info);
+		return *Info;
+	}
+
+	/**
+	 * If this is a script struct, return the script struct version of the object
+	 * @return Pointer to the definition info or nullptr if not of that type.
+	 */
+	virtual FUnrealScriptStructDefinitionInfo* AsScriptStruct();
+
+	/**
+	 * If this is a script struct, return the script struct version of the object
+	 * @return Reference to the definition.  Will assert if not of that type.
+	 */
+	FUnrealScriptStructDefinitionInfo& AsScriptStructChecked()
+	{
+		FUnrealScriptStructDefinitionInfo* Info = AsScriptStruct();
+		check(Info);
+		return *Info;
+	}
+
+	/**
+	 * If this is a class, return the class version of the object
+	 * @return Pointer to the definition info or nullptr if not of that type.
+	 */
+	virtual FUnrealClassDefinitionInfo* AsClass();
+
+	/**
+	 * If this is a class, return the class version of the object
+	 * @return Reference to the definition.  Will assert if not of that type.
+	 */
+	FUnrealClassDefinitionInfo& AsClassChecked()
+	{
+		FUnrealClassDefinitionInfo* Info = AsClass();
+		check(Info);
+		return *Info;
+	}
+
+	/**
+	 * Return the compilation scope associated with this object
+	 */
+	virtual TSharedRef<FScope> GetScope();
 
 	/**
 	 * Return the CPP version of the name
 	 */
 	const FString& GetNameCPP() const
 	{
-		// TEMPORARY STUB 
-		static FString EmptyString;
-		return EmptyString;
-		//return NameCPP;
+		return NameCPP;
 	}
 
 	/**
@@ -36,9 +222,7 @@ public:
 	 */
 	bool HasSource() const
 	{
-		// TEMPORARY STUB
-		return true;
-		//return SourceFile != nullptr;
+		return SourceFile != nullptr;
 	}
 
 	/**
@@ -46,6 +230,7 @@ public:
 	 */
 	int32 GetLineNumber() const
 	{
+		check(HasSource());
 		return LineNumber;
 	}
 
@@ -55,22 +240,502 @@ public:
 	 */
 	FUnrealSourceFile& GetUnrealSourceFile()
 	{
-		return SourceFile;
+		check(HasSource());
+		return *SourceFile;
 	}
 
-	void SetLineNumber(int32 InLineNumber)
+	const FUnrealSourceFile& GetUnrealSourceFile() const
 	{
-		LineNumber = InLineNumber;
+		check(HasSource());
+		return *SourceFile;
 	}
 
-	const FUnrealSourceFile& GetUnrealSourceFile() const
+	/**
+	 * Set the hash calculated from the generated code for this type
+	 */
+	void SetHash(uint32 InHash);
+
+	/**
+	 * Return the previously set hash. This method will assert if the hash has not been set.
+	 */
+	virtual uint32 GetHash(bool bIncludeNoExport = true) const;
+
+	/**
+	 * Return the hash as a code comment.
+	 */
+	void GetHashTag(FUHTStringBuilder& Out) const;
+
+//protected: TEMPORARY MAKE THESE PUBLIC
+	explicit FUnrealTypeDefinitionInfo(FString&& InNameCPP)
+		: NameCPP(MoveTemp(InNameCPP))
+	{ }
+
+	// TEMPORARY API - REMOVE
+	FUnrealTypeDefinitionInfo(FUnrealSourceFile& InSourceFile, int32 InLineNumber)
+		: SourceFile(&InSourceFile)
+		, LineNumber(InLineNumber)
+	{ }
+
+	FUnrealTypeDefinitionInfo(FUnrealSourceFile& InSourceFile, int32 InLineNumber, FString&& InNameCPP)
+		: NameCPP(MoveTemp(InNameCPP))
+		, SourceFile(&InSourceFile)
+		, LineNumber(InLineNumber)
+	{ }
+
+private:
+	FString NameCPP;
+	FUnrealSourceFile* SourceFile = nullptr;
+	int32 LineNumber = 0;
+	std::atomic<uint32> Hash = 0;
+};
+
+/**
+ * Class that stores information about type definitions derived from UProperty
+ */
+class FUnrealPropertyDefinitionInfo 
+	: public FUnrealTypeDefinitionInfo
+{
+public:
+	FUnrealPropertyDefinitionInfo(FUnrealSourceFile& InSourceFile, int32 InLineNumber, FProperty* InProperty)
+		: FUnrealTypeDefinitionInfo(InSourceFile, InLineNumber, FString())
+		, Property(InProperty)
+	{ }
+
+	FUnrealPropertyDefinitionInfo(FUnrealSourceFile& InSourceFile, int32 InLineNumber, FProperty* InProperty, bool bInUnsized)
+		: FUnrealTypeDefinitionInfo(InSourceFile, InLineNumber, FString())
+		, Property(InProperty)
+		, bUnsized(bInUnsized)
+	{ }
+
+	FUnrealPropertyDefinitionInfo(FUnrealSourceFile& InSourceFile, int32 InLineNumber, FProperty* InProperty, EAllocatorType InAllocatorType)
+		: FUnrealTypeDefinitionInfo(InSourceFile, InLineNumber, FString())
+		, Property(InProperty)
+		, AllocatorType(InAllocatorType)
+	{ }
+
+	virtual FUnrealPropertyDefinitionInfo* AsProperty() override
+	{
+		return this;
+	}
+
+	/**
+	 * Return the Engine instance associated with the compiler instance
+	 */
+	FProperty* GetProperty() const
+	{
+		check(Property);
+		return Property;
+	}
+
+	/**
+	 * Set the string that represents the array dimensions
+	 */
+	void SetArrayDimensions(const FString& InArrayDimensions)
+	{
+		ArrayDimensions = InArrayDimensions;
+		check(!ArrayDimensions.IsEmpty());
+	}
+
+	/**
+	 * Get the string that represents the array dimensions.  A nullptr is returned if the property doesn't have any dimensions
+	 */
+	const TCHAR* GetArrayDimensions() const
+	{
+		return ArrayDimensions.IsEmpty() ? nullptr : *ArrayDimensions;
+	}
+
+	/**
+	 * Return true if the property is unsized
+	 */
+	bool IsUnsized() const
+	{
+		return bUnsized;
+	}
+
+	/**
+	 * Return the allocator type
+	 */
+	EAllocatorType GetAllocatorType() const
 	{
-		return SourceFile;
+		return AllocatorType;
 	}
 
 private:
-	FUnrealSourceFile& SourceFile;
-	int32 LineNumber;
+	FString ArrayDimensions;
+	FProperty* Property = nullptr;
+	EAllocatorType AllocatorType = EAllocatorType::Default;
+	bool bUnsized = false;
+};
 
-	friend struct FUnrealTypeDefinitionInfoArchiveProxy;
+/**
+ * Class that stores information about type definitions derived from UObject.
+ */
+class FUnrealObjectDefinitionInfo
+	: public FUnrealTypeDefinitionInfo
+{
+public:
+
+	virtual FUnrealObjectDefinitionInfo* AsObject() override
+	{
+		return this;
+	}
+
+	/**
+	 * Return the Engine instance associated with the compiler instance
+	 */
+	UObject* GetObject() const
+	{
+		check(Object);
+		return Object;
+	}
+
+	/**
+	 * Set the Engine instance associated with the compiler instance
+	 */
+	virtual void SetObject(UObject* InObject)
+	{
+		check(InObject != nullptr);
+		check(Object == nullptr);
+		Object = InObject;
+	}
+
+protected:
+	using FUnrealTypeDefinitionInfo::FUnrealTypeDefinitionInfo;
+
+private:
+	UObject* Object = nullptr;
+};
+
+/**
+ * Class that stores information about packages.
+ */
+class FUnrealPackageDefinitionInfo : public FUnrealObjectDefinitionInfo
+{
+public:
+	// Constructor
+	FUnrealPackageDefinitionInfo(const FManifestModule& InModule, UPackage* InPackage);
+
+	virtual FUnrealPackageDefinitionInfo* AsPackage() override
+	{
+		return this;
+	}
+
+	/**
+	 * Return the Engine instance associated with the compiler instance
+	 */
+	UPackage* GetPackage() const
+	{
+		return static_cast<UPackage*>(GetObject());
+	}
+
+	/**
+	 * Return the module information from the manifest associated with this package
+	 */
+	const FManifestModule& GetModule()
+	{
+		return Module;
+	}
+
+	/**
+	 * Return a collection of all source files contained within this package.
+	 * This collection is always valid.
+	 */
+	TArray<TSharedRef<FUnrealSourceFile>>& GetAllSourceFiles()
+	{
+		return AllSourceFiles;
+	}
+
+	/**
+	 * Return a collection of all classes associated with this package.  This is not valid until parsing begins.
+	 */
+	TArray<UClass*>& GetAllClasses()
+	{
+		return AllClasses;
+	}
+
+	/**
+	 * If true, this package should generate the classes H file.  This is not valid until code generation begins.
+	 */
+	bool GetWriteClassesH() const
+	{
+		return bWriteClassesH;
+	}
+
+	/**
+	 * Set the flag indicating that the classes H file should be generated.
+	 */
+	void SetWriteClassesH(bool bInWriteClassesH)
+	{
+		bWriteClassesH = bInWriteClassesH;
+	}
+
+	/**
+	 * Return a string that references the "PACKAGE_API " macro with a trailing space.
+	 */
+	const FString& GetAPI() const
+	{
+		return API;
+	}
+
+	/**
+	 * Get the short name of the package uppercased.
+	 */
+	const FString& GetShortUpperName() const
+	{
+		return ShortUpperName;
+	}
+
+private:
+	const FManifestModule& Module;
+	TArray<TSharedRef<FUnrealSourceFile>> AllSourceFiles;
+	TArray<UClass*> AllClasses;
+	FString ShortUpperName;
+	FString API;
+	bool bWriteClassesH = false;
+};
+
+/**
+ * Class that stores information about type definitions derived from UField.
+ */
+class FUnrealFieldDefinitionInfo 
+	: public FUnrealObjectDefinitionInfo
+{
+protected:
+	using FUnrealObjectDefinitionInfo::FUnrealObjectDefinitionInfo;
+
+public:
+
+	virtual FUnrealFieldDefinitionInfo* AsField() override
+	{
+		return this;
+	}
+
+	/**
+	 * Return the Engine instance associated with the compiler instance
+	 */
+	UField* GetField() const
+	{
+		return static_cast<UField*>(GetObject());
+	}
+
+private:
+};
+
+/**
+ * Class that stores information about type definitions derived from UEnum
+ */
+class FUnrealEnumDefinitionInfo : public FUnrealFieldDefinitionInfo
+{
+public:
+	using FUnrealFieldDefinitionInfo::FUnrealFieldDefinitionInfo;
+
+	virtual FUnrealEnumDefinitionInfo* AsEnum() override
+	{
+		return this;
+	}
+
+	/**
+	 * Return the Engine instance associated with the compiler instance
+	 */
+	UEnum* GetEnum() const
+	{
+		return static_cast<UEnum*>(GetObject());
+	}
+
+	/**
+	 * Returns the underlying enumeration type
+	 */
+	EUnderlyingEnumType GetUnderlyingType() const
+	{
+		return UnderlyingType;
+	}
+
+	/** 
+	 * Set the underlying enum type
+	 */
+	void SetUnderlyingType(EUnderlyingEnumType InUnderlyingType)
+	{
+		UnderlyingType = InUnderlyingType;
+	}
+
+	/**
+	 * Return true if the enumeration is editor only
+	 */
+	bool IsEditorOnly() const
+	{
+		return bIsEditorOnly;
+	}
+
+	/**
+	 * Make the enumeration editor only
+	 */
+	void MakeEditorOnly()
+	{
+		bIsEditorOnly = true;
+	}
+
+private:
+	EUnderlyingEnumType UnderlyingType = EUnderlyingEnumType::Unspecified;
+	bool bIsEditorOnly = false;
+};
+
+/**
+ * Class that stores information about type definitions derived from UStruct
+ */
+class FUnrealStructDefinitionInfo : public FUnrealFieldDefinitionInfo
+{
+public:
+	using FUnrealFieldDefinitionInfo::FUnrealFieldDefinitionInfo;
+
+	virtual FUnrealStructDefinitionInfo* AsStruct() override
+	{
+		return this;
+	}
+
+	virtual TSharedRef<FScope> GetScope() override;
+
+	virtual void SetObject(UObject* InObject) override;
+
+	/**
+	 * Return the Engine instance associated with the compiler instance
+	 */
+	UStruct* GetStruct() const
+	{
+		return static_cast<UStruct*>(GetObject());
+	}
+
+private:
+	TSharedPtr<FScope> StructScope;
+};
+
+/**
+ * Class that stores information about type definitions derived from UScriptStruct
+ */
+class FUnrealScriptStructDefinitionInfo : public FUnrealStructDefinitionInfo
+{
+public:
+	FUnrealScriptStructDefinitionInfo(FUnrealSourceFile& InSourceFile, int32 InLineNumber, FString&& InNameCPP, FString&& InParentScopeCPP, FString&& InParentNameCPP)
+		: FUnrealStructDefinitionInfo(InSourceFile, InLineNumber, MoveTemp(InNameCPP))
+		, ParentScopeCPP(MoveTemp(InParentScopeCPP))
+		, ParentNameCPP(MoveTemp(InParentNameCPP))
+	{ }
+
+	virtual FUnrealScriptStructDefinitionInfo* AsScriptStruct() override
+	{
+		return this;
+	}
+
+	virtual uint32 GetHash(bool bIncludeNoExport = true) const override;
+
+	/**
+	 * Return the Engine instance associated with the compiler instance
+	 */
+	UScriptStruct* GetScriptStruct() const
+	{
+		return static_cast<UScriptStruct*>(GetObject());
+	}
+
+	/**
+	 * Return the parent structure scoped name (currently blank)
+	 */
+	const FString& GetParentScopeCPP() const
+	{
+		return ParentScopeCPP;
+	}
+
+	/**
+	 * Return the name of the parent structure.  If not derived from another struct, it will be empty
+	 */
+	const FString& GetParentNameCPP() const
+	{
+		return ParentNameCPP;
+	}
+
+private:
+	FString ParentScopeCPP;
+	FString ParentNameCPP;
+};
+
+/**
+ * Class that stores information about type definitions derived from UClass
+ */
+class FUnrealClassDefinitionInfo
+	: public FUnrealStructDefinitionInfo
+	, public FClassDeclarationMetaData
+{
+public:
+	FUnrealClassDefinitionInfo(FUnrealSourceFile& InSourceFile, int32 InLineNumber, FString&& InNameCPP, FString&& InBaseClassNameCPP, bool bInIsInterface)
+		: FUnrealStructDefinitionInfo(InSourceFile, InLineNumber, MoveTemp(InNameCPP))
+		, BaseClassNameCPP(MoveTemp(InBaseClassNameCPP))
+		, bIsInterface(bInIsInterface)
+	{ }
+
+	virtual FUnrealClassDefinitionInfo* AsClass() override
+	{
+		return this;
+	}
+
+	virtual uint32 GetHash(bool bIncludeNoExport = true) const override;
+
+	/**
+	 * Return the Engine instance associated with the compiler instance
+	 */
+	UClass* GetClass() const
+	{
+		return static_cast<UClass*>(GetObject());
+	}
+
+	/**
+	 * Get the archive type
+	 */
+	ESerializerArchiveType GetArchiveType() const
+	{
+		return ArchiveType;
+	}
+
+	/**
+	 * Set the archive type
+	 */
+	void AddArchiveType(ESerializerArchiveType InArchiveType)
+	{
+		ArchiveType |= InArchiveType;
+	}
+
+	/**
+	 * Get the enclosing define
+	 */
+	const FString& GetEnclosingDefine() const
+	{
+		return EnclosingDefine;
+	}
+
+	/** 
+	 * Set the enclosing define
+	 */
+	void SetEnclosingDefine(FString&& InEnclosingDefine)
+	{
+		EnclosingDefine = MoveTemp(InEnclosingDefine);
+	}
+
+	/**
+	 * Return true if this is an interface
+	 */
+	bool IsInterface() const
+	{
+		return bIsInterface;
+	}
+
+	/** 
+	 * Return the CPP name of the base class or blank if there is none.
+	 */
+	const FString& GetBaseClassNameCPP() const
+	{
+		return BaseClassNameCPP;
+	}
+
+private:
+	FString BaseClassNameCPP;
+	FString EnclosingDefine;
+	ESerializerArchiveType ArchiveType = ESerializerArchiveType::None;
+	bool bIsInterface = false;
 };
