Quick fix to Enum type definition constructor.
Moving the UHTConfig to BaseParser along with a support routine to be used by a future change.

#rnx

diff --git a/Engine/Source/Programs/UnrealHeaderTool/Private/BaseParser.cpp b/Engine/Source/Programs/UnrealHeaderTool/Private/BaseParser.cpp
--- a/Engine/Source/Programs/UnrealHeaderTool/Private/BaseParser.cpp
+++ b/Engine/Source/Programs/UnrealHeaderTool/Private/BaseParser.cpp
@@ -181,7 +181,8 @@ FString FPropertySpecifier::ConvertToString() const
 // FBaseParser
 
 FBaseParser::FBaseParser()
-	: StatementsParsed(0)
+	: UHTConfig(UHTConfig.Get())
+	, StatementsParsed(0)
 	, LinesParsed(0)
 {
 }
@@ -734,6 +735,18 @@ bool FBaseParser::GetSymbol( FToken& Token )
 	return false;
 }
 
+// Modify token to fix redirected types if needed
+void FBaseParser::RedirectTypeIdentifier(FToken& Token) const
+{
+	check(Token.TokenType == TOKEN_Identifier);
+
+	const FString* FoundRedirect = UHTConfig.TypeRedirectMap.Find(Token.Identifier);
+	if (FoundRedirect)
+	{
+		Token.SetIdentifier(**FoundRedirect);
+	}
+}
+
 bool FBaseParser::GetConstInt(int32& Result, const TCHAR* Tag)
 {
 	FToken Token;
diff --git a/Engine/Source/Programs/UnrealHeaderTool/Private/BaseParser.h b/Engine/Source/Programs/UnrealHeaderTool/Private/BaseParser.h
--- a/Engine/Source/Programs/UnrealHeaderTool/Private/BaseParser.h
+++ b/Engine/Source/Programs/UnrealHeaderTool/Private/BaseParser.h
@@ -4,12 +4,50 @@
 
 #include "CoreMinimal.h"
 #include "Containers/UnrealString.h"
+#include "GeneratedCodeVersion.h"
 
 class FToken;
 
+enum class EPointerMemberBehavior
+{
+	Disallow,
+	AllowSilently,
+	AllowAndLog,
+};
+
 /////////////////////////////////////////////////////
-// FBaseParser
+// UHTConfig
+
+struct FUHTConfig
+{
+	static const FUHTConfig& Get();
+
+	// Types that have been renamed, treat the old deprecated name as the new name for code generation
+	TMap<FString, FString> TypeRedirectMap;
+
+	// Special parsed struct names that do not require a prefix
+	TArray<FString> StructsWithNoPrefix;
+
+	// Special parsed struct names that have a 'T' prefix
+	TArray<FString> StructsWithTPrefix;
+
+	// Mapping from 'human-readable' macro substring to # of parameters for delegate declarations
+	// Index 0 is 1 parameter, Index 1 is 2, etc...
+	TArray<FString> DelegateParameterCountStrings;
 
+	// Default version of generated code. Defaults to oldest possible, unless specified otherwise in config.
+	EGeneratedCodeVersion DefaultGeneratedCodeVersion = EGeneratedCodeVersion::V1;
+
+	EPointerMemberBehavior NativePointerMemberBehavior = EPointerMemberBehavior::AllowSilently;
+
+	EPointerMemberBehavior ObjectPtrMemberBehavior = EPointerMemberBehavior::AllowSilently;
+
+private:
+	FUHTConfig();
+};
+
+/////////////////////////////////////////////////////
+// FBaseParser
 
 enum class ESymbolParseOption
 {
@@ -46,6 +84,9 @@ class FBaseParser
 protected:
 	FBaseParser();
 
+	// UHTConfig data
+	const FUHTConfig& UHTConfig;
+
 public:
 	// Input text.
 	const TCHAR* Input;
@@ -131,6 +172,9 @@ public:
 	bool GetIdentifier( FToken& Token, bool bNoConsts = false );
 	bool GetSymbol( FToken& Token );
 
+	// Modify token to fix redirected types if needed
+	void RedirectTypeIdentifier(FToken& Token) const;
+
 	/**
 	 * Get an int constant
 	 * @return true on success, otherwise false.
diff --git a/Engine/Source/Programs/UnrealHeaderTool/Private/HeaderParser.cpp b/Engine/Source/Programs/UnrealHeaderTool/Private/HeaderParser.cpp
--- a/Engine/Source/Programs/UnrealHeaderTool/Private/HeaderParser.cpp
+++ b/Engine/Source/Programs/UnrealHeaderTool/Private/HeaderParser.cpp
@@ -6974,18 +6974,6 @@ bool FHeaderParser::IsValidDelegateDeclaration(const FToken& Token) const
 	return (Token.TokenType == TOKEN_Identifier) && !FCString::Strncmp(Token.Identifier, TEXT("DECLARE_DYNAMIC_"), 16);
 }
 
-// Modify token to fix redirected types if needed
-void FHeaderParser::RedirectTypeIdentifier(FToken& Token) const
-{
-	check(Token.TokenType == TOKEN_Identifier);
-
-	const FString* FoundRedirect = UHTConfig.TypeRedirectMap.Find(Token.Identifier);
-	if (FoundRedirect)
-	{
-		Token.SetIdentifier(**FoundRedirect);
-	}
-}
-
 // Parse the parameter list of a function or delegate declaration
 void FHeaderParser::ParseParameterList(UFunction* Function, bool bExpectCommaBeforeName, TMap<FName, FString>* MetaData)
 {
@@ -8977,7 +8965,6 @@ FHeaderParser::FHeaderParser(FFeedbackContext* InWarn, const FManifestModule& In
 	, bSpottedAutogeneratedHeaderInclude(false)
 	, NestLevel(0)
 	, TopNest(nullptr)
-	, UHTConfig(FUHTConfig::Get())
 	, CurrentlyParsedModule(&InModule)
 {
 	// Determine if the current module is part of the engine or a game (we are more strict about things for Engine modules)
diff --git a/Engine/Source/Programs/UnrealHeaderTool/Private/HeaderParser.h b/Engine/Source/Programs/UnrealHeaderTool/Private/HeaderParser.h
--- a/Engine/Source/Programs/UnrealHeaderTool/Private/HeaderParser.h
+++ b/Engine/Source/Programs/UnrealHeaderTool/Private/HeaderParser.h
@@ -419,43 +419,6 @@ struct FDocumentationPolicy
 	bool bFloatRangesRequired = UHT_DOCUMENTATION_POLICY_DEFAULT;
 };
 
-enum class EPointerMemberBehavior
-{
-	Disallow,
-	AllowSilently,
-	AllowAndLog,
-};
-
-/////////////////////////////////////////////////////
-// UHTConfig
-
-struct FUHTConfig
-{
-	FUHTConfig();
-
-	static const FUHTConfig& Get();
-
-	// Types that have been renamed, treat the old deprecated name as the new name for code generation
-	TMap<FString, FString> TypeRedirectMap;
-
-	// Special parsed struct names that do not require a prefix
-	TArray<FString> StructsWithNoPrefix;
-
-	// Special parsed struct names that have a 'T' prefix
-	TArray<FString> StructsWithTPrefix;
-
-	// Mapping from 'human-readable' macro substring to # of parameters for delegate declarations
-	// Index 0 is 1 parameter, Index 1 is 2, etc...
-	TArray<FString> DelegateParameterCountStrings;
-
-	// Default version of generated code. Defaults to oldest possible, unless specified otherwise in config.
-	EGeneratedCodeVersion DefaultGeneratedCodeVersion = EGeneratedCodeVersion::V1;
-
-	EPointerMemberBehavior NativePointerMemberBehavior = EPointerMemberBehavior::AllowSilently;
-
-	EPointerMemberBehavior ObjectPtrMemberBehavior = EPointerMemberBehavior::AllowSilently;
-};
-
 /////////////////////////////////////////////////////
 // FHeaderParser
 
@@ -694,9 +657,6 @@ protected:
 
 	////////////////////////////////////////////////////
 
-	// UHTConfig data
-	const FUHTConfig& UHTConfig;
-
 	// List of all used identifiers for net service function declarations (every function must be unique)
 	TMap<int32, FString> UsedRPCIds;
 	// List of all net service functions with undeclared response functions 
@@ -717,9 +677,6 @@ protected:
 	// Parse the parameter list of a function or delegate declaration
 	void ParseParameterList(UFunction* Function, bool bExpectCommaBeforeName = false, TMap<FName, FString>* MetaData = NULL);
 
-	// Modify token to fix redirected types if needed
-	void RedirectTypeIdentifier(FToken& Token) const;
-
 public:
 	// Throws if a specifier value wasn't provided
 	static void RequireSpecifierValue(const FPropertySpecifier& Specifier, bool bRequireExactlyOne = false);
diff --git a/Engine/Source/Programs/UnrealHeaderTool/Private/UnrealTypeDefinitionInfo.h b/Engine/Source/Programs/UnrealHeaderTool/Private/UnrealTypeDefinitionInfo.h
--- a/Engine/Source/Programs/UnrealHeaderTool/Private/UnrealTypeDefinitionInfo.h
+++ b/Engine/Source/Programs/UnrealHeaderTool/Private/UnrealTypeDefinitionInfo.h
@@ -527,7 +527,9 @@ private:
 class FUnrealEnumDefinitionInfo : public FUnrealFieldDefinitionInfo
 {
 public:
-	using FUnrealFieldDefinitionInfo::FUnrealFieldDefinitionInfo;
+	FUnrealEnumDefinitionInfo(FUnrealSourceFile& InSourceFile, int32 InLineNumber, FString&& InNameCPP)
+		: FUnrealFieldDefinitionInfo(InSourceFile, InLineNumber, MoveTemp(InNameCPP))
+	{ }
 
 	virtual FUnrealEnumDefinitionInfo* AsEnum() override
 	{
