Fix to a silly compile error of mine on Mac.

#rnx

diff --git a/Engine/Source/Programs/UnrealHeaderTool/Private/BaseParser.cpp b/Engine/Source/Programs/UnrealHeaderTool/Private/BaseParser.cpp
--- a/Engine/Source/Programs/UnrealHeaderTool/Private/BaseParser.cpp
+++ b/Engine/Source/Programs/UnrealHeaderTool/Private/BaseParser.cpp
@@ -181,7 +181,7 @@ FString FPropertySpecifier::ConvertToString() const
 // FBaseParser
 
 FBaseParser::FBaseParser()
-	: UHTConfig(UHTConfig.Get())
+	: UHTConfig(FUHTConfig::Get())
 	, StatementsParsed(0)
 	, LinesParsed(0)
 {
