Fix MAC compile issue
Fix coding standard error with enum declaration.

#rnx

diff --git a/Engine/Source/Programs/UnrealHeaderTool/Private/CodeGenerator.cpp b/Engine/Source/Programs/UnrealHeaderTool/Private/CodeGenerator.cpp
--- a/Engine/Source/Programs/UnrealHeaderTool/Private/CodeGenerator.cpp
+++ b/Engine/Source/Programs/UnrealHeaderTool/Private/CodeGenerator.cpp
@@ -6011,9 +6011,9 @@ void FNativeClassHeaderGenerator::GenerateSourceFiles(
 		{
 			FResults::Try(GeneratedCPP.SourceFile, [&GeneratedCPP]()
 				{
-					FUnrealPackageDefinitionInfo& PackageDef = GeneratedCPP.PackageDef;
-					const FManifestModule& Module = PackageDef.GetModule();
-					const FNativeClassHeaderGenerator Generator(PackageDef);
+					FUnrealPackageDefinitionInfo& PackageDefLcl = GeneratedCPP.PackageDef;
+					const FManifestModule& Module = PackageDefLcl.GetModule();
+					const FNativeClassHeaderGenerator Generator(PackageDefLcl);
 					FUnrealSourceFile& SourceFile = GeneratedCPP.SourceFile;
 					FUHTStringBuilder& GeneratedFunctionDeclarations = GeneratedCPP.GeneratedFunctionDeclarations;
 					FUHTStringBuilder& GeneratedHeaderText = GeneratedCPP.Header.GetGeneratedBody();
@@ -6107,8 +6107,8 @@ void FNativeClassHeaderGenerator::GenerateSourceFiles(
 		{
 			FResults::Try(GeneratedCPP.SourceFile, [&GeneratedCPP]()
 				{
-					FUnrealPackageDefinitionInfo& PackageDef = GeneratedCPP.PackageDef;
-					const FManifestModule& Module = PackageDef.GetModule();
+					FUnrealPackageDefinitionInfo& PackageDefLcl = GeneratedCPP.PackageDef;
+					const FManifestModule& Module = PackageDefLcl.GetModule();
 					FUnrealSourceFile& SourceFile = GeneratedCPP.SourceFile;
 					FUHTStringBuilder& GeneratedHeaderText = GeneratedCPP.Header.GetGeneratedBody();
 					FUHTStringBuilder& GeneratedCPPText = GeneratedCPP.Source.GetGeneratedBody();
@@ -7282,7 +7282,7 @@ void GenerateSummary(TArray<FUnrealPackageDefinitionInfo*>& PackageDefs)
 	{
 		const FManifestModule& Module = PackageDef->GetModule();
 
-		double TotalTimes[ESourceFileTime::Count] = { 0.0 };
+		double TotalTimes[int32(ESourceFileTime::Count)] = { 0.0 };
 		int32 LinesParsed = 0;
 		int32 StatementsParsed = 0;
 		int32 SourceCount = 0;
@@ -7291,13 +7291,13 @@ void GenerateSummary(TArray<FUnrealPackageDefinitionInfo*>& PackageDefs)
 		{
 			for (int32 Index = 0; Index < int32(ESourceFileTime::Count); ++Index)
 			{
-				TotalTimes[Index] += SourceFile->GetTime(ESourceFileTime(Index));
+				TotalTimes[int32(Index)] += SourceFile->GetTime(ESourceFileTime(Index));
 			}
 			LinesParsed += SourceFile->GetLinesParsed();
 			StatementsParsed += SourceFile->GetStatementsParsed();
 		}
 		UE_LOG(LogCompile, Log, TEXT("Success: Module %s parsed %d sources(s), %d line(s), %d statement(s).  Times(secs) Load: %.3f, PreParse: %.3f, Parse: %.3f, Generate: %.3f."), *Module.Name,
-			SourceFiles.Num(), LinesParsed, StatementsParsed, TotalTimes[ESourceFileTime::Load], TotalTimes[ESourceFileTime::PreParse], TotalTimes[ESourceFileTime::Parse], TotalTimes[ESourceFileTime::Generate]);
+			SourceFiles.Num(), LinesParsed, StatementsParsed, TotalTimes[int32(ESourceFileTime::Load)], TotalTimes[int32(ESourceFileTime::PreParse)], TotalTimes[int32(ESourceFileTime::Parse)], TotalTimes[int32(ESourceFileTime::Generate)]);
 	}
 }
 
diff --git a/Engine/Source/Programs/UnrealHeaderTool/Private/UnrealSourceFile.h b/Engine/Source/Programs/UnrealHeaderTool/Private/UnrealSourceFile.h
--- a/Engine/Source/Programs/UnrealHeaderTool/Private/UnrealSourceFile.h
+++ b/Engine/Source/Programs/UnrealHeaderTool/Private/UnrealSourceFile.h
@@ -15,14 +15,14 @@ class UStruct;
 class FArchive;
 class FClassMetaData;
 
-enum ETopologicalState : uint8
+enum class ETopologicalState : uint8
 {
 	Unmarked,
 	Temporary,
 	Permanent,
 };
 
-enum ESourceFileTime : uint8
+enum class ESourceFileTime : uint8
 {
 	Load,
 	PreParse,
@@ -400,7 +400,7 @@ public:
 	 */
 	double& GetTime(ESourceFileTime Time)
 	{
-		return Times[Time];
+		return Times[int32(Time)];
 	}
 
 	/**
@@ -476,7 +476,7 @@ private:
 	FString Content;
 
 	// Different timers for the source
-	double Times[ESourceFileTime::Count] = { 0.0 };
+	double Times[int32(ESourceFileTime::Count)] = { 0.0 };
 
 	// Number of statements parsed.
 	int32 StatementsParsed = 0;
