Fixing UHT pre-parsing issue exposed by the addition of UENUM And USTRUCT parsing.

#rnx

diff --git a/Engine/Source/Programs/UnrealHeaderTool/Private/HeaderParser.cpp b/Engine/Source/Programs/UnrealHeaderTool/Private/HeaderParser.cpp
--- a/Engine/Source/Programs/UnrealHeaderTool/Private/HeaderParser.cpp
+++ b/Engine/Source/Programs/UnrealHeaderTool/Private/HeaderParser.cpp
@@ -9301,11 +9301,10 @@ void FHeaderParser::SimplifiedClassParse(FUnrealSourceFile& SourceFile, const TC
 	CurrentLine = 0;
 	Buffer      = *ClassHeaderTextStrippedOfCppText;
 
-	const TCHAR* StartOfLine            = Buffer;
 	bool         bFoundGeneratedInclude = false;
 	bool         bFoundExportedClasses  = false;
 
-	while (FParse::Line(&Buffer, StrLine, true))
+	for (const TCHAR* StartOfLine = Buffer; FParse::Line(&Buffer, StrLine, true); StartOfLine = Buffer)
 	{
 		CurrentLine++;
 
@@ -9511,8 +9510,6 @@ void FHeaderParser::SimplifiedClassParse(FUnrealSourceFile& SourceFile, const TC
 				}
 			}
 		}
-	
-		StartOfLine = Buffer;
 	}
 
 	if (bFoundExportedClasses && !bFoundGeneratedInclude)
