Enable UEnum and UStruct change tracking for properties in UHT.  This fixes LiveCoding and HotReload issues where the engine wouldn't detect changes to containing objects when the UEnum or UStruct changed.

#rnx

diff --git a/Engine/Source/Programs/UnrealHeaderTool/Private/UnrealHeaderTool.h b/Engine/Source/Programs/UnrealHeaderTool/Private/UnrealHeaderTool.h
--- a/Engine/Source/Programs/UnrealHeaderTool/Private/UnrealHeaderTool.h
+++ b/Engine/Source/Programs/UnrealHeaderTool/Private/UnrealHeaderTool.h
@@ -3,7 +3,7 @@
 #include "CoreMinimal.h"
 
 // Enable property tags for enums and structures
-#define UHT_ENABLE_VALUE_PROPERTY_TAG 0
+#define UHT_ENABLE_VALUE_PROPERTY_TAG 1
 
 // Enable property tags for class and interfaces wrapped in porters 
 // DO NOT ENABLE: This does not handle forward declarations.  Also,
