1) Create the Function type information that provides the UHT data for functions.
2) Remove the global hash storage system since it is moved into the Function type information.
3) Remove some unsed classes and methods.
4) Added the TEMPORARY auto creation of base engine types until they can be added to the no export file.

#rnx

diff --git a/Engine/Source/Programs/UnrealHeaderTool/Private/ClassMaps.cpp b/Engine/Source/Programs/UnrealHeaderTool/Private/ClassMaps.cpp
--- a/Engine/Source/Programs/UnrealHeaderTool/Private/ClassMaps.cpp
+++ b/Engine/Source/Programs/UnrealHeaderTool/Private/ClassMaps.cpp
@@ -6,5 +6,3 @@
 
 FUnrealSourceFiles GUnrealSourceFilesMap;
 FTypeDefinitionInfoMap GTypeDefinitionInfoMap;
-TMap<UFunction*, uint32> GGeneratedCodeHashes;
-FRWLock GGeneratedCodeHashesLock;
diff --git a/Engine/Source/Programs/UnrealHeaderTool/Private/ClassMaps.h b/Engine/Source/Programs/UnrealHeaderTool/Private/ClassMaps.h
--- a/Engine/Source/Programs/UnrealHeaderTool/Private/ClassMaps.h
+++ b/Engine/Source/Programs/UnrealHeaderTool/Private/ClassMaps.h
@@ -71,10 +71,18 @@ struct FTypeDefinitionInfoMap : public FFreezableContainer
 		return **TypeDef;
 	}
 
+	//NOTE: Currently UFunctions are created during the parsing phase and can not be frozen
+	void Add(UFunction* Object, TSharedRef<FUnrealTypeDefinitionInfo>&& Definition)
+	{
+		DefinitionsByUObject.Add(Object, Definition);
+		// At this point, do not add the name
+		//DefinitionsByName.Add(Object->GetFName(), Definition);
+	}
+
 	//NOTE: FFields (properties) are not frozen since they are added during the parsing phase
 	void Add(FField* Field, TSharedRef<FUnrealTypeDefinitionInfo>&& Definition)
 	{
-		DefinitionsByFField.Add(Field, Definition);
+		DefinitionsByFField.Add(Field, MoveTemp(Definition));
 	}
 	bool Contains(const FField* Field) { return DefinitionsByFField.Contains(Field); }
 	TSharedRef<FUnrealTypeDefinitionInfo>* Find(const FField* Field)
@@ -128,32 +136,3 @@ private:
 
 extern FUnrealSourceFiles GUnrealSourceFilesMap;
 extern FTypeDefinitionInfoMap GTypeDefinitionInfoMap;
-extern TMap<UFunction*, uint32> GGeneratedCodeHashes;
-extern FRWLock GGeneratedCodeHashesLock;
-
-/** Types access specifiers. */
-enum EAccessSpecifier
-{
-	ACCESS_NotAnAccessSpecifier = 0,
-	ACCESS_Public,
-	ACCESS_Private,
-	ACCESS_Protected,
-	ACCESS_Num,
-};
-
-inline FArchive& operator<<(FArchive& Ar, EAccessSpecifier& ObjectType)
-{
-	if (Ar.IsLoading())
-	{
-		int32 Value;
-		Ar << Value;
-		ObjectType = EAccessSpecifier(Value);
-	}
-	else if (Ar.IsSaving())
-	{
-		int32 Value = (int32)ObjectType;
-		Ar << Value;
-	}
-
-	return Ar;
-}
diff --git a/Engine/Source/Programs/UnrealHeaderTool/Private/CodeGenerator.cpp b/Engine/Source/Programs/UnrealHeaderTool/Private/CodeGenerator.cpp
--- a/Engine/Source/Programs/UnrealHeaderTool/Private/CodeGenerator.cpp
+++ b/Engine/Source/Programs/UnrealHeaderTool/Private/CodeGenerator.cpp
@@ -44,6 +44,7 @@
 #include "UnrealSourceFile.h"
 #include "ParserHelper.h"
 #include "Classes.h"
+#include "ClassMaps.h"
 #include "NativeClassExporter.h"
 #include "ProfilingDebugging/ScopedTimers.h"
 #include "HeaderParser.h"
@@ -60,7 +61,6 @@
 #include "UnrealTypeDefinitionInfo.h"
 #include "Misc/WildcardString.h"
 
-#include "Exceptions.h"
 #include "UObject/FieldIterator.h"
 #include "UObject/FieldPath.h"
 
@@ -491,36 +491,6 @@ FString Macroize(const TCHAR* MacroName, FString&& StringToMacroize)
 	return FString::Printf(TEXT("#define %s%s\r\n%s"), MacroName, Result.Len() ? TEXT(" \\") : TEXT(""), *Result);
 }
 
-/** Add a hash for the given function */
-static void AddGeneratedCodeHash(UFunction* Field, uint32 Hash)
-{
-	FRWScopeLock Lock(GGeneratedCodeHashesLock, SLT_Write);
-	GGeneratedCodeHashes.Add(Field, Hash);
-}
-
-/** Generates a Hash tag string for the specified function */
-static FString GetGeneratedCodeHashTag(UFunction* Field)
-{
-	FString Tag;
-	bool bFoundHash = false;
-	uint32 Hash = 0;
-
-	{
-		FRWScopeLock Lock(GGeneratedCodeHashesLock, SLT_ReadOnly);
-		if (const uint32* FieldHash = GGeneratedCodeHashes.Find(Field))
-		{
-			bFoundHash = true;
-			Hash = *FieldHash;
-		}
-	}
-
-	if (bFoundHash)
-	{
-		Tag = FString::Printf(TEXT(" // %u"), Hash);
-	}
-	return Tag;
-}
-
 struct FParmsAndReturnProperties
 {
 	FParmsAndReturnProperties()
@@ -2441,6 +2411,8 @@ void FNativeClassHeaderGenerator::ExportNativeGeneratedInitCode(FOutputDevice& O
 
 			for (UFunction* Function : FunctionsToExport)
 			{
+				FUnrealFunctionDefinitionInfo& FuncDef = GTypeDefinitionInfoMap.FindChecked(Function).AsFunctionChecked();
+
 				const bool bIsEditorOnlyFunction = Function->HasAnyFunctionFlags(FUNC_EditorOnly);
 
 				if (!Function->IsA<UDelegateFunction>())
@@ -2448,12 +2420,15 @@ void FNativeClassHeaderGenerator::ExportNativeGeneratedInitCode(FOutputDevice& O
 					ExportFunction(Out, OutReferenceGatherers, SourceFile, Function, bIsNoExport);
 				}
 
+				FUHTStringBuilder FuncHashTag;
+				FuncDef.GetHashTag(FuncHashTag);
+
 				StaticDefinitions.Logf(
 					TEXT("%s\t\t{ &%s, %s },%s\r\n%s"),
 					BEGIN_WRAP_EDITOR_ONLY(bIsEditorOnlyFunction),
 					*GetSingletonNameFuncAddr(Function, OutReferenceGatherers.UniqueCrossModuleReferences),
 					*FNativeClassHeaderGenerator::GetUTF8OverriddenNameForLiteral(Function),
-					*GetGeneratedCodeHashTag(Function),
+					*FuncHashTag,
 					END_WRAP_EDITOR_ONLY(bIsEditorOnlyFunction)
 				);
 			}
@@ -2815,8 +2790,8 @@ void FNativeClassHeaderGenerator::ExportFunction(FOutputDevice& Out, FReferenceG
 
 	TTuple<FString, FString> PropertyRange = OutputProperties(CurrentFunctionText, StaticDefinitions, OutReferenceGatherers, *FString::Printf(TEXT("%s::"), *StaticsStructName), Props, TEXT("\t\t"), TEXT("\t"));
 
-	const FFunctionData* CompilerInfo = FFunctionData::FindForFunction(Function);
-	const FFuncInfo&     FunctionData = CompilerInfo->GetFunctionData();
+	FUnrealFunctionDefinitionInfo& FuncDef = GTypeDefinitionInfoMap.FindChecked(Function).AsFunctionChecked();
+	const FFuncInfo&     FunctionData = FuncDef.GetFunctionData();
 	const bool           bIsNet       = !!(FunctionData.FunctionFlags & (FUNC_NetRequest | FUNC_NetResponse));
 
 	FString MetaDataParams = OutputMetaDataCodeForObject(CurrentFunctionText, StaticDefinitions, Function, *FString::Printf(TEXT("%s::Function_MetaDataParams"), *StaticsStructName), TEXT("\t\t"), TEXT("\t"));
@@ -2872,7 +2847,7 @@ void FNativeClassHeaderGenerator::ExportFunction(FOutputDevice& Out, FReferenceG
 	}
 
 	uint32 FunctionHash = GenerateTextHash(*CurrentFunctionText);
-	AddGeneratedCodeHash(Function, FunctionHash);
+	FuncDef.SetHash(FunctionHash);
 	Out.Log(CurrentFunctionText);
 }
 
@@ -2953,9 +2928,8 @@ void FNativeClassHeaderGenerator::ExportInterfaceCallFunctions(FOutputDevice& Ou
 	{
 		FString FunctionName = Function->GetName();
 
-		FFunctionData* CompilerInfo = FFunctionData::FindForFunction(Function);
-
-		const FFuncInfo& FunctionData = CompilerInfo->GetFunctionData();
+		FUnrealFunctionDefinitionInfo& FuncDef = GTypeDefinitionInfoMap.FindChecked(Function).AsFunctionChecked();
+		const FFuncInfo& FunctionData = FuncDef.GetFunctionData();
 		const TCHAR* ConstQualifier = FunctionData.FunctionReference->HasAllFunctionFlags(FUNC_Const) ? TEXT("const ") : TEXT("");
 		FString ExtraParam = FString::Printf(TEXT("%sUObject* O"), ConstQualifier);
 
@@ -4582,9 +4556,8 @@ void FNativeClassHeaderGenerator::ExportDelegateDeclaration(FOutputDevice& Out,
 	// Unmangle the function name
 	const FString DelegateName = Function->GetName().LeftChop(HEADER_GENERATED_DELEGATE_SIGNATURE_SUFFIX_LENGTH);
 
-	const FFunctionData* CompilerInfo = FFunctionData::FindForFunction(Function);
-
-	FFuncInfo FunctionData = CompilerInfo->GetFunctionData();
+	FUnrealFunctionDefinitionInfo& FuncDef = GTypeDefinitionInfoMap.FindChecked(Function).AsFunctionChecked();
+	FFuncInfo FunctionData = FuncDef.GetFunctionData();
 
 	// Add class name to beginning of function, to avoid collisions with other classes with the same delegate name in this scope
 	check(FunctionData.MarshallAndCallName.StartsWith(DelegateStr));
@@ -4625,9 +4598,8 @@ void FNativeClassHeaderGenerator::ExportDelegateDefinition(FOutputDevice& Out, F
 	// Unmangle the function name
 	const FString DelegateName = Function->GetName().LeftChop(HEADER_GENERATED_DELEGATE_SIGNATURE_SUFFIX_LENGTH);
 
-	const FFunctionData* CompilerInfo = FFunctionData::FindForFunction(Function);
-
-	FFuncInfo FunctionData = CompilerInfo->GetFunctionData();
+	FUnrealFunctionDefinitionInfo& FuncDef = GTypeDefinitionInfoMap.FindChecked(Function).AsFunctionChecked();
+	FFuncInfo FunctionData = FuncDef.GetFunctionData();
 
 	// Always export delegate wrapper functions as inline
 	FunctionData.FunctionExportFlags |= FUNCEXPORT_Inline;
@@ -5534,9 +5506,8 @@ void FNativeClassHeaderGenerator::ExportNativeFunctions(FOutputDevice& OutGenera
 		const bool bEditorOnlyFunc = Function->HasAnyFunctionFlags(FUNC_EditorOnly);
 		FNativeFunctionStringBuilder& FuncStringBuilders = bEditorOnlyFunc ? EditorStringBuilders : RuntimeStringBuilders;
 
-		FFunctionData* CompilerInfo = FFunctionData::FindForFunction(Function);
-
-		const FFuncInfo& FunctionData = CompilerInfo->GetFunctionData();
+		FUnrealFunctionDefinitionInfo& FuncDef = GTypeDefinitionInfoMap.FindChecked(Function).AsFunctionChecked();
+		const FFuncInfo& FunctionData = FuncDef.GetFunctionData();
 
 		// Custom thunks don't get any C++ stub function generated
 		if (FunctionData.FunctionExportFlags & FUNCEXPORT_CustomThunk)
@@ -5731,8 +5702,8 @@ void FNativeClassHeaderGenerator::ExportCallbackFunctions(
 		// Never expecting to export delegate functions this way
 		check(!Function->HasAnyFunctionFlags(FUNC_Delegate));
 
-		FFunctionData*   CompilerInfo = FFunctionData::FindForFunction(Function);
-		const FFuncInfo& FunctionData = CompilerInfo->GetFunctionData();
+		FUnrealFunctionDefinitionInfo& FuncDef = GTypeDefinitionInfoMap.FindChecked(Function).AsFunctionChecked();
+		const FFuncInfo& FunctionData = FuncDef.GetFunctionData();
 		FString          FunctionName = Function->GetName();
 		UClass*          Class = CastChecked<UClass>(Function->GetOuter());
 		const FString    ClassName = FNameLookupCPP::GetNameCPP(Class);
@@ -5911,8 +5882,11 @@ bool FNativeClassHeaderGenerator::WriteSource(const FManifestModule& Module, FGe
 			{
 				if (TSharedRef<FUnrealTypeDefinitionInfo>* WithinTypeInfo = GTypeDefinitionInfoMap.Find(Class->ClassWithin))
 				{
-					FString Header = GetBuildPath((*WithinTypeInfo)->GetUnrealSourceFile());
-					RelativeIncludes.AddUnique(MoveTemp(Header));
+					if ((*WithinTypeInfo)->HasSource())
+					{
+						FString Header = GetBuildPath((*WithinTypeInfo)->GetUnrealSourceFile());
+						RelativeIncludes.AddUnique(MoveTemp(Header));
+					}
 				}
 			}
 
@@ -7314,6 +7288,20 @@ ECompilationResult::Type UnrealHeaderTool_Main(const FString& ModuleInfoFilename
 	// The meta data keywords must be initialized prior to going wide
 	FBaseParser::InitMetadataKeywords();
 
+	// TEMPORARY!!!
+	// At this time, there is a collection of classes that aren't listed in the NoExport
+	// file.  This means that we don't have an associated type definition for them.
+	// Force the creation of type data for them.  Once these are added to NoExport, then
+	// the HasSource can be looked at to be removed.  However, it might change output slightly.  
+	// Use -WRITEREF and -VERIFYREF to detect these changes.
+	for (TObjectIterator<UClass> ClassIt; ClassIt; ++ClassIt)
+	{
+		FUnrealClassDefinitionInfo* ClassDef = new FUnrealClassDefinitionInfo(FNameLookupCPP::GetNameCPP(*ClassIt));
+		ClassDef->SetObject(*ClassIt);
+		GTypeDefinitionInfoMap.Add(*ClassIt, MakeShareable<FUnrealTypeDefinitionInfo>(ClassDef));
+	}
+
+
 	// Load the manifest file, giving a list of all modules to be processed, pre-sorted by dependency ordering
 #if !PLATFORM_EXCEPTIONS_DISABLED
 	try
diff --git a/Engine/Source/Programs/UnrealHeaderTool/Private/HeaderParser.cpp b/Engine/Source/Programs/UnrealHeaderTool/Private/HeaderParser.cpp
--- a/Engine/Source/Programs/UnrealHeaderTool/Private/HeaderParser.cpp
+++ b/Engine/Source/Programs/UnrealHeaderTool/Private/HeaderParser.cpp
@@ -14,6 +14,7 @@
 #include "ProfilingDebugging/ScopedTimers.h"
 #include "NativeClassExporter.h"
 #include "Classes.h"
+#include "ClassMaps.h"
 #include "StringUtils.h"
 #include "Misc/DefaultValueHelper.h"
 #include "Manifest.h"
@@ -1703,7 +1704,11 @@ void AddModuleRelativePathToMetadata(UField* Type, TMap<FName, FString> &MetaDat
 	TSharedRef<FUnrealTypeDefinitionInfo>* TypeDefinitionPtr = GTypeDefinitionInfoMap.Find(Type);
 	if (TypeDefinitionPtr != nullptr)
 	{
-		MetaData.Add(NAME_ModuleRelativePath, (*TypeDefinitionPtr)->GetUnrealSourceFile().GetModuleRelativePath());
+		// Don't add module relative paths to functions.
+		if ((*TypeDefinitionPtr)->AsFunction() == nullptr)
+		{
+			MetaData.Add(NAME_ModuleRelativePath, (*TypeDefinitionPtr)->GetUnrealSourceFile().GetModuleRelativePath());
+		}
 	}
 }
 
@@ -5650,7 +5655,7 @@ FProperty* FHeaderParser::GetVarNameAndDim
 	VarProperty.StartPos = InputPos;
 	FClassMetaData* ScopeData = GScriptHelper.FindClassData(Scope);
 	check(ScopeData);
-	ScopeData->AddProperty(FToken(VarProperty), &SourceFile);
+	ScopeData->AddProperty(FToken(VarProperty));
 
 	// if we had any metadata, add it to the class
 	AddMetaDataToClassData(VarProperty.TokenProperty, TMap<FName,FString>(VarProperty.MetaData));
@@ -7349,7 +7354,9 @@ UDelegateFunction* FHeaderParser::CompileDelegateDeclaration(const TCHAR* Delega
 	FCString::Strncpy(FuncInfoFunctionIdentifier, FuncInfo.Function.Identifier, NAME_SIZE);
 
 	FuncInfo.MacroLine = InputLine;
-	FFunctionData::Add(MoveTemp(FuncInfo));
+	TSharedRef<FUnrealFunctionDefinitionInfo> FuncDef = MakeShareable(new FUnrealFunctionDefinitionInfo(SourceFile, InputLine, FString(FuncInfoFunctionIdentifier), MoveTemp(FuncInfo)));
+	FuncDef->SetObject(FuncDef->GetFunctionData().FunctionReference);
+	GTypeDefinitionInfoMap.Add(FuncDef->GetFunctionData().FunctionReference, MoveTemp(FuncDef));
 
 	// Create the return value property
 	if (bHasReturnValue)
@@ -7786,7 +7793,9 @@ void FHeaderParser::CompileFunctionDeclaration()
 
 	GetCurrentScope()->AddType(TopFunction);
 
-	FFunctionData* StoredFuncData = FFunctionData::Add(FFuncInfo(FuncInfo));
+	TSharedRef<FUnrealFunctionDefinitionInfo> FuncDef = MakeShareable(new FUnrealFunctionDefinitionInfo(SourceFile, InputLine, FString(FuncInfo.Function.Identifier), FFuncInfo(FuncInfo)));
+	FuncDef->SetObject(FuncDef->GetFunctionData().FunctionReference);
+	GTypeDefinitionInfoMap.Add(FuncDef->GetFunctionData().FunctionReference, MoveTemp(FuncDef));
 	if (FuncInfo.FunctionReference->HasAnyFunctionFlags(FUNC_Delegate))
 	{
 		GetCurrentClassData()->MarkContainsDelegate();
@@ -7915,7 +7924,7 @@ void FHeaderParser::CompileFunctionDeclaration()
 	// Make sure any new flags made it to the function
 	//@TODO: UCREMOVAL: Ideally the flags didn't get copied midway thru parsing the function declaration, and we could avoid this
 	TopFunction->FunctionFlags |= FuncInfo.FunctionFlags;
-	StoredFuncData->UpdateFunctionData(FuncInfo);
+	FuncDef->UpdateFunctionData(FuncInfo);
 
 	// Bind the function.
 	TopFunction->Bind();
@@ -8896,7 +8905,7 @@ ECompilationResult::Type FHeaderParser::Parse(
 	catch (TCHAR* ErrorMsg)
 	{
 		Warn->Log(ELogVerbosity::Error, ErrorMsg);
-		Result = FResults::GetResults();
+		Result = ECompilationResult::OtherCompilationError;
 	}
 #endif
 
diff --git a/Engine/Source/Programs/UnrealHeaderTool/Private/HeaderProvider.cpp b/Engine/Source/Programs/UnrealHeaderTool/Private/HeaderProvider.cpp
--- a/Engine/Source/Programs/UnrealHeaderTool/Private/HeaderProvider.cpp
+++ b/Engine/Source/Programs/UnrealHeaderTool/Private/HeaderProvider.cpp
@@ -34,7 +34,7 @@ FUnrealSourceFile* FHeaderProvider::Resolve(const FUnrealSourceFile& ParentSourc
 			FName IdName(*Id, FNAME_Find);
 			if (TSharedRef<FUnrealTypeDefinitionInfo>* Source = GTypeDefinitionInfoMap.FindByName(IdName))
 			{
-				Cache = &(*Source)->GetUnrealSourceFile();
+				Cache = (*Source)->HasSource() ? &(*Source)->GetUnrealSourceFile() : nullptr;
 				// There is an edge case with interfaces.  If you define the UMyInterface and IMyInterface in the same
 				// source file as a class that implements the interface, a HeaderProvider for IMyInterface is added 
 				// at the pre-parse time that later (incorrectly) resolves to UMyInterface.  This results in
@@ -54,7 +54,7 @@ FUnrealSourceFile* FHeaderProvider::Resolve(const FUnrealSourceFile& ParentSourc
 				FName IdName(*GetClassNameWithPrefixRemoved(Id), FNAME_Find);
 				if (TSharedRef<FUnrealTypeDefinitionInfo>* Source = GTypeDefinitionInfoMap.FindByName(IdName))
 				{
-					Cache = &(*Source)->GetUnrealSourceFile();
+					Cache = (*Source)->HasSource() ? &(*Source)->GetUnrealSourceFile() : nullptr;
 				}
 			}
 			if (Cache == nullptr)
@@ -62,7 +62,7 @@ FUnrealSourceFile* FHeaderProvider::Resolve(const FUnrealSourceFile& ParentSourc
 				FName IdName(*Id, FNAME_Find);
 				if (TSharedRef<FUnrealTypeDefinitionInfo>* Source = GTypeDefinitionInfoMap.FindByName(IdName))
 				{
-					Cache = &(*Source)->GetUnrealSourceFile();
+					Cache = (*Source)->HasSource() ? &(*Source)->GetUnrealSourceFile() : nullptr;
 				}
 			}
 			break;
diff --git a/Engine/Source/Programs/UnrealHeaderTool/Private/ParserHelper.cpp b/Engine/Source/Programs/UnrealHeaderTool/Private/ParserHelper.cpp
--- a/Engine/Source/Programs/UnrealHeaderTool/Private/ParserHelper.cpp
+++ b/Engine/Source/Programs/UnrealHeaderTool/Private/ParserHelper.cpp
@@ -5,6 +5,8 @@
 #include "UnrealHeaderTool.h"
 #include "Algo/Find.h"
 #include "Misc/DefaultValueHelper.h"
+#include "UnrealTypeDefinitionInfo.h"
+#include "ClassMaps.h"
 
 /////////////////////////////////////////////////////
 // FClassMetaData
@@ -47,14 +49,14 @@ FTokenData* FClassMetaData::FindTokenData( FProperty* Prop )
 		if ( OuterFunction != NULL )
 		{
 			// function parameter, return, or local property
-			FFunctionData* FuncData = nullptr;
-			if (FFunctionData::TryFindForFunction(OuterFunction, FuncData))
+			if (TSharedRef<FUnrealTypeDefinitionInfo>* TypeDef = GTypeDefinitionInfoMap.Find(OuterFunction))
 			{
-				FPropertyData& FunctionParameters = FuncData->GetParameterData();
+				FUnrealFunctionDefinitionInfo& FuncDef = (*TypeDef)->AsFunctionChecked();
+				FPropertyData& FunctionParameters = FuncDef.GetParameterData();
 				Result = FunctionParameters.Find(Prop);
 				if ( Result == NULL )
 				{
-					Result = FuncData->GetReturnTokenData();
+					Result = FuncDef.GetReturnTokenData();
 				}
 			}
 			else
@@ -83,6 +85,51 @@ FTokenData* FClassMetaData::FindTokenData( FProperty* Prop )
 	return Result;
 }
 
+void FClassMetaData::AddProperty(FToken&& PropertyToken)
+{
+	FProperty* Prop = PropertyToken.TokenProperty;
+	check(Prop);
+
+	UObject* Outer = Prop->GetOwner<UObject>();
+	check(Outer);
+	UStruct* OuterClass = Cast<UStruct>(Outer);
+	if (OuterClass != NULL)
+	{
+		// global property
+		GlobalPropertyData.Set(Prop, MoveTemp(PropertyToken));
+	}
+	else
+	{
+		checkNoEntry();
+		UFunction* OuterFunction = Cast<UFunction>(Outer);
+		if (OuterFunction != NULL)
+		{
+			// function parameter, return, or local property
+			GTypeDefinitionInfoMap.FindChecked(OuterFunction).AsFunctionChecked().AddProperty(MoveTemp(PropertyToken));
+		}
+	}
+
+	// update the optimization flags
+	if (!bContainsDelegates)
+	{
+		if (Prop->IsA(FDelegateProperty::StaticClass()) || Prop->IsA(FMulticastDelegateProperty::StaticClass()))
+		{
+			bContainsDelegates = true;
+		}
+		else
+		{
+			FArrayProperty* ArrayProp = CastField<FArrayProperty>(Prop);
+			if (ArrayProp != NULL)
+			{
+				if (ArrayProp->Inner->IsA(FDelegateProperty::StaticClass()) || ArrayProp->Inner->IsA(FMulticastDelegateProperty::StaticClass()))
+				{
+					bContainsDelegates = true;
+				}
+			}
+		}
+	}
+}
+
 void FClassMetaData::AddInheritanceParent(FString&& InParent, FUnrealSourceFile* UnrealSourceFile)
 {
 	MultipleInheritanceParents.Add(new FMultipleInheritanceBaseClass(MoveTemp(InParent)));
@@ -236,44 +283,6 @@ bool FAdvancedDisplayParameterHandler::CanMarkMore() const
 	return bUseNumber ? (NumberLeaveUnmarked > 0) : (0 != ParametersNames.Num());
 }
 
-TMap<UFunction*, TUniqueObj<FFunctionData> > FFunctionData::FunctionDataMap;
-
-FFunctionData* FFunctionData::FindForFunction(UFunction* Function)
-{
-	TUniqueObj<FFunctionData>* Output = FunctionDataMap.Find(Function);
-
-	check(Output);
-
-	return &(*Output).Get();
-}
-
-FFunctionData* FFunctionData::Add(UFunction* Function)
-{
-	TUniqueObj<FFunctionData>& Output = FunctionDataMap.Add(Function);
-
-	return &Output.Get();
-}
-
-FFunctionData* FFunctionData::Add(FFuncInfo&& FunctionInfo)
-{
-	TUniqueObj<FFunctionData>& Output = FunctionDataMap.Emplace(FunctionInfo.FunctionReference, MoveTemp(FunctionInfo));
-
-	return &Output.Get();
-}
-
-bool FFunctionData::TryFindForFunction(UFunction* Function, FFunctionData*& OutData)
-{
-	TUniqueObj<FFunctionData>* Output = FunctionDataMap.Find(Function);
-
-	if (!Output)
-	{
-		return false;
-	}
-
-	OutData = &(*Output).Get();
-	return true;
-}
-
 FClassMetaData* FCompilerMetadataManager::AddClassData(UStruct* Struct, FUnrealSourceFile* UnrealSourceFile)
 {
 	TUniquePtr<FClassMetaData>* pClassData = Find(Struct);
@@ -293,11 +302,11 @@ FClassMetaData* FCompilerMetadataManager::AddInterfaceClassData(UStruct* Struct,
 	return ClassData;
 }
 
-FTokenData* FPropertyData::Set(FProperty* InKey, FTokenData&& InValue, FUnrealSourceFile* UnrealSourceFile)
+FTokenData* FPropertyData::Set(FProperty* InKey, FTokenData&& InValue)
 {
 	FTokenData* Result = NULL;
 
-	TSharedPtr<FTokenData>* pResult = Super::Find(InKey);
+	TSharedPtr<FTokenData>* pResult = PropertyMap.Find(InKey);
 	if (pResult != NULL)
 	{
 		Result = pResult->Get();
@@ -305,7 +314,7 @@ FTokenData* FPropertyData::Set(FProperty* InKey, FTokenData&& InValue, FUnrealSo
 	}
 	else
 	{
-		pResult = &Super::Emplace(InKey, new FTokenData(MoveTemp(InValue)));
+		pResult = &PropertyMap.Emplace(InKey, new FTokenData(MoveTemp(InValue)));
 		Result = pResult->Get();
 	}
 
diff --git a/Engine/Source/Programs/UnrealHeaderTool/Private/ParserHelper.h b/Engine/Source/Programs/UnrealHeaderTool/Private/ParserHelper.h
--- a/Engine/Source/Programs/UnrealHeaderTool/Private/ParserHelper.h
+++ b/Engine/Source/Programs/UnrealHeaderTool/Private/ParserHelper.h
@@ -9,11 +9,10 @@
 #include "UObject/TextProperty.h"
 #include "UObject/EnumProperty.h"
 #include "UObject/FieldPathProperty.h"
+#include "UObject/Stack.h"
 #include "UnrealHeaderToolGlobals.h"
-#include "ClassMaps.h"
 #include "Templates/UniqueObj.h"
 #include "Templates/UniquePtr.h"
-#include "UnrealTypeDefinitionInfo.h"
 
 class UEnum;
 class UScriptStruct;
@@ -87,6 +86,22 @@ enum class EIntType
 	Unsized // e.g. int, unsigned int
 };
 
+enum class EAllocatorType
+{
+	Default,
+	MemoryImage
+};
+
+/** Types access specifiers. */
+enum EAccessSpecifier
+{
+	ACCESS_NotAnAccessSpecifier = 0,
+	ACCESS_Public,
+	ACCESS_Private,
+	ACCESS_Protected,
+	ACCESS_Num,
+};
+
 #ifndef CASE_TEXT
 #define CASE_TEXT(txt) case txt: return TEXT(#txt)
 #endif
@@ -1183,9 +1198,9 @@ struct FTokenData
  * Class for storing data about a list of properties.  Though FToken contains a reference to its
  * associated FProperty, it's faster lookup to use the FProperty as the key in a TMap.
  */
-class FPropertyData : public TMap< FProperty*, TSharedPtr<FTokenData> >
+class FPropertyData
 {
-	typedef TMap<FProperty*, TSharedPtr<FTokenData> >	Super;
+	TMap<FProperty*, TSharedPtr<FTokenData>> PropertyMap;
 
 public:
 	/**
@@ -1198,7 +1213,7 @@ public:
 	{
 		FTokenData* Result = NULL;
 
-		TSharedPtr<FTokenData>* pResult = Super::Find(Key);
+		TSharedPtr<FTokenData>* pResult = PropertyMap.Find(Key);
 		if ( pResult != NULL )
 		{
 			Result = pResult->Get();
@@ -1209,7 +1224,7 @@ public:
 	{
 		const FTokenData* Result = NULL;
 
-		const TSharedPtr<FTokenData>* pResult = Super::Find(Key);
+		const TSharedPtr<FTokenData>* pResult = PropertyMap.Find(Key);
 		if ( pResult != NULL )
 		{
 			Result = pResult->Get();
@@ -1217,6 +1232,11 @@ public:
 		return Result;
 	}
 
+	void Shrink()
+	{
+		PropertyMap.Shrink();
+	}
+
 	/**
 	 * Sets the value associated with a key.  If the key already exists in the map, uses the same
 	 * value pointer and reinitalized the FTokenData with the input value.
@@ -1226,220 +1246,7 @@ public:
 	 *
 	 * @return	a pointer to token data created associated with the property
 	 */
-	FTokenData* Set(FProperty* InKey, FTokenData&& InValue, FUnrealSourceFile* UnrealSourceFile);
-
-	/**
-	 * (debug) Dumps the values of this FPropertyData to the log file
-	 * 
-	 * @param	Indent	number of spaces to insert at the beginning of each line
-	 */	
-	void Dump( int32 Indent )
-	{
-		for (auto& Kvp : *this)
-		{
-			TSharedPtr<FTokenData>& PointerVal = Kvp.Value;
-			FToken& Token = PointerVal->Token;
-			if ( Token.Type != CPT_None )
-			{
-				UE_LOG(LogCompile, Log, TEXT("%s%s"), FCString::Spc(Indent), *Token.Describe());
-			}
-		}
-	}
-};
-
-/**
- * Class for storing additional data about compiled structs and struct properties
- */
-class FStructData
-{
-public:
-	/** info about the struct itself */
-	FToken			StructData;
-
-private:
-	/** info for the properties contained in this struct */
-	FPropertyData	StructPropertyData;
-
-public:
-	/**
-	 * Adds a new struct property token
-	 * 
-	 * @param	PropertyToken	token that should be added to the list
-	 */
-	void AddStructProperty(FTokenData&& PropertyToken, FUnrealSourceFile* UnrealSourceFile)
-	{
-		check(PropertyToken.Token.TokenProperty);
-		StructPropertyData.Set(PropertyToken.Token.TokenProperty, MoveTemp(PropertyToken), UnrealSourceFile);
-	}
-
-	FPropertyData& GetStructPropertyData()
-	{
-		return StructPropertyData;
-	}
-	const FPropertyData& GetStructPropertyData() const
-	{
-		return StructPropertyData;
-	}
-
-	/**
-	* (debug) Dumps the values of this FStructData to the log file
-	* 
-	* @param	Indent	number of spaces to insert at the beginning of each line
-	*/	
-	void Dump( int32 Indent )
-	{
-		UE_LOG(LogCompile, Log, TEXT("%s%s"), FCString::Spc(Indent), *StructData.Describe());
-
-		UE_LOG(LogCompile, Log, TEXT("%sproperties:"), FCString::Spc(Indent));
-		StructPropertyData.Dump(Indent + 4);
-	}
-
-	/** Constructor */
-	FStructData(FToken&& StructToken) : StructData(MoveTemp(StructToken)) {}
-
-	friend struct FStructDataArchiveProxy;
-};
-
-/**
- * Class for storing additional data about compiled function properties.
- */
-class FFunctionData
-{
-	/** info about the function associated with this FFunctionData */
-	FFuncInfo		FunctionData;
-
-	/** return value for this function */
-	FTokenData		ReturnTypeData;
-
-	/** function parameter data */
-	FPropertyData	ParameterData;
-
-	/**
-	 * Adds a new parameter token
-	 * 
-	 * @param	PropertyToken	token that should be added to the list
-	 */
-	void AddParameter(FToken&& PropertyToken, FUnrealSourceFile* UnrealSourceFile)
-	{
-		check(PropertyToken.TokenProperty);
-		ParameterData.Set(PropertyToken.TokenProperty, MoveTemp(PropertyToken), UnrealSourceFile);
-	}
-
-	/**
-	 * Sets the value of the return token for this function
-	 * 
-	 * @param	PropertyToken	token that should be added
-	 */
-	void SetReturnData(FToken&& PropertyToken )
-	{
-		check(PropertyToken.TokenProperty);
-		ReturnTypeData.Token = MoveTemp(PropertyToken);
-	}
-
-public:
-	/** Constructors */
-	FFunctionData() = default;
-	FFunctionData(FFunctionData&& Other) = default;
-	FFunctionData(FFuncInfo&& inFunctionData )
-		: FunctionData(MoveTemp(inFunctionData))
-	{}
-
-	FFunctionData& operator=(FFunctionData&& Other) = default;
-	
-	/** @name getters */
-	//@{
-	const	FFuncInfo&		GetFunctionData()	const	{	return FunctionData;	}
-	const	FToken&			GetReturnData()		const	{	return ReturnTypeData.Token;	}
-	const	FPropertyData&	GetParameterData()	const	{	return ParameterData;	}
-	FPropertyData&			GetParameterData()			{	return ParameterData;	}
-	FTokenData* GetReturnTokenData() { return &ReturnTypeData; }
-	//@}
-
-	void UpdateFunctionData(FFuncInfo& UpdatedFuncData)
-	{
-		//@TODO: UCREMOVAL: Some more thorough evaluation should be done here
-		FunctionData.FunctionFlags |= UpdatedFuncData.FunctionFlags;
-		FunctionData.FunctionExportFlags |= UpdatedFuncData.FunctionExportFlags;
-	}
-
-	/**
-	 * Adds a new function property to be tracked.  Determines whether the property is a
-	 * function parameter, local property, or return value, and adds it to the appropriate
-	 * list
-	 * 
-	 * @param	PropertyToken	the property to add
-	 */
-	void AddProperty(FToken&& PropertyToken, FUnrealSourceFile* UnrealSourceFile)
-	{
-		const FProperty* Prop = PropertyToken.TokenProperty;
-		check(Prop);
-		check( (Prop->PropertyFlags&CPF_Parm) != 0 );
-
-		if ( (Prop->PropertyFlags&CPF_ReturnParm) != 0 )
-		{
-			SetReturnData(MoveTemp(PropertyToken));
-		}
-		else
-		{
-			AddParameter(MoveTemp(PropertyToken), UnrealSourceFile);
-		}
-	}
-
-	/**
-	 * (debug) Dumps the values of this FFunctionData to the log file
-	 * 
-	 * @param	Indent	number of spaces to insert at the beginning of each line
-	 */	
-	void Dump( int32 Indent )
-	{
-		UE_LOG(LogCompile, Log, TEXT("%sparameters:"), FCString::Spc(Indent));
-		ParameterData.Dump(Indent + 4);
-
-		UE_LOG(LogCompile, Log, TEXT("%sreturn prop:"), FCString::Spc(Indent));
-		if ( ReturnTypeData.Token.Type != CPT_None )
-		{
-			UE_LOG(LogCompile, Log, TEXT("%s%s"), FCString::Spc(Indent + 4), *ReturnTypeData.Token.Describe());
-		}
-	}
-
-	/**
-	 * Sets the specified function export flags
-	 */
-	void SetFunctionExportFlag( uint32 NewFlags )
-	{
-		FunctionData.FunctionExportFlags |= NewFlags;
-	}
-
-	/**
-	 * Clears the specified function export flags
-	 */
-	void ClearFunctionExportFlags( uint32 ClearFlags )
-	{
-		FunctionData.FunctionExportFlags &= ~ClearFlags;
-	}
-
-	/**
-	 * Finds function data for given function object.
-	 */
-	static FFunctionData* FindForFunction(UFunction* Function);
-
-	/**
-	 * Adds function data object for given function object.
-	 */
-	static FFunctionData* Add(UFunction* Function);
-
-	/**
-	 * Adds function data object for given function object.
-	 */
-	static FFunctionData* Add(FFuncInfo&& FunctionInfo);
-
-	/**
-	 * Tries to find function data for given function object.
-	 */
-	static bool TryFindForFunction(UFunction* Function, FFunctionData*& OutData);
-
-private:
-	static TMap<UFunction*, TUniqueObj<FFunctionData> > FunctionDataMap;
+	FTokenData* Set(FProperty* InKey, FTokenData&& InValue);
 };
 
 /**
@@ -1586,50 +1393,7 @@ public:
 	 * 
 	 * @param	PropertyToken	the property to add
 	 */
-	void AddProperty(FToken&& PropertyToken, FUnrealSourceFile* UnrealSourceFile)
-	{
-		FProperty* Prop = PropertyToken.TokenProperty;
-		check(Prop);
-
-		UObject* Outer = Prop->GetOwner<UObject>();
-		check(Outer);
-		UStruct* OuterClass = Cast<UStruct>(Outer);
-		if ( OuterClass != NULL )
-		{
-			// global property
-			GlobalPropertyData.Set(Prop, MoveTemp(PropertyToken), UnrealSourceFile);
-		}
-		else
-		{
-			checkNoEntry();
-			UFunction* OuterFunction = Cast<UFunction>(Outer);
-			if ( OuterFunction != NULL )
-			{
-				// function parameter, return, or local property
-				FFunctionData::FindForFunction(OuterFunction)->AddProperty(MoveTemp(PropertyToken), UnrealSourceFile);
-			}
-		}
-
-		// update the optimization flags
-		if ( !bContainsDelegates )
-		{
-			if( Prop->IsA( FDelegateProperty::StaticClass() ) || Prop->IsA( FMulticastDelegateProperty::StaticClass() ) )
-			{
-				bContainsDelegates = true;
-			}
-			else
-			{
-				FArrayProperty* ArrayProp = CastField<FArrayProperty>(Prop);
-				if ( ArrayProp != NULL )
-				{
-					if( ArrayProp->Inner->IsA( FDelegateProperty::StaticClass() ) || ArrayProp->Inner->IsA( FMulticastDelegateProperty::StaticClass() ) )
-					{
-						bContainsDelegates = true;
-					}
-				}
-			}
-		}
-	}
+	void AddProperty(FToken&& PropertyToken);
 
 	/**
 	 * Adds new editor-only metadata (key/value pairs) to the class or struct that
@@ -1682,18 +1446,6 @@ public:
 		}
 	}
 
-	/**
-	 * Finds the metadata for the function specified
-	 * 
-	 * @param	Func	the function to search for
-	 *
-	 * @return	pointer to the metadata for the function specified, or NULL
-	 *			if the function doesn't exist in the list (for example, if it
-	 *			is declared in a package that is already compiled and has had its
-	 *			source stripped)
-	 */
-	FFunctionData* FindFunctionData( UFunction* Func );
-
 	/**
 	 * Finds the metadata for the property specified
 	 * 
@@ -1706,13 +1458,6 @@ public:
 	 */
 	FTokenData* FindTokenData( FProperty* Prop );
 
-	/**
-	 * (debug) Dumps the values of this FFunctionData to the log file
-	 * 
-	 * @param	Indent	number of spaces to insert at the beginning of each line
-	 */	
-	void Dump( int32 Indent );
-
 	/**
 	 * Add a string to the list of inheritance parents for this class.
 	 *
diff --git a/Engine/Source/Programs/UnrealHeaderTool/Private/Scope.cpp b/Engine/Source/Programs/UnrealHeaderTool/Private/Scope.cpp
--- a/Engine/Source/Programs/UnrealHeaderTool/Private/Scope.cpp
+++ b/Engine/Source/Programs/UnrealHeaderTool/Private/Scope.cpp
@@ -4,6 +4,8 @@
 #include "UnrealHeaderTool.h"
 #include "UObject/ErrorException.h"
 #include "ParserHelper.h"
+#include "UnrealTypeDefinitionInfo.h"
+#include "ClassMaps.h"
 
 extern FCompilerMetadataManager GScriptHelper;
 
diff --git a/Engine/Source/Programs/UnrealHeaderTool/Private/UnrealSourceFile.cpp b/Engine/Source/Programs/UnrealHeaderTool/Private/UnrealSourceFile.cpp
--- a/Engine/Source/Programs/UnrealHeaderTool/Private/UnrealSourceFile.cpp
+++ b/Engine/Source/Programs/UnrealHeaderTool/Private/UnrealSourceFile.cpp
@@ -7,6 +7,7 @@
 #include "Algo/Find.h"
 #include "StringUtils.h"
 #include "Exceptions.h"
+#include "ClassMaps.h"
 
 void FUnrealSourceFile::AddDefinedClass(TSharedRef<FUnrealTypeDefinitionInfo> ClassDecl)
 {
diff --git a/Engine/Source/Programs/UnrealHeaderTool/Private/UnrealTypeDefinitionInfo.cpp b/Engine/Source/Programs/UnrealHeaderTool/Private/UnrealTypeDefinitionInfo.cpp
--- a/Engine/Source/Programs/UnrealHeaderTool/Private/UnrealTypeDefinitionInfo.cpp
+++ b/Engine/Source/Programs/UnrealHeaderTool/Private/UnrealTypeDefinitionInfo.cpp
@@ -42,6 +42,11 @@ FUnrealScriptStructDefinitionInfo* FUnrealTypeDefinitionInfo::AsScriptStruct()
 	return nullptr;
 }
 
+FUnrealFunctionDefinitionInfo* FUnrealTypeDefinitionInfo::AsFunction()
+{
+	return nullptr;
+}
+
 FUnrealClassDefinitionInfo* FUnrealTypeDefinitionInfo::AsClass()
 {
 	return nullptr;
@@ -111,7 +116,14 @@ uint32 FUnrealScriptStructDefinitionInfo::GetHash(bool bIncludeNoExport) const
 
 TSharedRef<FScope> FUnrealStructDefinitionInfo::GetScope()
 {
-	return StructScope.ToSharedRef();
+	if (StructScope.IsValid())
+	{
+		return StructScope.ToSharedRef();
+	}
+	else
+	{
+		return FUnrealFieldDefinitionInfo::GetScope();
+	}
 }
 
 void FUnrealStructDefinitionInfo::SetObject(UObject* InObject)
@@ -120,7 +132,12 @@ void FUnrealStructDefinitionInfo::SetObject(UObject* InObject)
 
 	FUnrealFieldDefinitionInfo::SetObject(InObject);
 
-	StructScope = MakeShared<FStructScope>(static_cast<UStruct*>(InObject), &GetUnrealSourceFile().GetScope().Get());
+
+	// Don't create a scope for things without a source.  Those are builtin types
+	if (HasSource())
+	{
+		StructScope = MakeShared<FStructScope>(static_cast<UStruct*>(InObject), &GetUnrealSourceFile().GetScope().Get());
+	}
 }
 
 uint32 FUnrealClassDefinitionInfo::GetHash(bool bIncludeNoExport) const
diff --git a/Engine/Source/Programs/UnrealHeaderTool/Private/UnrealTypeDefinitionInfo.h b/Engine/Source/Programs/UnrealHeaderTool/Private/UnrealTypeDefinitionInfo.h
--- a/Engine/Source/Programs/UnrealHeaderTool/Private/UnrealTypeDefinitionInfo.h
+++ b/Engine/Source/Programs/UnrealHeaderTool/Private/UnrealTypeDefinitionInfo.h
@@ -7,6 +7,8 @@
 #include "ClassDeclarationMetaData.h"
 #include <atomic>
 
+#include "ParserHelper.h" // TEMPORARY
+
 // Forward declarations.
 class FClass;
 class FScope;
@@ -31,6 +33,7 @@ class FUnrealTypeDefinitionInfo;							// Base for all types, provides virtual m
 			class FUnrealStructDefinitionInfo;				// Represents UStruct
 				class FUnrealScriptStructDefinitionInfo;	// Represents UScriptStruct
 				class FUnrealClassDefinitionInfo;			// Represents UClass
+				class FUnrealFunctionDefinitionInfo;		// Represents UFunction
 
 enum class EUnderlyingEnumType
 {
@@ -45,12 +48,6 @@ enum class EUnderlyingEnumType
 	int64
 };
 
-enum class EAllocatorType
-{
-	Default,
-	MemoryImage
-};
-
 enum class ESerializerArchiveType
 {
 	None = 0,
@@ -187,6 +184,23 @@ public:
 		return *Info;
 	}
 
+	/**
+	 * If this is a function, return the function version of the object
+	 * @return Pointer to the definition info or nullptr if not of that type.
+	 */
+	virtual FUnrealFunctionDefinitionInfo* AsFunction();
+
+	/**
+	 * If this is a function, return the function version of the object
+	 * @return Reference to the definition.  Will assert if not of that type.
+	 */
+	FUnrealFunctionDefinitionInfo& AsFunctionChecked()
+	{
+		FUnrealFunctionDefinitionInfo* Info = AsFunction();
+		check(Info);
+		return *Info;
+	}
+
 	/**
 	 * If this is a class, return the class version of the object
 	 * @return Pointer to the definition info or nullptr if not of that type.
@@ -652,6 +666,118 @@ private:
 	FString ParentNameCPP;
 };
 
+/**
+ * Class that stores information about type definitions derrived from UFunction
+ */
+class FUnrealFunctionDefinitionInfo : public FUnrealStructDefinitionInfo
+{
+public:
+	FUnrealFunctionDefinitionInfo(FUnrealSourceFile& InSourceFile, int32 InLineNumber, FString&& InNameCPP, FFuncInfo&& InFuncInfo)
+		: FUnrealStructDefinitionInfo(InSourceFile, InLineNumber, MoveTemp(InNameCPP))
+		, FunctionData(MoveTemp(InFuncInfo))
+	{ }
+
+	virtual FUnrealFunctionDefinitionInfo* AsFunction() override
+	{
+		return this;
+	}
+
+	/**
+	 * Return the Engine instance associated with the compiler instance
+	 */
+	UFunction* GetFunction() const
+	{
+		return static_cast<UFunction*>(GetField());
+	}
+
+	/** @name getters */
+	//@{
+	const	FFuncInfo& GetFunctionData()	const { return FunctionData; }
+	const	FToken& GetReturnData()		const { return ReturnTypeData.Token; }
+	const	FPropertyData& GetParameterData()	const { return ParameterData; }
+	FPropertyData& GetParameterData() { return ParameterData; }
+	FTokenData* GetReturnTokenData() { return &ReturnTypeData; }
+	//@}
+
+	void UpdateFunctionData(FFuncInfo& UpdatedFuncData)
+	{
+		//@TODO: UCREMOVAL: Some more thorough evaluation should be done here
+		FunctionData.FunctionFlags |= UpdatedFuncData.FunctionFlags;
+		FunctionData.FunctionExportFlags |= UpdatedFuncData.FunctionExportFlags;
+	}
+
+	/**
+	 * Adds a new function property to be tracked.  Determines whether the property is a
+	 * function parameter, local property, or return value, and adds it to the appropriate
+	 * list
+	 *
+	 * @param	PropertyToken	the property to add
+	 */
+	void AddProperty(FToken&& PropertyToken)
+	{
+		const FProperty* Prop = PropertyToken.TokenProperty;
+		check(Prop);
+		check((Prop->PropertyFlags & CPF_Parm) != 0);
+
+		if ((Prop->PropertyFlags & CPF_ReturnParm) != 0)
+		{
+			SetReturnData(MoveTemp(PropertyToken));
+		}
+		else
+		{
+			AddParameter(MoveTemp(PropertyToken));
+		}
+	}
+
+	/**
+	 * Sets the specified function export flags
+	 */
+	void SetFunctionExportFlag(uint32 NewFlags)
+	{
+		FunctionData.FunctionExportFlags |= NewFlags;
+	}
+
+	/**
+	 * Clears the specified function export flags
+	 */
+	void ClearFunctionExportFlags(uint32 ClearFlags)
+	{
+		FunctionData.FunctionExportFlags &= ~ClearFlags;
+	}
+
+private:
+	/**
+	 * Adds a new parameter token
+	 *
+	 * @param	PropertyToken	token that should be added to the list
+	 */
+	void AddParameter(FToken&& PropertyToken)
+	{
+		check(PropertyToken.TokenProperty);
+		ParameterData.Set(PropertyToken.TokenProperty, MoveTemp(PropertyToken));
+	}
+
+	/**
+	 * Sets the value of the return token for this function
+	 *
+	 * @param	PropertyToken	token that should be added
+	 */
+	void SetReturnData(FToken&& PropertyToken)
+	{
+		check(PropertyToken.TokenProperty);
+		ReturnTypeData.Token = MoveTemp(PropertyToken);
+	}
+
+	/** info about the function associated with this FFunctionData */
+	FFuncInfo		FunctionData;
+
+	/** return value for this function */
+	FTokenData		ReturnTypeData;
+
+	/** function parameter data */
+	FPropertyData	ParameterData;
+};
+
 /**
  * Class that stores information about type definitions derived from UClass
  */
@@ -666,6 +792,10 @@ public:
 		, bIsInterface(bInIsInterface)
 	{ }
 
+	FUnrealClassDefinitionInfo(FString&& InNameCPP)
+		: FUnrealStructDefinitionInfo(MoveTemp(InNameCPP))
+	{ }
+
 	virtual FUnrealClassDefinitionInfo* AsClass() override
 	{
 		return this;
