Remove the singleton type cache and the package name cache.

#rnx

diff --git a/Engine/Source/Programs/UnrealHeaderTool/Private/ClassMaps.h b/Engine/Source/Programs/UnrealHeaderTool/Private/ClassMaps.h
--- a/Engine/Source/Programs/UnrealHeaderTool/Private/ClassMaps.h
+++ b/Engine/Source/Programs/UnrealHeaderTool/Private/ClassMaps.h
@@ -71,6 +71,15 @@ struct FTypeDefinitionInfoMap : public FFreezableContainer
 		return **TypeDef;
 	}
 
+	template <typename Lambda>
+	void ForAllTypes(Lambda&& InLambda)
+	{
+		for (const TPair<UObject*, TSharedRef<FUnrealTypeDefinitionInfo>>& KVP : DefinitionsByUObject)
+		{
+			InLambda(*KVP.Value);
+		}
+	}
+
 	//NOTE: Currently UFunctions are created during the parsing phase and can not be frozen
 	void Add(UFunction* Object, TSharedRef<FUnrealTypeDefinitionInfo>&& Definition)
 	{
diff --git a/Engine/Source/Programs/UnrealHeaderTool/Private/CodeGenerator.cpp b/Engine/Source/Programs/UnrealHeaderTool/Private/CodeGenerator.cpp
--- a/Engine/Source/Programs/UnrealHeaderTool/Private/CodeGenerator.cpp
+++ b/Engine/Source/Programs/UnrealHeaderTool/Private/CodeGenerator.cpp
@@ -811,254 +811,36 @@ void FNativeClassHeaderGenerator::ExportProperties(FOutputDevice& Out, UStruct*
 	}
 }
 
-/**
- * Class that is representing a type singleton.
- */
-struct FTypeSingleton
-{
-public:
-	/** Constructor */
-	FTypeSingleton(FString InName, UField* InType)
-		: Name(MoveTemp(InName)), Type(InType) {}
-
-	/**
-	 * Gets this singleton's name.
-	 */
-	const FString& GetName() const
-	{
-		return Name;
-	}
-
-	/**
-	 * Gets this singleton's extern declaration.
-	 */
-	const FString& GetExternDecl() const
-	{
-		FRWScopeLock Lock(ExternDeclLock, SLT_ReadOnly);
-		if (ExternDecl.IsEmpty())
-		{
-			Lock.ReleaseReadOnlyLockAndAcquireWriteLock_USE_WITH_CAUTION();
-
-			// Verify the decl is still empty in case another thread had also been waiting on writing this data and got the write lock first
-			if (ExternDecl.IsEmpty())
-			{
-				ExternDecl = GenerateExternDecl(Type, GetName());
-			}
-		}
-
-		return ExternDecl;
-	}
-
-private:
-	/**
-	 * Extern declaration generator.
-	 */
-	static FString GenerateExternDecl(UField* InType, const FString& InName)
-	{
-		const TCHAR* TypeStr = nullptr;
-
-		if (InType->GetClass() == UClass::StaticClass())
-		{
-			TypeStr = TEXT("UClass");
-		}
-		else if (InType->GetClass() == UFunction::StaticClass() || InType->GetClass() == UDelegateFunction::StaticClass() || InType->GetClass() == USparseDelegateFunction::StaticClass())
-		{
-			TypeStr = TEXT("UFunction");
-		}
-		else if (InType->GetClass() == UScriptStruct::StaticClass())
-		{
-			TypeStr = TEXT("UScriptStruct");
-		}
-		else if (InType->GetClass() == UEnum::StaticClass())
-		{
-			TypeStr = TEXT("UEnum");
-		}
-		else
-		{
-			FError::Throwf(TEXT("Unsupported item type to get extern for."));
-		}
-
-		return FString::Printf(
-			TEXT("\t%s_API %s* %s;\r\n"),
-			*FPackageName::GetShortName(InType->GetOutermost()).ToUpper(),
-			TypeStr,
-			*InName
-		);
-	}
-
-	/** Field that stores this singleton name. */
-	FString Name;
+static FString GNullPtr(TEXT("nullptr"));
 
-	/** Cached field that stores this singleton extern declaration. */
-	mutable FString ExternDecl;
-
-	/** Mutex to ensure 2 threads don't try to generate the extern decl at the same time. */
-	mutable FRWLock ExternDeclLock;
-
-	/** Type of the singleton */
-	UField* Type;
-};
-
-/**
- * Class that represents type singleton cache.
- */
-class FTypeSingletonCache
+const FString& FNativeClassHeaderGenerator::GetPackageSingletonName(UPackage* Item, TSet<FString>* UniqueCrossModuleReferences)
 {
-public:
-	/**
-	 * Gets type singleton from cache.
-	 *
-	 * @param Type Singleton type.
-	 * @param bRequiresValidObject Does it require a valid object?
-	 */
-	static const FTypeSingleton& Get(UField* Type, bool bRequiresValidObject = true)
-	{
-		FTypeSingletonCacheKey Key(Type, bRequiresValidObject);
-
-		FRWScopeLock Lock(Mutex, SLT_ReadOnly);
-		TUniquePtr<FTypeSingleton>* SingletonPtr = CacheData.Find(Key);
-		if (SingletonPtr == nullptr)
-		{
-			TUniquePtr<FTypeSingleton> TypeSingleton(MakeUnique<FTypeSingleton>(GenerateSingletonName(Type, bRequiresValidObject), Type));
-
-			Lock.ReleaseReadOnlyLockAndAcquireWriteLock_USE_WITH_CAUTION();
-
-			// Check the map again in case another thread had also been waiting on writing this data and got the write lock first
-			SingletonPtr = CacheData.Find(Key);
-			if (SingletonPtr == nullptr)
-			{
-				SingletonPtr = &CacheData.Add(Key, MoveTemp(TypeSingleton));
-			}
-		}
-
-		return **SingletonPtr;
-	}
-
-private:
-	/**
-	 * Private type that represents cache map key.
-	 */
-	struct FTypeSingletonCacheKey
-	{
-		/** FTypeSingleton type */
-		UField* Type;
-
-		/** If this type singleton requires valid object. */
-		bool bRequiresValidObject;
-
-		/* Constructor */
-		FTypeSingletonCacheKey(UField* InType, bool bInRequiresValidObject)
-			: Type(InType), bRequiresValidObject(bInRequiresValidObject)
-		{}
-
-		/**
-		 * Equality operator.
-		 *
-		 * @param Other Other key.
-		 *
-		 * @returns True if this is equal to Other. False otherwise.
-		 */
-		bool operator==(const FTypeSingletonCacheKey& Other) const
-		{
-			return Type == Other.Type && bRequiresValidObject == Other.bRequiresValidObject;
-		}
-
-		/**
-		 * Gets hash value for this object.
-		 */
-		friend uint32 GetTypeHash(const FTypeSingletonCacheKey& Object)
-		{
-			return HashCombine(
-				GetTypeHash(Object.Type),
-				GetTypeHash(Object.bRequiresValidObject)
-			);
-		}
-	};
-
-	/**
-	 * Generates singleton name.
-	 */
-	static FString GenerateSingletonName(UField* Item, bool bRequiresValidObject)
-	{
-		check(Item);
-
-		bool bNoRegister = false;
-		if (UClass* ItemClass = Cast<UClass>(Item))
-		{
-			if (!bRequiresValidObject && !ItemClass->HasAllClassFlags(CLASS_Intrinsic))
-			{
-				bNoRegister = true;
-			}
-		}
-
-		const TCHAR* Suffix = (bNoRegister ? TEXT("_NoRegister") : TEXT(""));
-
-		FString Result;
-		for (UObject* Outer = Item; Outer; Outer = Outer->GetOuter())
-		{
-			if (!Result.IsEmpty())
-			{
-				Result = TEXT("_") + Result;
-			}
-
-			if (Cast<UClass>(Outer) || Cast<UScriptStruct>(Outer))
-			{
-				Result = FNameLookupCPP::GetNameCPP(Cast<UStruct>(Outer)) + Result;
-
-				// Structs can also have UPackage outer.
-				if (Cast<UClass>(Outer) || Cast<UPackage>(Outer->GetOuter()))
-				{
-					break;
-				}
-			}
-			else
-			{
-				Result = Outer->GetName() + Result;
-			}
-		}
-
-		// Can't use long package names in function names.
-		if (Result.StartsWith(TEXT("/Script/"), ESearchCase::CaseSensitive))
-		{
-			Result = FPackageName::GetShortName(Result);
-		}
-
-		const FString ClassString = FNameLookupCPP::GetNameCPP(Item->GetClass());
-		return FString::Printf(TEXT("Z_Construct_%s_%s%s()"), *ClassString, *Result, Suffix);
-	}
-
-	static TMap<FTypeSingletonCacheKey, TUniquePtr<FTypeSingleton>> CacheData;
-	static FRWLock Mutex;
-};
-
-TMap<FTypeSingletonCache::FTypeSingletonCacheKey, TUniquePtr<FTypeSingleton>> FTypeSingletonCache::CacheData;
-FRWLock FTypeSingletonCache::Mutex;
+	check(Item);
+	FUnrealPackageDefinitionInfo& Package = GTypeDefinitionInfoMap.FindChecked(Item).AsPackageChecked();
+	Package.AddCrossModuleReference(UniqueCrossModuleReferences);
+	return Package.GetSingletonName();
+}
 
-const FString& FNativeClassHeaderGenerator::GetSingletonName(UField* Item, TSet<FString>* UniqueCrossModuleReferences, bool bRequiresValidObject)
+const FString& FNativeClassHeaderGenerator::GetPackageSingletonNameFuncAddr(UPackage* Item, TSet<FString>* UniqueCrossModuleReferences)
 {
-	const FTypeSingleton& Cache = FTypeSingletonCache::Get(Item, bRequiresValidObject);
-
-	// We don't need to export UFunction externs, though we may need the externs for UDelegateFunctions
-	if (UniqueCrossModuleReferences && (!Item->IsA<UFunction>() || Item->IsA<UDelegateFunction>()))
-	{
-		UniqueCrossModuleReferences->Add(Cache.GetExternDecl());
-	}
-
-	return Cache.GetName();
+	check(Item);
+	FUnrealPackageDefinitionInfo& Package = GTypeDefinitionInfoMap.FindChecked(Item).AsPackageChecked();
+	Package.AddCrossModuleReference(UniqueCrossModuleReferences);
+	return Package.GetSingletonNameChopped();
 }
 
-FString FNativeClassHeaderGenerator::GetSingletonNameFuncAddr(UField* Item, TSet<FString>* UniqueCrossModuleReferences, bool bRequiresValidObject)
+const FString& FNativeClassHeaderGenerator::GetSingletonNameFuncAddr(UField* Item, TSet<FString>* UniqueCrossModuleReferences, bool bRequiresValidObject)
 {
-	FString Result;
 	if (!Item)
 	{
-		Result = TEXT("nullptr");
+		return GNullPtr;
 	}
 	else
 	{
-		Result = GetSingletonName(Item, UniqueCrossModuleReferences, bRequiresValidObject).LeftChop(2);
+		FUnrealFieldDefinitionInfo& Field = GTypeDefinitionInfoMap.FindChecked(Item).AsFieldChecked();
+		Field.AddCrossModuleReference(UniqueCrossModuleReferences, bRequiresValidObject);
+		return Field.GetSingletonNameChopped(bRequiresValidObject);
 	}
-	return Result;
 }
 
 void FNativeClassHeaderGenerator::GetPropertyTag(FUHTStringBuilder& Out, UObject* Object)
@@ -2197,52 +1979,6 @@ static TArray<UScriptStruct*> FindNoExportStructs(UStruct* Start)
 	return Result;
 }
 
-struct FPackageSingletonStrings
-{
-	FPackageSingletonStrings(FString&& InPackageSingletonName)
-		: PackageSingletonName(MoveTemp(InPackageSingletonName))
-		, PackageUniqueCrossModuleReference(FString::Printf(TEXT("\tUPackage* %s;\r\n"), *PackageSingletonName))
-	{
-	}
-
-	FString PackageSingletonName;
-	FString PackageUniqueCrossModuleReference;
-};
-
-static TMap<const UPackage*, TUniquePtr<FPackageSingletonStrings>> PackageSingletonNames;
-static FRWLock PackageSingletonNamesLock;
-
-const FString& FNativeClassHeaderGenerator::GetPackageSingletonName(const UPackage* InPackage, TSet<FString>* UniqueCrossModuleReferences)
-{
-	FRWScopeLock Lock(PackageSingletonNamesLock, SLT_ReadOnly);
-
-	TUniquePtr<FPackageSingletonStrings>* PackageSingletonStrings = PackageSingletonNames.Find(InPackage);
-	if (PackageSingletonStrings == nullptr)
-	{
-		FString PackageName = InPackage->GetName();
-		PackageName.ReplaceInline(TEXT("/"), TEXT("_"), ESearchCase::CaseSensitive);
-
-		TUniquePtr<FPackageSingletonStrings> NewPackageSingletonStrings(MakeUnique<FPackageSingletonStrings>(FString::Printf(TEXT("Z_Construct_UPackage_%s()"), *PackageName)));
-
-		Lock.ReleaseReadOnlyLockAndAcquireWriteLock_USE_WITH_CAUTION();
-		
-		// Check the map again in case another thread had also been waiting on writing this data and got the write lock first
-		PackageSingletonStrings = PackageSingletonNames.Find(InPackage);
-
-		if (PackageSingletonStrings == nullptr)
-		{
-			PackageSingletonStrings = &PackageSingletonNames.Add(InPackage, MoveTemp(NewPackageSingletonStrings));
-		}
-	}
-
-	if (UniqueCrossModuleReferences)
-	{
-		UniqueCrossModuleReferences->Add((*PackageSingletonStrings)->PackageUniqueCrossModuleReference);
-	}
-
-	return (*PackageSingletonStrings)->PackageSingletonName;
-}
-
 void FNativeClassHeaderGenerator::ExportGeneratedPackageInitCode(FOutputDevice& Out, const TCHAR* InDeclarations, uint32 Hash)
 {
 	UPackage* Package = PackageDef.GetPackage();
@@ -2262,14 +1998,14 @@ void FNativeClassHeaderGenerator::ExportGeneratedPackageInitCode(FOutputDevice&
 		{
 			return !bADel;
 		}
-		const FString& AName = FTypeSingletonCache::Get(A, true).GetName();
-		const FString& BName = FTypeSingletonCache::Get(B, true).GetName();
+		const FString& AName = GTypeDefinitionInfoMap.FindChecked(A).AsFieldChecked().GetSingletonName(true);
+		const FString& BName = GTypeDefinitionInfoMap.FindChecked(B).AsFieldChecked().GetSingletonName(true);
 		return AName < BName;
 	});
 
 	for (UField* ScriptType : Singletons)
 	{
-		Out.Log(FTypeSingletonCache::Get(ScriptType, true).GetExternDecl());
+		Out.Log(GTypeDefinitionInfoMap.FindChecked(ScriptType).AsFieldChecked().GetExternDecl(true));
 	}
 
 	FOutputDeviceNull OutputDeviceNull;
@@ -2288,7 +2024,7 @@ void FNativeClassHeaderGenerator::ExportGeneratedPackageInitCode(FOutputDevice&
 		Out.Logf(TEXT("\t\t\tstatic UObject* (*const SingletonFuncArray[])() = {\r\n"));
 		for (UField* ScriptType : Singletons)
 		{
-			const FString Name = FTypeSingletonCache::Get(ScriptType, true).GetName().LeftChop(2);
+			const FString& Name = GTypeDefinitionInfoMap.FindChecked(ScriptType).AsFieldChecked().GetSingletonNameChopped(true);
 
 			Out.Logf(TEXT("\t\t\t\t(UObject* (*)())%s,\r\n"), *Name);
 		}
@@ -2323,6 +2059,7 @@ void FNativeClassHeaderGenerator::ExportNativeGeneratedInitCode(FOutputDevice& O
 	check(!OutFriendText.Len());
 
 	UE_CLOG(Class->ClassGeneratedBy, LogCompile, Fatal, TEXT("For intrinsic and compiled-in classes, ClassGeneratedBy should always be null"));
+
 	FUnrealClassDefinitionInfo& ClassDef = GTypeDefinitionInfoMap.FindChecked(Class).AsClassChecked();
 
 	const bool   bIsNoExport  = Class->HasAnyClassFlags(CLASS_NoExport);
@@ -2364,21 +2101,25 @@ void FNativeClassHeaderGenerator::ExportNativeGeneratedInitCode(FOutputDevice& O
 	{
 		// simple ::StaticClass wrapper to avoid header, link and DLL hell
 		{
-			const FString& SingletonNameNoRegister = GetSingletonName(Class, OutReferenceGatherers.UniqueCrossModuleReferences, false);
+			ClassDef.AddCrossModuleReference(OutReferenceGatherers.UniqueCrossModuleReferences, false);
+			const FString& SingletonNameNoRegister = ClassDef.GetSingletonName(false);
 
-			OutDeclarations.Log(FTypeSingletonCache::Get(Class, false).GetExternDecl());
+			OutDeclarations.Log(ClassDef.GetExternDecl(false));
 
 			GeneratedClassRegisterFunctionText.Logf(TEXT("\tUClass* %s\r\n"), *SingletonNameNoRegister);
 			GeneratedClassRegisterFunctionText.Logf(TEXT("\t{\r\n"));
 			GeneratedClassRegisterFunctionText.Logf(TEXT("\t\treturn %s::StaticClass();\r\n"), *ClassNameCPP);
 			GeneratedClassRegisterFunctionText.Logf(TEXT("\t}\r\n"));
 		}
-		const FString& SingletonName = GetSingletonName(Class, OutReferenceGatherers.UniqueCrossModuleReferences);
 
-		FString StaticsStructName = SingletonName.LeftChop(2) + TEXT("_Statics");
+		// NOTE: We are adding the cross module reference here to avoid output differences.
+		ClassDef.AddCrossModuleReference(OutReferenceGatherers.UniqueCrossModuleReferences, true);
+		const FString& SingletonName = ClassDef.GetSingletonName(true);
+
+		FString StaticsStructName = ClassDef.GetSingletonNameChopped(true) + TEXT("_Statics");
 
 		OutFriendText.Logf(TEXT("\tfriend struct %s;\r\n"), *StaticsStructName);
-		OutDeclarations.Log(FTypeSingletonCache::Get(Class).GetExternDecl());
+		OutDeclarations.Log(GTypeDefinitionInfoMap.FindChecked(Class).AsFieldChecked().GetExternDecl(true));
 
 		GeneratedClassRegisterFunctionText.Logf(TEXT("\tstruct %s\r\n"), *StaticsStructName);
 		GeneratedClassRegisterFunctionText.Logf(TEXT("\t{\r\n"));
@@ -2389,15 +2130,17 @@ void FNativeClassHeaderGenerator::ExportNativeGeneratedInitCode(FOutputDevice& O
 		UClass* SuperClass = Class->GetSuperClass();
 		if (SuperClass && SuperClass != Class)
 		{
-			OutDeclarations.Log(FTypeSingletonCache::Get(SuperClass).GetExternDecl());
-			Singletons.Logf(TEXT("\t\t(UObject* (*)())%s,\r\n"), *GetSingletonName(SuperClass, OutReferenceGatherers.UniqueCrossModuleReferences).LeftChop(2));
+			FUnrealClassDefinitionInfo& SuperClassDef = GTypeDefinitionInfoMap.FindChecked(SuperClass).AsClassChecked();
+			SuperClassDef.AddCrossModuleReference(OutReferenceGatherers.UniqueCrossModuleReferences, true);
+			OutDeclarations.Log(SuperClassDef.GetExternDecl(true));
+			Singletons.Logf(TEXT("\t\t(UObject* (*)())%s,\r\n"), *SuperClassDef.GetSingletonNameChopped(true));
 		}
 		if (!bIsDynamic)
 		{
-			const FString& PackageSingletonName = GetPackageSingletonName(Class->GetOutermost(), OutReferenceGatherers.UniqueCrossModuleReferences);
-
-			OutDeclarations.Logf(TEXT("\t%s_API UPackage* %s;\r\n"), *ApiString, *PackageSingletonName);
-			Singletons.Logf(TEXT("\t\t(UObject* (*)())%s,\r\n"), *PackageSingletonName.LeftChop(2));
+			FUnrealPackageDefinitionInfo& ClassPackageDef = GTypeDefinitionInfoMap.FindChecked(Class->GetOutermost()).AsPackageChecked();
+			PackageDef.AddCrossModuleReference(OutReferenceGatherers.UniqueCrossModuleReferences);
+			OutDeclarations.Logf(TEXT("\t%s_API UPackage* %s;\r\n"), *ApiString, *ClassPackageDef.GetSingletonName());
+			Singletons.Logf(TEXT("\t\t(UObject* (*)())%s,\r\n"), *ClassPackageDef.GetSingletonNameChopped());
 		}
 
 		const TCHAR* SingletonsArray;
@@ -2504,6 +2247,8 @@ void FNativeClassHeaderGenerator::ExportNativeGeneratedInitCode(FOutputDevice& O
 			for (const FImplementedInterface& Inter : Class->Interfaces)
 			{
 				check(Inter.Class);
+				FUnrealFieldDefinitionInfo& InterClassDef = GTypeDefinitionInfoMap.FindChecked(Inter.Class).AsFieldChecked();
+				InterClassDef.AddCrossModuleReference(OutReferenceGatherers.UniqueCrossModuleReferences, false);
 				FString OffsetString;
 				if (Inter.PointerOffset)
 				{
@@ -2515,7 +2260,7 @@ void FNativeClassHeaderGenerator::ExportNativeGeneratedInitCode(FOutputDevice& O
 				}
 				StaticDefinitions.Logf(
 					TEXT("\t\t\t{ %s, %s, %s },\r\n"),
-					*GetSingletonName(Inter.Class, OutReferenceGatherers.UniqueCrossModuleReferences, false).LeftChop(2),
+					*InterClassDef.GetSingletonNameChopped(false),
 					*OffsetString,
 					Inter.bImplementedByK2 ? TEXT("true") : TEXT("false")
 				);
@@ -2615,7 +2360,7 @@ void FNativeClassHeaderGenerator::ExportNativeGeneratedInitCode(FOutputDevice& O
 		OutFriendText.Reset();
 	}
 
-	FString SingletonName = GetSingletonName(Class, OutReferenceGatherers.UniqueCrossModuleReferences);
+	FString SingletonName = ClassDef.GetSingletonName(true);
 	SingletonName.ReplaceInline(TEXT("()"), TEXT(""), ESearchCase::CaseSensitive); // function address
 
 	FString OverriddenClassName = *FNativeClassHeaderGenerator::GetOverriddenName(Class);
@@ -2729,14 +2474,17 @@ void FNativeClassHeaderGenerator::ExportNativeGeneratedInitCode(FOutputDevice& O
 
 void FNativeClassHeaderGenerator::ExportFunction(FOutputDevice& Out, FReferenceGatherers& OutReferenceGatherers, const FUnrealSourceFile& SourceFile, UFunction* Function, bool bIsNoExport) const
 {
+	FUnrealFunctionDefinitionInfo& FuncDef = GTypeDefinitionInfoMap.FindChecked(Function).AsFunctionChecked();
+	FuncDef.AddCrossModuleReference(OutReferenceGatherers.UniqueCrossModuleReferences, true);
+
 	UFunction* SuperFunction = Function->GetSuperFunction();
 
 	const bool bIsEditorOnlyFunction = Function->HasAnyFunctionFlags(FUNC_EditorOnly);
 
 	bool bIsDelegate = Function->HasAnyFunctionFlags(FUNC_Delegate);
 
-	const FString& SingletonName = GetSingletonName(Function, OutReferenceGatherers.UniqueCrossModuleReferences);
-	FString StaticsStructName = SingletonName.LeftChop(2) + TEXT("_Statics");
+	const FString& SingletonName = FuncDef.GetSingletonName(true);
+	FString StaticsStructName = FuncDef.GetSingletonNameChopped(true) + TEXT("_Statics");
 
 	FUHTStringBuilder CurrentFunctionText;
 	FUHTStringBuilder StaticDefinitions;
@@ -2768,7 +2516,7 @@ void FNativeClassHeaderGenerator::ExportFunction(FOutputDevice& Out, FReferenceG
 	FString OuterFunc;
 	if (UObject* Outer = Function->GetOuter())
 	{
-		OuterFunc = Outer->IsA<UPackage>() ? GetPackageSingletonName((UPackage*)Outer, OutReferenceGatherers.UniqueCrossModuleReferences).LeftChop(2) : GetSingletonNameFuncAddr(Function->GetOwnerClass(), OutReferenceGatherers.UniqueCrossModuleReferences);
+		OuterFunc = Outer->IsA<UPackage>() ? GetPackageSingletonNameFuncAddr((UPackage*)Outer, OutReferenceGatherers.UniqueCrossModuleReferences) : GetSingletonNameFuncAddr(Function->GetOwnerClass(), OutReferenceGatherers.UniqueCrossModuleReferences);
 	}
 	else
 	{
@@ -2811,7 +2559,6 @@ void FNativeClassHeaderGenerator::ExportFunction(FOutputDevice& Out, FReferenceG
 
 	TTuple<FString, FString> PropertyRange = OutputProperties(CurrentFunctionText, StaticDefinitions, OutReferenceGatherers, *FString::Printf(TEXT("%s::"), *StaticsStructName), Props, TEXT("\t\t"), TEXT("\t"));
 
-	FUnrealFunctionDefinitionInfo& FuncDef = GTypeDefinitionInfoMap.FindChecked(Function).AsFunctionChecked();
 	const FFuncInfo&     FunctionData = FuncDef.GetFunctionData();
 	const bool           bIsNet       = !!(FunctionData.FunctionFlags & (FUNC_NetRequest | FUNC_NetResponse));
 
@@ -3769,6 +3516,7 @@ void FNativeClassHeaderGenerator::ExportEnum(FOutputDevice& Out, UEnum* Enum) co
 void FNativeClassHeaderGenerator::ExportGeneratedStructBodyMacros(FOutputDevice& OutGeneratedHeaderText, FOutputDevice& Out, FReferenceGatherers& OutReferenceGatherers, const FUnrealSourceFile& SourceFile, UScriptStruct* Struct) const
 {
 	FUnrealScriptStructDefinitionInfo& ScriptStructDef = GTypeDefinitionInfoMap.FindChecked(Struct).AsScriptStructChecked();
+	ScriptStructDef.AddCrossModuleReference(OutReferenceGatherers.UniqueCrossModuleReferences, true);
 
 	const bool bIsDynamic = FClass::IsDynamic(Struct);
 	const FString ActualStructName = FNativeClassHeaderGenerator::GetOverriddenName(Struct);
@@ -3778,8 +3526,8 @@ void FNativeClassHeaderGenerator::ExportGeneratedStructBodyMacros(FOutputDevice&
 
 	const FString StructNameCPP = FNameLookupCPP::GetNameCPP(Struct);
 
-	const FString& SingletonName = GetSingletonName(Struct, OutReferenceGatherers.UniqueCrossModuleReferences);
-	const FString ChoppedSingletonName = SingletonName.LeftChop(2);
+	const FString& SingletonName = ScriptStructDef.GetSingletonName(true);
+	const FString& ChoppedSingletonName = ScriptStructDef.GetSingletonNameChopped(true);
 
 	const FString RigVMParameterPrefix = TEXT("FRigVMExecuteContext& RigVMExecuteContext");
 	TArray<FString> RigVMVirtualFuncProlog, RigVMVirtualFuncEpilog, RigVMStubProlog;
@@ -4108,7 +3856,7 @@ void FNativeClassHeaderGenerator::ExportGeneratedStructBodyMacros(FOutputDevice&
 	if (BaseStruct)
 	{
 		CastChecked<UScriptStruct>(BaseStruct); // this better actually be a script struct
-		GetSingletonName(BaseStruct, OutReferenceGatherers.UniqueCrossModuleReferences); // Call to potentially collect references
+		GTypeDefinitionInfoMap.FindChecked(BaseStruct).AsFieldChecked().AddCrossModuleReference(OutReferenceGatherers.UniqueCrossModuleReferences, true);
 	}
 
 	EStructFlags UncomputedFlags = (EStructFlags)(Struct->StructFlags & ~STRUCT_ComputedFlags);
@@ -4116,7 +3864,7 @@ void FNativeClassHeaderGenerator::ExportGeneratedStructBodyMacros(FOutputDevice&
 	FString OuterFunc;
 	if (!bIsDynamic)
 	{
-		OuterFunc = GetPackageSingletonName(CastChecked<UPackage>(Struct->GetOuter()), OutReferenceGatherers.UniqueCrossModuleReferences).LeftChop(2);
+		OuterFunc = GetPackageSingletonNameFuncAddr(CastChecked<UPackage>(Struct->GetOuter()), OutReferenceGatherers.UniqueCrossModuleReferences);
 	}
 	else
 	{
@@ -4297,14 +4045,15 @@ void FNativeClassHeaderGenerator::ExportGeneratedStructBodyMacros(FOutputDevice&
 
 void FNativeClassHeaderGenerator::ExportGeneratedEnumInitCode(FOutputDevice& Out, FReferenceGatherers& OutReferenceGatherers, const FUnrealSourceFile& SourceFile, UEnum* Enum) const
 {
-	const bool    bIsDynamic            = FClass::IsDynamic(static_cast<UField*>(Enum));
-	const FString SingletonName         = GetSingletonNameFuncAddr(Enum, OutReferenceGatherers.UniqueCrossModuleReferences);
+	FUnrealEnumDefinitionInfo& EnumDef  = GTypeDefinitionInfoMap.FindChecked(Enum).AsEnumChecked();
+	const bool bIsDynamic               = FClass::IsDynamic(static_cast<UField*>(Enum));
+	const FString& SingletonName        = EnumDef.GetSingletonNameChopped(true);
 	const FString EnumNameCpp           = Enum->GetName(); //UserDefinedEnum should already have a valid cpp name.
 	const FString OverriddenEnumNameCpp = FNativeClassHeaderGenerator::GetOverriddenName(Enum);
-	const FString StaticsStructName     = SingletonName/*.LeftChop(2)*/ + TEXT("_Statics");
+	const FString StaticsStructName     = SingletonName + TEXT("_Statics");
+	const bool bIsEditorOnlyDataType    = EnumDef.IsEditorOnly();
 
-	FUnrealEnumDefinitionInfo& EnumDef = GTypeDefinitionInfoMap.FindChecked(Enum).AsEnumChecked();
-	const bool bIsEditorOnlyDataType = EnumDef.IsEditorOnly();
+	EnumDef.AddCrossModuleReference(OutReferenceGatherers.UniqueCrossModuleReferences, true);
 
 	FMacroBlockEmitter EditorOnlyData(Out, TEXT("WITH_EDITORONLY_DATA"));
 	EditorOnlyData(bIsEditorOnlyDataType);
@@ -4345,7 +4094,7 @@ void FNativeClassHeaderGenerator::ExportGeneratedEnumInitCode(FOutputDevice& Out
 
 	Out.Logf(TEXT("\t}\r\n"));
 
-	const FString& EnumSingletonName = GetSingletonName(Enum, OutReferenceGatherers.UniqueCrossModuleReferences);
+	const FString& EnumSingletonName = EnumDef.GetSingletonName(true);
 	const FString HashFuncName = FString::Printf(TEXT("Get_%s_Hash"), *SingletonName);
 
 	Out.Logf(TEXT("\ttemplate<> %sUEnum* StaticEnum<%s>()\r\n"), *GetAPIString(), *Enum->CppType);
@@ -6052,7 +5801,7 @@ void FNativeClassHeaderGenerator::GenerateSourceFiles(
 						// Is this ever not the case?
 						if (Enum->GetOuter()->IsA(UPackage::StaticClass()))
 						{
-							GeneratedFunctionDeclarations.Log(FTypeSingletonCache::Get(Enum).GetExternDecl());
+							GeneratedFunctionDeclarations.Log(GTypeDefinitionInfoMap.FindChecked(Enum).AsFieldChecked().GetExternDecl(true));
 							Generator.ExportGeneratedEnumInitCode(GeneratedCPPText, ReferenceGatherers, SourceFile, Enum);
 						}
 					}
@@ -6061,14 +5810,14 @@ void FNativeClassHeaderGenerator::GenerateSourceFiles(
 					// reverse the order.
 					for (UScriptStruct* Struct : Structs)
 					{
-						GeneratedFunctionDeclarations.Log(FTypeSingletonCache::Get(Struct).GetExternDecl());
+						GeneratedFunctionDeclarations.Log(GTypeDefinitionInfoMap.FindChecked(Struct).AsFieldChecked().GetExternDecl(true));
 						Generator.ExportGeneratedStructBodyMacros(GeneratedHeaderText, GeneratedCPPText, ReferenceGatherers, SourceFile, Struct);
 					}
 
 					// export delegate definitions
 					for (UDelegateFunction* Func : DelegateFunctions)
 					{
-						GeneratedFunctionDeclarations.Log(FTypeSingletonCache::Get(Func).GetExternDecl());
+						GeneratedFunctionDeclarations.Log(GTypeDefinitionInfoMap.FindChecked(Func).AsFieldChecked().GetExternDecl(true));
 						Generator.ExportDelegateDeclaration(GeneratedCPPText, ReferenceGatherers, SourceFile, Func);
 					}
 
@@ -7157,6 +6906,8 @@ void PrepareTypesForExport(TArray<FUnrealPackageDefinitionInfo*>& PackageDefs)
 		);
 	}
 	FResults::WaitForErrorTasks();
+
+	GTypeDefinitionInfoMap.ForAllTypes([](FUnrealTypeDefinitionInfo& TypeDef) { TypeDef.PostParseFinalize(); });
 }
 
 void Export(TArray<FUnrealPackageDefinitionInfo*>& PackageDefs, TArray<FUnrealSourceFile*>& OrderedSourceFiles)
@@ -7179,7 +6930,7 @@ void Export(TArray<FUnrealPackageDefinitionInfo*>& PackageDefs, TArray<FUnrealSo
 	FResults::WaitForErrorTasks();
 }
 
-// Exports the class to all vailable plugins
+// Exports the class to all available plugins
 void ExportClassToScriptPlugins(UClass* Class, const FManifestModule& Module, IScriptGeneratorPluginInterface& ScriptPlugin)
 {
 	TSharedRef<FUnrealTypeDefinitionInfo>* DefinitionInfoRef = GTypeDefinitionInfoMap.Find(Class);
diff --git a/Engine/Source/Programs/UnrealHeaderTool/Private/HeaderParser.cpp b/Engine/Source/Programs/UnrealHeaderTool/Private/HeaderParser.cpp
--- a/Engine/Source/Programs/UnrealHeaderTool/Private/HeaderParser.cpp
+++ b/Engine/Source/Programs/UnrealHeaderTool/Private/HeaderParser.cpp
@@ -7385,7 +7385,7 @@ UDelegateFunction* FHeaderParser::CompileDelegateDeclaration(const TCHAR* Delega
 	FuncInfo.MacroLine = InputLine;
 	TSharedRef<FUnrealFunctionDefinitionInfo> FuncDef = MakeShareable(new FUnrealFunctionDefinitionInfo(SourceFile, InputLine, FString(FuncInfoFunctionIdentifier), MoveTemp(FuncInfo)));
 	FuncDef->SetObject(FuncDef->GetFunctionData().FunctionReference);
-	GTypeDefinitionInfoMap.Add(FuncDef->GetFunctionData().FunctionReference, MoveTemp(FuncDef));
+	GTypeDefinitionInfoMap.Add(FuncDef->GetFunctionData().FunctionReference, FuncDef);
 
 	// Create the return value property
 	if (bHasReturnValue)
@@ -7824,7 +7824,7 @@ void FHeaderParser::CompileFunctionDeclaration()
 
 	TSharedRef<FUnrealFunctionDefinitionInfo> FuncDef = MakeShareable(new FUnrealFunctionDefinitionInfo(SourceFile, InputLine, FString(FuncInfo.Function.Identifier), FFuncInfo(FuncInfo)));
 	FuncDef->SetObject(FuncDef->GetFunctionData().FunctionReference);
-	GTypeDefinitionInfoMap.Add(FuncDef->GetFunctionData().FunctionReference, MoveTemp(FuncDef));
+	GTypeDefinitionInfoMap.Add(FuncDef->GetFunctionData().FunctionReference, FuncDef);
 	if (FuncInfo.FunctionReference->HasAnyFunctionFlags(FUNC_Delegate))
 	{
 		GetCurrentClassData()->MarkContainsDelegate();
diff --git a/Engine/Source/Programs/UnrealHeaderTool/Private/NativeClassExporter.h b/Engine/Source/Programs/UnrealHeaderTool/Private/NativeClassExporter.h
--- a/Engine/Source/Programs/UnrealHeaderTool/Private/NativeClassExporter.h
+++ b/Engine/Source/Programs/UnrealHeaderTool/Private/NativeClassExporter.h
@@ -329,14 +329,14 @@ private:
 	 */
 	static void ExportProperties(FOutputDevice& Out, UStruct* Struct, int32 TextIndent);
 
-	/** Return the name of the singleton function that returns the UObject for Item */
-	static const FString& GetPackageSingletonName(const UPackage* Item, TSet<FString>* UniqueCrossModuleReferences);
+	/** Return the name of the singleton function */
+	static const FString& GetPackageSingletonName(UPackage* Item, TSet<FString>* UniqueCrossModuleReferences);
 
-	/** Return the name of the singleton function that returns the UObject for Item */
-	static const FString& GetSingletonName(UField* Item, TSet<FString>* UniqueCrossModuleReferences, bool bRequiresValidObject=true);
+	/** Return the address of the singleton function */
+	static const FString& GetPackageSingletonNameFuncAddr(UPackage* Item, TSet<FString>* UniqueCrossModuleReferences);
 
 	/** Return the address of the singleton function - handles nullptr */
-	static FString GetSingletonNameFuncAddr(UField* Item, TSet<FString>* UniqueCrossModuleReferences, bool bRequiresValidObject=true);
+	static const FString& GetSingletonNameFuncAddr(UField* Item, TSet<FString>* UniqueCrossModuleReferences, bool bRequiresValidObject = true);
 
 	/**
 	 * Returns the name (overridden if marked up) or "" wrappers for use in a string literal.
diff --git a/Engine/Source/Programs/UnrealHeaderTool/Private/UnrealTypeDefinitionInfo.cpp b/Engine/Source/Programs/UnrealHeaderTool/Private/UnrealTypeDefinitionInfo.cpp
--- a/Engine/Source/Programs/UnrealHeaderTool/Private/UnrealTypeDefinitionInfo.cpp
+++ b/Engine/Source/Programs/UnrealHeaderTool/Private/UnrealTypeDefinitionInfo.cpp
@@ -7,6 +7,63 @@
 #include "Misc/PackageName.h"
 #include "UObject/ErrorException.h"
 
+namespace
+{
+	/**
+	 * As part of the singleton name, collect the parent chain names
+	 */
+	static void AddOuterNames(FUHTStringBuilder& Out, UObject* Outer)
+	{
+		if (Outer == nullptr)
+		{
+			return;
+		}
+
+		if (Cast<UClass>(Outer) || Cast<UScriptStruct>(Outer))
+		{
+			// Structs can also have UPackage outer.
+			if (!Outer->IsA<UClass>() && !Outer->GetOuter()->IsA<UPackage>())
+			{
+				AddOuterNames(Out, Outer->GetOuter());
+			}
+			Out.Append(TEXT("_"));
+			Out.Append(FNameLookupCPP::GetNameCPP(Cast<UStruct>(Outer)));
+		}
+		else if (UPackage* Package = Cast<UPackage>(Outer))
+		{
+			Out.Append(TEXT("_"));
+			Out.Append(FPackageName::GetShortName(Outer->GetName()));
+		}
+		else
+		{
+			AddOuterNames(Out, Outer->GetOuter());
+			Out.Append(TEXT("_"));
+			Out.Append(Outer->GetName());
+		}
+	}
+
+	/**
+	 * Generates singleton name.
+	 */
+	static void GenerateSingletonName(FUHTStringBuilder& Out, UField* Item, bool bRequiresValidObject)
+	{
+		check(Item);
+
+		Out.Append(TEXT("Z_Construct_"));
+		Out.Append(FNameLookupCPP::GetNameCPP(Item->GetClass()));
+		AddOuterNames(Out, Item);
+
+		if (UClass* ItemClass = Cast<UClass>(Item))
+		{
+			if (!bRequiresValidObject && !ItemClass->HasAllClassFlags(CLASS_Intrinsic))
+			{
+				Out.Append(TEXT("_NoRegister"));
+			}
+		}
+		Out.Append(TEXT("()"));
+	}
+}
+
 FUnrealPropertyDefinitionInfo* FUnrealTypeDefinitionInfo::AsProperty()
 {
 	return nullptr;
@@ -101,6 +158,58 @@ void FUnrealTypeDefinitionInfo::GetHashTag(FUHTStringBuilder& Out) const
 	}
 }
 
+void FUnrealPackageDefinitionInfo::PostParseFinalize()
+{
+	UPackage* Package = GetPackage();
+
+	FString PackageName = Package->GetName();
+	PackageName.ReplaceInline(TEXT("/"), TEXT("_"), ESearchCase::CaseSensitive);
+
+	SingletonName.Appendf(TEXT("Z_Construct_UPackage_%s()"), *PackageName);
+	SingletonNameChopped = SingletonName.LeftChop(2);
+	ExternDecl.Appendf(TEXT("\tUPackage* %s;\r\n"), *SingletonName);
+}
+
+void FUnrealPackageDefinitionInfo::AddCrossModuleReference(TSet<FString>* UniqueCrossModuleReferences) const
+{
+	if (UniqueCrossModuleReferences)
+	{
+		UniqueCrossModuleReferences->Add(GetExternDecl());
+	}
+}
+
+void FUnrealFieldDefinitionInfo::PostParseFinalize()
+{
+	const TCHAR* TypeStr = GetSimplifiedTypeClass();
+	UField* Field = GetField();
+	FString PackageShortName = FPackageName::GetShortName(Field->GetOutermost()).ToUpper();
+
+	FUHTStringBuilder Out;
+	GenerateSingletonName(Out, Field, false);
+	ExternDecl[0].Appendf(TEXT("\t%s_API %s* %s;\r\n"), *PackageShortName, TypeStr, *Out);
+	SingletonName[0] = Out;
+	SingletonNameChopped[0] = SingletonName[0].LeftChop(2);
+
+	Out.Reset();
+	GenerateSingletonName(Out, Field, true);
+	ExternDecl[1].Appendf(TEXT("\t%s_API %s* %s;\r\n"), *PackageShortName, TypeStr, *Out);
+	SingletonName[1] = Out;
+	SingletonNameChopped[1] = SingletonName[1].LeftChop(2);
+}
+
+void FUnrealFieldDefinitionInfo::AddCrossModuleReference(TSet<FString>* UniqueCrossModuleReferences, bool bRequiresValidObject) const
+{
+	// We don't need to export UFunction externs, though we may need the externs for UDelegateFunctions
+	if (UniqueCrossModuleReferences)
+	{
+		UField* Field = GetField();
+		if (!Field->IsA<UFunction>() || Field->IsA<UDelegateFunction>())
+		{
+			UniqueCrossModuleReferences->Add(GetExternDecl(bRequiresValidObject));
+		}
+	}
+}
+
 uint32 FUnrealScriptStructDefinitionInfo::GetHash(bool bIncludeNoExport) const
 {
 	if (!bIncludeNoExport)
diff --git a/Engine/Source/Programs/UnrealHeaderTool/Private/UnrealTypeDefinitionInfo.h b/Engine/Source/Programs/UnrealHeaderTool/Private/UnrealTypeDefinitionInfo.h
--- a/Engine/Source/Programs/UnrealHeaderTool/Private/UnrealTypeDefinitionInfo.h
+++ b/Engine/Source/Programs/UnrealHeaderTool/Private/UnrealTypeDefinitionInfo.h
@@ -218,6 +218,17 @@ public:
 		return *Info;
 	}
 
+	/**
+	 * Perform any post parsing finalization and validation
+	 */
+	virtual void PostParseFinalize()
+	{}
+
+	/**
+	 * Return the C++ type name that represents this type (i.e. UClass)
+	 */
+	virtual const TCHAR* GetSimplifiedTypeClass() const = 0;
+
 	/**
 	 * Return the compilation scope associated with this object
 	 */
@@ -326,6 +337,14 @@ public:
 		return this;
 	}
 
+	/**
+	 * Return the C++ type name that represents this type (i.e. UClass)
+	 */
+	virtual const TCHAR* GetSimplifiedTypeClass() const override
+	{
+		return TEXT("UProperty");
+	}
+
 	/**
 	 * Return the Engine instance associated with the compiler instance
 	 */
@@ -436,6 +455,14 @@ public:
 		return static_cast<UPackage*>(GetObject());
 	}
 
+	/**
+	 * Return the C++ type name that represents this type (i.e. UClass)
+	 */
+	virtual const TCHAR* GetSimplifiedTypeClass() const override
+	{
+		return TEXT("UPackage");
+	}
+
 	/**
 	 * Return the module information from the manifest associated with this package
 	 */
@@ -493,10 +520,47 @@ public:
 		return ShortUpperName;
 	}
 
+	/**
+	 * Perform any post parsing finalization and validation
+	 */
+	virtual void PostParseFinalize() override;
+
+	/**
+	 * Add a unique cross module reference for this field
+	 */
+	void AddCrossModuleReference(TSet<FString>* UniqueCrossModuleReferences) const;
+
+	/**
+	 * Return the name of the singleton for this field.  Only valid post parsing
+	 */
+	const FString& GetSingletonName() const
+	{
+		return SingletonName;
+	}
+
+	/**
+	 * Return the name of the singleton without the trailing "()" for this field.  Only valid post parsing
+	 */
+	const FString& GetSingletonNameChopped() const
+	{
+		return SingletonNameChopped;
+	}
+
+	/**
+	 * Return the external declaration for this field.  Only valid post parsing
+	 */
+	const FString& GetExternDecl() const
+	{
+		return ExternDecl;
+	}
+
 private:
 	const FManifestModule& Module;
 	TArray<TSharedRef<FUnrealSourceFile>> AllSourceFiles;
 	TArray<UClass*> AllClasses;
+	FString SingletonName;
+	FString SingletonNameChopped;
+	FString ExternDecl;
 	FString ShortUpperName;
 	FString API;
 	bool bWriteClassesH = false;
@@ -526,7 +590,44 @@ public:
 		return static_cast<UField*>(GetObject());
 	}
 
+	/**
+	 * Perform any post parsing finalization and validation
+	 */
+	virtual void PostParseFinalize() override;
+
+	/**
+	 * Add a unique cross module reference for this field
+	 */
+	void AddCrossModuleReference(TSet<FString>* UniqueCrossModuleReferences, bool bRequiresValidObject) const;
+
+	/**
+	 * Return the name of the singleton for this field.  Only valid post parsing
+	 */
+	const FString& GetSingletonName(bool bRequiresValidObject) const
+	{
+		return SingletonName[bRequiresValidObject];
+	}
+
+	/**
+	 * Return the name of the singleton without the trailing "()" for this field.  Only valid post parsing
+	 */
+	const FString& GetSingletonNameChopped(bool bRequiresValidObject) const
+	{
+		return SingletonNameChopped[bRequiresValidObject];
+	}
+
+	/**
+	 * Return the external declaration for this field.  Only valid post parsing
+	 */
+	const FString& GetExternDecl(bool bRequiresValidObject) const
+	{
+		return ExternDecl[bRequiresValidObject];
+	}
+
 private:
+	FString SingletonName[2];
+	FString SingletonNameChopped[2];
+	FString ExternDecl[2];
 };
 
 /**
@@ -544,6 +645,14 @@ public:
 		return this;
 	}
 
+	/**
+	 * Return the C++ type name that represents this type (i.e. UClass)
+	 */
+	virtual const TCHAR* GetSimplifiedTypeClass() const override
+	{
+		return TEXT("UEnum");
+	}
+
 	/**
 	 * Return the Engine instance associated with the compiler instance
 	 */
@@ -637,6 +746,14 @@ public:
 
 	virtual uint32 GetHash(bool bIncludeNoExport = true) const override;
 
+	/**
+	 * Return the C++ type name that represents this type (i.e. UClass)
+	 */
+	virtual const TCHAR* GetSimplifiedTypeClass() const override
+	{
+		return TEXT("UScriptStruct");
+	}
+
 	/**
 	 * Return the Engine instance associated with the compiler instance
 	 */
@@ -682,6 +799,14 @@ public:
 		return this;
 	}
 
+	/**
+	 * Return the C++ type name that represents this type (i.e. UClass)
+	 */
+	virtual const TCHAR* GetSimplifiedTypeClass() const override
+	{
+		return TEXT("UFunction");
+	}
+
 	/**
 	 * Return the Engine instance associated with the compiler instance
 	 */
@@ -801,6 +926,14 @@ public:
 		return this;
 	}
 
+	/**
+	 * Return the C++ type name that represents this type (i.e. UClass)
+	 */
+	virtual const TCHAR* GetSimplifiedTypeClass() const override
+	{
+		return TEXT("UClass");
+	}
+
 	virtual uint32 GetHash(bool bIncludeNoExport = true) const override;
 
 	/**
