Removed the FTokenData and FPropertyData.  Information has been moved to the type system.

#rnx

diff --git a/Engine/Source/Programs/UnrealHeaderTool/Private/HeaderParser.cpp b/Engine/Source/Programs/UnrealHeaderTool/Private/HeaderParser.cpp
--- a/Engine/Source/Programs/UnrealHeaderTool/Private/HeaderParser.cpp
+++ b/Engine/Source/Programs/UnrealHeaderTool/Private/HeaderParser.cpp
@@ -3158,14 +3158,14 @@ void FHeaderParser::FixupDelegateProperties( UStruct* Struct, FScope& Scope, TMa
 				// this property and verify that the delegate property's "type" is an actual delegate function
 				FClassMetaData* StructData = GScriptHelper.FindClassData(Struct);
 				check(StructData);
-				FTokenData* DelegatePropertyToken = StructData->FindTokenData(Property);
+				FToken* DelegatePropertyToken = StructData->FindTokenData(Property);
 				check(DelegatePropertyToken);
 
 				// attempt to find the delegate function in the map of functions we've already found
-				UFunction* SourceDelegateFunction = DelegateCache.FindRef(DelegatePropertyToken->Token.DelegateName);
+				UFunction* SourceDelegateFunction = DelegateCache.FindRef(DelegatePropertyToken->DelegateName);
 				if (SourceDelegateFunction == nullptr)
 				{
-					FString NameOfDelegateFunction = DelegatePropertyToken->Token.DelegateName.ToString() + FString( HEADER_GENERATED_DELEGATE_SIGNATURE_SUFFIX );
+					FString NameOfDelegateFunction = DelegatePropertyToken->DelegateName.ToString() + FString( HEADER_GENERATED_DELEGATE_SIGNATURE_SUFFIX );
 					if ( !NameOfDelegateFunction.Contains(TEXT(".")) )
 					{
 						// an unqualified delegate function name - search for a delegate function by this name within the current scope
@@ -3173,8 +3173,8 @@ void FHeaderParser::FixupDelegateProperties( UStruct* Struct, FScope& Scope, TMa
 						if (SourceDelegateFunction == nullptr)
 						{
 							// Try to find in other packages.
-							UObject* DelegateSignatureOuter = DelegatePropertyToken->Token.DelegateSignatureOwnerClass 
-								? ((UObject*)DelegatePropertyToken->Token.DelegateSignatureOwnerClass) 
+							UObject* DelegateSignatureOuter = DelegatePropertyToken->DelegateSignatureOwnerClass 
+								? ((UObject*)DelegatePropertyToken->DelegateSignatureOwnerClass) 
 								: ((UObject*)ANY_PACKAGE);
 							SourceDelegateFunction = Cast<UFunction>(StaticFindObject(UFunction::StaticClass(), DelegateSignatureOuter, *NameOfDelegateFunction));
 
@@ -3193,14 +3193,14 @@ void FHeaderParser::FixupDelegateProperties( UStruct* Struct, FScope& Scope, TMa
 						// verify that we got a valid string for the class name
 						if ( DelegateClassName.Len() == 0 )
 						{
-							UngetToken(DelegatePropertyToken->Token);
+							UngetToken(*DelegatePropertyToken);
 							FError::Throwf(TEXT("Invalid scope specified in delegate property function reference: '%s'"), *NameOfDelegateFunction);
 						}
 
 						// verify that we got a valid string for the name of the function
 						if ( DelegateName.Len() == 0 )
 						{
-							UngetToken(DelegatePropertyToken->Token);
+							UngetToken(*DelegatePropertyToken);
 							FError::Throwf(TEXT("Invalid delegate name specified in delegate property function reference '%s'"), *NameOfDelegateFunction);
 						}
 
@@ -3215,12 +3215,12 @@ void FHeaderParser::FixupDelegateProperties( UStruct* Struct, FScope& Scope, TMa
 
 					if ( SourceDelegateFunction == NULL )
 					{
-						UngetToken(DelegatePropertyToken->Token);
+						UngetToken(*DelegatePropertyToken);
 						FError::Throwf(TEXT("Failed to find delegate function '%s'"), *NameOfDelegateFunction);
 					}
 					else if ( (SourceDelegateFunction->FunctionFlags&FUNC_Delegate) == 0 )
 					{
-						UngetToken(DelegatePropertyToken->Token);
+						UngetToken(*DelegatePropertyToken);
 						FError::Throwf(TEXT("Only delegate functions can be used as the type for a delegate property; '%s' is not a delegate."), *NameOfDelegateFunction);
 					}
 				}
@@ -3228,14 +3228,14 @@ void FHeaderParser::FixupDelegateProperties( UStruct* Struct, FScope& Scope, TMa
 				// successfully found the delegate function that this delegate property corresponds to
 
 				// save this into the delegate cache for faster lookup later
-				DelegateCache.Add(DelegatePropertyToken->Token.DelegateName, SourceDelegateFunction);
+				DelegateCache.Add(DelegatePropertyToken->DelegateName, SourceDelegateFunction);
 
 				// bind it to the delegate property
 				if( DelegateProperty != NULL )
 				{
 					if( !SourceDelegateFunction->HasAnyFunctionFlags( FUNC_MulticastDelegate ) )
 					{
-						DelegateProperty->SignatureFunction = DelegatePropertyToken->Token.Function = SourceDelegateFunction;
+						DelegateProperty->SignatureFunction = DelegatePropertyToken->Function = SourceDelegateFunction;
 					}
 					else
 					{
@@ -3246,7 +3246,7 @@ void FHeaderParser::FixupDelegateProperties( UStruct* Struct, FScope& Scope, TMa
 				{
 					if( SourceDelegateFunction->HasAnyFunctionFlags( FUNC_MulticastDelegate ) )
 					{
-						MulticastDelegateProperty->SignatureFunction = DelegatePropertyToken->Token.Function = SourceDelegateFunction;
+						MulticastDelegateProperty->SignatureFunction = DelegatePropertyToken->Function = SourceDelegateFunction;
 
 						if(MulticastDelegateProperty->HasAnyPropertyFlags(CPF_BlueprintAssignable | CPF_BlueprintCallable))
 						{
@@ -3528,11 +3528,11 @@ void FHeaderParser::VerifyPropertyMarkups( UClass* TargetClass )
 
 			FClassMetaData* TargetClassData = GScriptHelper.FindClassData(TargetClass);
 			check(TargetClassData);
-			FTokenData* PropertyToken = TargetClassData->FindTokenData(Prop);
+			FToken* PropertyToken = TargetClassData->FindTokenData(Prop);
 			check(PropertyToken);
 
-			TGuardValue<int32> GuardedInputPos(InputPos, PropertyToken->Token.StartPos);
-			TGuardValue<int32> GuardedInputLine(InputLine, PropertyToken->Token.StartLine);
+			TGuardValue<int32> GuardedInputPos(InputPos, PropertyToken->StartPos);
+			TGuardValue<int32> GuardedInputLine(InputLine, PropertyToken->StartLine);
 
 			if (Prop->HasAnyPropertyFlags(CPF_RepNotify))
 			{
@@ -7358,6 +7358,15 @@ UDelegateFunction* FHeaderParser::CompileDelegateDeclaration(const TCHAR* Delega
 		SDF->DelegateName = DelegateName.Identifier;
 	}
 
+
+	// Save off function identifier in case we need it to error later, since we are moving the FuncInfo
+	TCHAR FuncInfoFunctionIdentifier[NAME_SIZE];
+	FCString::Strncpy(FuncInfoFunctionIdentifier, FuncInfo.Function.Identifier, NAME_SIZE);
+
+	TSharedRef<FUnrealFunctionDefinitionInfo> FuncDef = MakeShareable(new FUnrealFunctionDefinitionInfo(SourceFile, InputLine, FString(FuncInfoFunctionIdentifier), MoveTemp(FuncInfo)));
+	FuncDef->SetObject(FuncDef->GetFunctionData().FunctionReference);
+	GTypeDefinitionInfoMap.Add(FuncDef->GetFunctionData().FunctionReference, FuncDef);
+
 	// Get parameter list.
 	if (FoundParamCount)
 	{
@@ -7378,14 +7387,8 @@ UDelegateFunction* FHeaderParser::CompileDelegateDeclaration(const TCHAR* Delega
 		RequireSymbol(TEXT(')'), TEXT("Delegate Declaration"));
 	}
 
-	// Save off function identifier in case we need it to error later, since we are moving the FuncInfo
-	TCHAR FuncInfoFunctionIdentifier[NAME_SIZE];
-	FCString::Strncpy(FuncInfoFunctionIdentifier, FuncInfo.Function.Identifier, NAME_SIZE);
-
-	FuncInfo.MacroLine = InputLine;
-	TSharedRef<FUnrealFunctionDefinitionInfo> FuncDef = MakeShareable(new FUnrealFunctionDefinitionInfo(SourceFile, InputLine, FString(FuncInfoFunctionIdentifier), MoveTemp(FuncInfo)));
-	FuncDef->SetObject(FuncDef->GetFunctionData().FunctionReference);
-	GTypeDefinitionInfoMap.Add(FuncDef->GetFunctionData().FunctionReference, FuncDef);
+	// The macro line must be set here
+	FuncDef->GetFunctionData().MacroLine = InputLine;
 
 	// Create the return value property
 	if (bHasReturnValue)
diff --git a/Engine/Source/Programs/UnrealHeaderTool/Private/HeaderParser.h b/Engine/Source/Programs/UnrealHeaderTool/Private/HeaderParser.h
--- a/Engine/Source/Programs/UnrealHeaderTool/Private/HeaderParser.h
+++ b/Engine/Source/Programs/UnrealHeaderTool/Private/HeaderParser.h
@@ -8,6 +8,7 @@
 #include "Misc/CompilationResult.h"
 #include "Scope.h"
 #include "UnrealTypeDefinitionInfo.h"
+#include "GeneratedCodeVersion.h"
 
 class UClass;
 enum class EGeneratedCodeVersion : uint8;
diff --git a/Engine/Source/Programs/UnrealHeaderTool/Private/ParserHelper.cpp b/Engine/Source/Programs/UnrealHeaderTool/Private/ParserHelper.cpp
--- a/Engine/Source/Programs/UnrealHeaderTool/Private/ParserHelper.cpp
+++ b/Engine/Source/Programs/UnrealHeaderTool/Private/ParserHelper.cpp
@@ -21,68 +21,11 @@
  *			is declared in a package that is already compiled and has had its
  *			source stripped)
  */
-FTokenData* FClassMetaData::FindTokenData( FProperty* Prop )
+FToken* FClassMetaData::FindTokenData(FProperty* Prop)
 {
 	check(Prop);
-
-	FTokenData* Result = nullptr;
-	UObject* Outer = Prop->GetOwner<UObject>();
-	check(Outer);
-	UClass* OuterClass = nullptr;
-	if (Outer->IsA<UStruct>())
-	{
-		Result = GlobalPropertyData.Find(Prop);
-
-		if (Result == nullptr)
-		{
-			OuterClass = Cast<UClass>(Outer);
-
-			if (OuterClass != nullptr && OuterClass->GetSuperClass() != OuterClass)
-			{
-				OuterClass = OuterClass->GetSuperClass();
-			}
-		}
-	}
-	else
-	{
-		UFunction* OuterFunction = Cast<UFunction>(Outer);
-		if ( OuterFunction != NULL )
-		{
-			// function parameter, return, or local property
-			if (TSharedRef<FUnrealTypeDefinitionInfo>* TypeDef = GTypeDefinitionInfoMap.Find(OuterFunction))
-			{
-				FUnrealFunctionDefinitionInfo& FuncDef = (*TypeDef)->AsFunctionChecked();
-				FPropertyData& FunctionParameters = FuncDef.GetParameterData();
-				Result = FunctionParameters.Find(Prop);
-				if ( Result == NULL )
-				{
-					Result = FuncDef.GetReturnTokenData();
-				}
-			}
-			else
-			{
-				OuterClass = OuterFunction->GetOwnerClass();
-			}
-		}
-		else
-		{
-			// struct property
-			UScriptStruct* OuterStruct = Cast<UScriptStruct>(Outer);
-			check(OuterStruct != NULL);
-			OuterClass = OuterStruct->GetOwnerClass();
-		}
-	}
-
-	if (Result == nullptr && OuterClass != nullptr)
-	{
-		FClassMetaData* SuperClassData = GScriptHelper.FindClassData(OuterClass);
-		if (SuperClassData && SuperClassData != this)
-		{
-			Result = SuperClassData->FindTokenData(Prop);
-		}
-	}
-
-	return Result;
+	FUnrealPropertyDefinitionInfo& PropDef = GTypeDefinitionInfoMap.FindChecked(Prop).AsPropertyChecked();
+	return &PropDef.GetToken();
 }
 
 void FClassMetaData::AddProperty(FToken&& PropertyToken)
@@ -90,24 +33,13 @@ void FClassMetaData::AddProperty(FToken&& PropertyToken)
 	FProperty* Prop = PropertyToken.TokenProperty;
 	check(Prop);
 
+	FUnrealPropertyDefinitionInfo& PropDef = GTypeDefinitionInfoMap.FindChecked(Prop).AsPropertyChecked();
+	PropDef.GetToken() = MoveTemp(PropertyToken);
+
+	// Add the property to the owner
 	UObject* Outer = Prop->GetOwner<UObject>();
 	check(Outer);
-	UStruct* OuterClass = Cast<UStruct>(Outer);
-	if (OuterClass != NULL)
-	{
-		// global property
-		GlobalPropertyData.Set(Prop, MoveTemp(PropertyToken));
-	}
-	else
-	{
-		checkNoEntry();
-		UFunction* OuterFunction = Cast<UFunction>(Outer);
-		if (OuterFunction != NULL)
-		{
-			// function parameter, return, or local property
-			GTypeDefinitionInfoMap.FindChecked(OuterFunction).AsFunctionChecked().AddProperty(MoveTemp(PropertyToken));
-		}
-	}
+	GTypeDefinitionInfoMap.FindChecked(Outer).AsStructChecked().AddProperty(PropDef);
 
 	// update the optimization flags
 	if (!bContainsDelegates)
@@ -302,25 +234,6 @@ FClassMetaData* FCompilerMetadataManager::AddInterfaceClassData(UStruct* Struct,
 	return ClassData;
 }
 
-FTokenData* FPropertyData::Set(FProperty* InKey, FTokenData&& InValue)
-{
-	FTokenData* Result = NULL;
-
-	TSharedPtr<FTokenData>* pResult = PropertyMap.Find(InKey);
-	if (pResult != NULL)
-	{
-		Result = pResult->Get();
-		*Result = MoveTemp(InValue);
-	}
-	else
-	{
-		pResult = &PropertyMap.Emplace(InKey, new FTokenData(MoveTemp(InValue)));
-		Result = pResult->Get();
-	}
-
-	return Result;
-}
-
 void FCompilerMetadataManager::CheckForNoIInterfaces()
 {
 	for (const TPair<UStruct*, FClassMetaData*>& StructDataPair : InterfacesToVerify)
diff --git a/Engine/Source/Programs/UnrealHeaderTool/Private/ParserHelper.h b/Engine/Source/Programs/UnrealHeaderTool/Private/ParserHelper.h
--- a/Engine/Source/Programs/UnrealHeaderTool/Private/ParserHelper.h
+++ b/Engine/Source/Programs/UnrealHeaderTool/Private/ParserHelper.h
@@ -1173,82 +1173,6 @@ struct FFuncInfo
 	}
 };
 
-/**
- * Stores "compiler" data about an FToken.  "Compiler" data is data that is associated with a
- * specific property, function or class that is only needed during script compile.
- * This class is designed to make adding new compiler data very simple.
- *
- * - stores the raw evaluated bytecode associated with an FToken
- */
-struct FTokenData
-{
-	/** The token tracked by this FTokenData. */
-	FToken Token;
-
-	FTokenData() = default;
-	FTokenData(FTokenData&&) = default;
-	FTokenData& operator=(FTokenData&&) = default;
-	FTokenData(FToken&& InToken)
-		: Token(MoveTemp(InToken))
-	{
-	}
-};
-
-/**
- * Class for storing data about a list of properties.  Though FToken contains a reference to its
- * associated FProperty, it's faster lookup to use the FProperty as the key in a TMap.
- */
-class FPropertyData
-{
-	TMap<FProperty*, TSharedPtr<FTokenData>> PropertyMap;
-
-public:
-	/**
-	 * Returns the value associated with a specified key.
-	 * @param	Key - The key to search for.
-	 * @return	A pointer to the value associated with the specified key, or NULL if the key isn't contained in this map.  The pointer
-	 *			is only valid until the next change to any key in the map.
-	 */
-	FTokenData* Find(FProperty* Key)
-	{
-		FTokenData* Result = NULL;
-
-		TSharedPtr<FTokenData>* pResult = PropertyMap.Find(Key);
-		if ( pResult != NULL )
-		{
-			Result = pResult->Get();
-		}
-		return Result;
-	}
-	const FTokenData* Find(FProperty* Key) const
-	{
-		const FTokenData* Result = NULL;
-
-		const TSharedPtr<FTokenData>* pResult = PropertyMap.Find(Key);
-		if ( pResult != NULL )
-		{
-			Result = pResult->Get();
-		}
-		return Result;
-	}
-
-	void Shrink()
-	{
-		PropertyMap.Shrink();
-	}
-
-	/**
-	 * Sets the value associated with a key.  If the key already exists in the map, uses the same
-	 * value pointer and reinitalized the FTokenData with the input value.
-	 *
-	 * @param	InKey	the property to get a token wrapper for
-	 * @param	InValue	the token wrapper for the specified key
-	 *
-	 * @return	a pointer to token data created associated with the property
-	 */
-	FTokenData* Set(FProperty* InKey, FTokenData&& InValue);
-};
-
 /**
  * Tracks information about a multiple inheritance parent declaration for native script classes.
  */
@@ -1292,9 +1216,6 @@ enum class EParsedInterface
  */
 class FClassMetaData
 {
-	/** member properties for this class */
-	FPropertyData											GlobalPropertyData;
-
 	/** base classes to multiply inherit from (other than the main base class */
 	TArray<FMultipleInheritanceBaseClass*>					MultipleInheritanceParents;
 
@@ -1456,7 +1377,7 @@ public:
 	 *			is declared in a package that is already compiled and has had its
 	 *			source stripped)
 	 */
-	FTokenData* FindTokenData( FProperty* Prop );
+	FToken* FindTokenData( FProperty* Prop );
 
 	/**
 	 * Add a string to the list of inheritance parents for this class.
@@ -1495,7 +1416,6 @@ public:
 	 */
 	void Shrink()
 	{
-		GlobalPropertyData.Shrink();
 		MultipleInheritanceParents.Shrink();
 	}
 
@@ -1516,8 +1436,6 @@ public:
 
 	/** Parsed interface state */
 	EParsedInterface ParsedInterface = EParsedInterface::NotAnInterface;
-
-	friend struct FClassMetaDataArchiveProxy;
 };
 
 /**
diff --git a/Engine/Source/Programs/UnrealHeaderTool/Private/UnrealTypeDefinitionInfo.h b/Engine/Source/Programs/UnrealHeaderTool/Private/UnrealTypeDefinitionInfo.h
--- a/Engine/Source/Programs/UnrealHeaderTool/Private/UnrealTypeDefinitionInfo.h
+++ b/Engine/Source/Programs/UnrealHeaderTool/Private/UnrealTypeDefinitionInfo.h
@@ -387,7 +387,20 @@ public:
 		return AllocatorType;
 	}
 
+	/**
+	 * Return the token associated with the property parsing
+	 */
+	FToken& GetToken()
+	{
+		return Token;
+	}
+	const FToken& GetToken() const
+	{
+		return Token;
+	}
+
 private:
+	FToken Token;
 	FString ArrayDimensions;
 	FProperty* Property = nullptr;
 	EAllocatorType AllocatorType = EAllocatorType::Default;
@@ -723,8 +736,31 @@ public:
 		return static_cast<UStruct*>(GetObject());
 	}
 
+	/**
+	 * Add a new property to the structure
+	 */
+	virtual void AddProperty(FUnrealPropertyDefinitionInfo& PropertyDef)
+	{
+		Properties.Add(&PropertyDef);
+	}
+
+	/**
+	 * Get the collection of properties
+	 */
+	const TArray<FUnrealPropertyDefinitionInfo*> GetProperties() const
+	{
+		return Properties;
+	}
+	TArray<FUnrealPropertyDefinitionInfo*> GetParameters()
+	{
+		return Properties;
+	}
+
 private:
 	TSharedPtr<FScope> StructScope;
+
+	/** Properties of the structure */
+	TArray<FUnrealPropertyDefinitionInfo*> Properties;
 };
 
 /**
@@ -817,11 +853,9 @@ public:
 
 	/** @name getters */
 	//@{
-	const	FFuncInfo& GetFunctionData()	const { return FunctionData; }
-	const	FToken& GetReturnData()		const { return ReturnTypeData.Token; }
-	const	FPropertyData& GetParameterData()	const { return ParameterData; }
-	FPropertyData& GetParameterData() { return ParameterData; }
-	FTokenData* GetReturnTokenData() { return &ReturnTypeData; }
+	FFuncInfo& GetFunctionData() { return FunctionData; }
+	const FFuncInfo& GetFunctionData() const { return FunctionData; }
+	FUnrealPropertyDefinitionInfo* GetReturn() const { return ReturnProperty; }
 	//@}
 
 	void UpdateFunctionData(FFuncInfo& UpdatedFuncData)
@@ -838,19 +872,19 @@ public:
 	 *
 	 * @param	PropertyToken	the property to add
 	 */
-	void AddProperty(FToken&& PropertyToken)
+	virtual void AddProperty(FUnrealPropertyDefinitionInfo& PropertyDef) override
 	{
-		const FProperty* Prop = PropertyToken.TokenProperty;
-		check(Prop);
+		const FProperty* Prop = PropertyDef.GetProperty();
 		check((Prop->PropertyFlags & CPF_Parm) != 0);
 
 		if ((Prop->PropertyFlags & CPF_ReturnParm) != 0)
 		{
-			SetReturnData(MoveTemp(PropertyToken));
+			check(ReturnProperty == nullptr);
+			ReturnProperty = &PropertyDef;
 		}
 		else
 		{
-			AddParameter(MoveTemp(PropertyToken));
+			FUnrealStructDefinitionInfo::AddProperty(PropertyDef);
 		}
 	}
 
@@ -871,36 +905,12 @@ public:
 	}
 
 private:
-	/**
-	 * Adds a new parameter token
-	 *
-	 * @param	PropertyToken	token that should be added to the list
-	 */
-	void AddParameter(FToken&& PropertyToken)
-	{
-		check(PropertyToken.TokenProperty);
-		ParameterData.Set(PropertyToken.TokenProperty, MoveTemp(PropertyToken));
-	}
-
-	/**
-	 * Sets the value of the return token for this function
-	 *
-	 * @param	PropertyToken	token that should be added
-	 */
-	void SetReturnData(FToken&& PropertyToken)
-	{
-		check(PropertyToken.TokenProperty);
-		ReturnTypeData.Token = MoveTemp(PropertyToken);
-	}
 
 	/** info about the function associated with this FFunctionData */
-	FFuncInfo		FunctionData;
+	FFuncInfo FunctionData;
 
 	/** return value for this function */
-	FTokenData		ReturnTypeData;
-
-	/** function parameter data */
-	FPropertyData	ParameterData;
+	FUnrealPropertyDefinitionInfo* ReturnProperty = nullptr;
 };
 
 /**
