Remove the compiler meta data manager.
Elimination of some quick to resolve type system lookups.

#rnx

diff --git a/Engine/Source/Programs/UnrealHeaderTool/Private/CodeGenerator.cpp b/Engine/Source/Programs/UnrealHeaderTool/Private/CodeGenerator.cpp
--- a/Engine/Source/Programs/UnrealHeaderTool/Private/CodeGenerator.cpp
+++ b/Engine/Source/Programs/UnrealHeaderTool/Private/CodeGenerator.cpp
@@ -82,8 +82,6 @@ UClass* ProcessParsedClass(bool bClassIsAnInterface, const FString& ClassName, c
 UEnum* ProcessParsedEnum(const FString& EnumName, UObject* InParent, EObjectFlags Flags);
 UScriptStruct* ProcessParsedStruct(const FString& StructName, UObject* InParent, EObjectFlags Flags);
 
-FCompilerMetadataManager GScriptHelper;
-
 // Array of all the temporary header async file tasks so we can ensure they have completed before issuing our timings
 static FGraphEventArray GAsyncFileTasks;
 
@@ -829,6 +827,12 @@ const FString& FNativeClassHeaderGenerator::GetPackageSingletonNameFuncAddr(UPac
 	return Package.GetSingletonNameChopped();
 }
 
+const FString& FNativeClassHeaderGenerator::GetSingletonNameFuncAddr(FUnrealFieldDefinitionInfo& Field, TSet<FString>* UniqueCrossModuleReferences, bool bRequiresValidObject)
+{
+	Field.AddCrossModuleReference(UniqueCrossModuleReferences, bRequiresValidObject);
+	return Field.GetSingletonNameChopped(bRequiresValidObject);
+}
+
 const FString& FNativeClassHeaderGenerator::GetSingletonNameFuncAddr(UField* Item, TSet<FString>* UniqueCrossModuleReferences, bool bRequiresValidObject)
 {
 	if (!Item)
@@ -838,8 +842,7 @@ const FString& FNativeClassHeaderGenerator::GetSingletonNameFuncAddr(UField* Ite
 	else
 	{
 		FUnrealFieldDefinitionInfo& Field = GTypeDefinitionInfoMap.FindChecked(Item).AsFieldChecked();
-		Field.AddCrossModuleReference(UniqueCrossModuleReferences, bRequiresValidObject);
-		return Field.GetSingletonNameChopped(bRequiresValidObject);
+		return GetSingletonNameFuncAddr(Field, UniqueCrossModuleReferences, bRequiresValidObject);
 	}
 }
 
@@ -2190,7 +2193,7 @@ void FNativeClassHeaderGenerator::ExportNativeGeneratedInitCode(FOutputDevice& O
 				StaticDefinitions.Logf(
 					TEXT("%s\t\t{ &%s, %s },%s\r\n%s"),
 					BEGIN_WRAP_EDITOR_ONLY(bIsEditorOnlyFunction),
-					*GetSingletonNameFuncAddr(Function, OutReferenceGatherers.UniqueCrossModuleReferences),
+					*GetSingletonNameFuncAddr(FuncDef, OutReferenceGatherers.UniqueCrossModuleReferences),
 					*FNativeClassHeaderGenerator::GetUTF8OverriddenNameForLiteral(Function),
 					*FuncHashTag,
 					END_WRAP_EDITOR_ONLY(bIsEditorOnlyFunction)
@@ -2220,8 +2223,7 @@ void FNativeClassHeaderGenerator::ExportNativeGeneratedInitCode(FOutputDevice& O
 		TMap<FName, FString>* MetaDataMap = UMetaData::GetMapForObject(Class);
 		if (MetaDataMap)
 		{
-			FClassMetaData* ClassMetaData = GScriptHelper.FindClassData(Class);
-			if (ClassMetaData && ClassMetaData->bObjectInitializerConstructorDeclared)
+			if (ClassDef.GetStructMetaData().bObjectInitializerConstructorDeclared)
 			{
 				MetaDataMap->Add(NAME_ObjectInitializerConstructorDeclared, FString());
 			}
@@ -2813,23 +2815,21 @@ FString GetGeneratedMacroDeprecationWarning(const TCHAR* MacroName)
 FString GetPreservedAccessSpecifierString(FClass* Class)
 {
 	FString PreservedAccessSpecifier;
-	if (FClassMetaData* Data = GScriptHelper.FindClassData(Class))
-	{
-		switch (Data->GeneratedBodyMacroAccessSpecifier)
-		{
-		case EAccessSpecifier::ACCESS_Private:
-			PreservedAccessSpecifier = "private:";
-			break;
-		case EAccessSpecifier::ACCESS_Protected:
-			PreservedAccessSpecifier = "protected:";
-			break;
-		case EAccessSpecifier::ACCESS_Public:
-			PreservedAccessSpecifier = "public:";
-			break;
-		case EAccessSpecifier::ACCESS_NotAnAccessSpecifier :
-			PreservedAccessSpecifier = FString::Printf(TEXT("static_assert(false, \"Unknown access specifier for GENERATED_BODY() macro in class %s.\");"), *GetNameSafe(Class));
-			break;
-		}
+	FStructMetaData& Data = GTypeDefinitionInfoMap.FindChecked(Class).AsStructChecked().GetStructMetaData();
+	switch (Data.GeneratedBodyMacroAccessSpecifier)
+	{
+	case EAccessSpecifier::ACCESS_Private:
+		PreservedAccessSpecifier = "private:";
+		break;
+	case EAccessSpecifier::ACCESS_Protected:
+		PreservedAccessSpecifier = "protected:";
+		break;
+	case EAccessSpecifier::ACCESS_Public:
+		PreservedAccessSpecifier = "public:";
+		break;
+	case EAccessSpecifier::ACCESS_NotAnAccessSpecifier :
+		PreservedAccessSpecifier = FString::Printf(TEXT("static_assert(false, \"Unknown access specifier for GENERATED_BODY() macro in class %s.\");"), *GetNameSafe(Class));
+		break;
 	}
 
 	return PreservedAccessSpecifier + LINE_TERMINATOR;
@@ -2854,13 +2854,12 @@ void FNativeClassHeaderGenerator::ExportClassFromSourceFileInner(
 	FUHTStringBuilder EnhancedUObjectConstructorsMacroCall;
 
 	FUnrealClassDefinitionInfo& ClassDef = GTypeDefinitionInfoMap.FindChecked(Class).AsClassChecked();
-	FClassMetaData* ClassData = GScriptHelper.FindClassData(Class);
-	checkf(ClassData, TEXT("No class data generated for file %s"), *SourceFile.GetFilename());
+	FStructMetaData& StructData = ClassDef.GetStructMetaData();
 
 	// C++ -> VM stubs (native function execs)
 	FUHTStringBuilder ClassMacroCalls;
 	FUHTStringBuilder ClassNoPureDeclsMacroCalls;
-	ExportNativeFunctions(OutGeneratedHeaderText, OutCpp, ClassMacroCalls, ClassNoPureDeclsMacroCalls, OutReferenceGatherers, SourceFile, Class, ClassData);
+	ExportNativeFunctions(OutGeneratedHeaderText, OutCpp, ClassMacroCalls, ClassNoPureDeclsMacroCalls, OutReferenceGatherers, SourceFile, Class, StructData);
 
 	// Get Callback functions
 	TArray<UFunction*> CallbackFunctions;
@@ -2885,14 +2884,14 @@ void FNativeClassHeaderGenerator::ExportClassFromSourceFileInner(
 			ExportEventParm(UClassMacroContent, OutReferenceGatherers.ForwardDeclarations, Function, /*Indent=*/ 1, /*bOutputConstructor=*/ true, EExportingState::Normal);
 		}
 
-		FString MacroName = SourceFile.GetGeneratedMacroName(ClassData, TEXT("_EVENT_PARMS"));
+		FString MacroName = SourceFile.GetGeneratedMacroName(StructData, TEXT("_EVENT_PARMS"));
 		WriteMacro(OutGeneratedHeaderText, MacroName, UClassMacroContent);
 		PrologMacroCalls.Logf(TEXT("\t%s\r\n"), *MacroName);
 
 		// VM -> C++ proxies (events and delegates).
 		FOutputDeviceNull NullOutput;
 		FOutputDevice& CallbackOut = Class->HasAnyClassFlags(CLASS_NoExport) ? NullOutput : OutCpp;
-		FString CallbackWrappersMacroName = SourceFile.GetGeneratedMacroName(ClassData, TEXT("_CALLBACK_WRAPPERS"));
+		FString CallbackWrappersMacroName = SourceFile.GetGeneratedMacroName(StructData, TEXT("_CALLBACK_WRAPPERS"));
 		ExportCallbackFunctions(
 			OutGeneratedHeaderText,
 			CallbackOut,
@@ -2946,7 +2945,7 @@ void FNativeClassHeaderGenerator::ExportClassFromSourceFileInner(
 		FUHTStringBuilder Boilerplate, BoilerPlateCPP;
 		const TCHAR* MacroNameHeader;
 		const TCHAR* MacroNameCPP;
-		GeneratedSerializeFunctionHeaderMacroName = SourceFile.GetGeneratedMacroName(ClassData, TEXT("_ARCHIVESERIALIZER"));
+		GeneratedSerializeFunctionHeaderMacroName = SourceFile.GetGeneratedMacroName(StructData, TEXT("_ARCHIVESERIALIZER"));
 
 		if (ClassDef.GetArchiveType() == ESerializerArchiveType::StructuredArchiveRecord)
 		{
@@ -3016,12 +3015,12 @@ void FNativeClassHeaderGenerator::ExportClassFromSourceFileInner(
 
 		if (Class->HasAnyClassFlags(CLASS_Interface))
 		{
-			ExportConstructorsMacros(OutGeneratedHeaderText, OutCpp, StandardUObjectConstructorsMacroCall, EnhancedUObjectConstructorsMacroCall, SourceFile.GetGeneratedMacroName(ClassData), Class, *APIArg);
+			ExportConstructorsMacros(OutGeneratedHeaderText, OutCpp, StandardUObjectConstructorsMacroCall, EnhancedUObjectConstructorsMacroCall, SourceFile.GetGeneratedMacroName(StructData), Class, *APIArg);
 
-			FString InterfaceMacroName = SourceFile.GetGeneratedMacroName(ClassData, TEXT("_GENERATED_UINTERFACE_BODY"));
+			FString InterfaceMacroName = SourceFile.GetGeneratedMacroName(StructData, TEXT("_GENERATED_UINTERFACE_BODY"));
 			OutGeneratedHeaderText.Log(Macroize(*(InterfaceMacroName + TEXT("()")), *Boilerplate));
 
-			int32 ClassGeneratedBodyLine = ClassData->GetGeneratedBodyLine();
+			int32 ClassGeneratedBodyLine = StructData.GetGeneratedBodyLine();
 
 			FString DeprecationWarning = GetGeneratedMacroDeprecationWarning(TEXT("GENERATED_UINTERFACE_BODY"));
 
@@ -3087,11 +3086,11 @@ void FNativeClassHeaderGenerator::ExportClassFromSourceFileInner(
 				WriteReplicatedMacroData(ClassRange, *ClassCPPName, *APIArg, Class, SuperClass, InterfaceBoilerplate, SourceFile, OutFlags);
 			}
 
-			FString NoPureDeclsMacroName = SourceFile.GetGeneratedMacroName(ClassData, TEXT("_INCLASS_IINTERFACE_NO_PURE_DECLS"));
+			FString NoPureDeclsMacroName = SourceFile.GetGeneratedMacroName(StructData, TEXT("_INCLASS_IINTERFACE_NO_PURE_DECLS"));
 			WriteMacro(OutGeneratedHeaderText, NoPureDeclsMacroName, InterfaceBoilerplate);
 			ClassNoPureDeclsMacroCalls.Logf(TEXT("\t%s\r\n"), *NoPureDeclsMacroName);
 
-			FString MacroName = SourceFile.GetGeneratedMacroName(ClassData, TEXT("_INCLASS_IINTERFACE"));
+			FString MacroName = SourceFile.GetGeneratedMacroName(StructData, TEXT("_INCLASS_IINTERFACE"));
 			WriteMacro(OutGeneratedHeaderText, MacroName, InterfaceBoilerplate);
 			ClassMacroCalls.Logf(TEXT("\t%s\r\n"), *MacroName);
 		}
@@ -3115,21 +3114,21 @@ void FNativeClassHeaderGenerator::ExportClassFromSourceFileInner(
 			}
 
 			{
-				FString NoPureDeclsMacroName = SourceFile.GetGeneratedMacroName(ClassData, TEXT("_INCLASS_NO_PURE_DECLS"));
+				FString NoPureDeclsMacroName = SourceFile.GetGeneratedMacroName(StructData, TEXT("_INCLASS_NO_PURE_DECLS"));
 				WriteMacro(OutGeneratedHeaderText, NoPureDeclsMacroName, Boilerplate);
 				ClassNoPureDeclsMacroCalls.Logf(TEXT("\t%s\r\n"), *NoPureDeclsMacroName);
 
-				FString MacroName = SourceFile.GetGeneratedMacroName(ClassData, TEXT("_INCLASS"));
+				FString MacroName = SourceFile.GetGeneratedMacroName(StructData, TEXT("_INCLASS"));
 				WriteMacro(OutGeneratedHeaderText, MacroName, Boilerplate);
 				ClassMacroCalls.Logf(TEXT("\t%s\r\n"), *MacroName);
 
-				ExportConstructorsMacros(OutGeneratedHeaderText, OutCpp, StandardUObjectConstructorsMacroCall, EnhancedUObjectConstructorsMacroCall, SourceFile.GetGeneratedMacroName(ClassData), Class, *APIArg);
+				ExportConstructorsMacros(OutGeneratedHeaderText, OutCpp, StandardUObjectConstructorsMacroCall, EnhancedUObjectConstructorsMacroCall, SourceFile.GetGeneratedMacroName(StructData), Class, *APIArg);
 			}
 		}
 	}
 
 	{
-		FString MacroName = SourceFile.GetGeneratedMacroName(ClassData->GetPrologLine(), TEXT("_PROLOG"));
+		FString MacroName = SourceFile.GetGeneratedMacroName(StructData.GetPrologLine(), TEXT("_PROLOG"));
 		WriteMacro(OutGeneratedHeaderText, MacroName, PrologMacroCalls);
 	}
 
@@ -3148,7 +3147,7 @@ void FNativeClassHeaderGenerator::ExportClassFromSourceFileInner(
 		if (bIsIInterface)
 		{
 			MacroName = TEXT("GENERATED_IINTERFACE_BODY()");
-			GeneratedBodyLine = ClassData->GetInterfaceGeneratedBodyLine();
+			GeneratedBodyLine = StructData.GetInterfaceGeneratedBodyLine();
 			LegacyGeneratedBody = ClassMacroCalls;
 			GeneratedBody = ClassNoPureDeclsMacroCalls;
 		}
@@ -3156,7 +3155,7 @@ void FNativeClassHeaderGenerator::ExportClassFromSourceFileInner(
 		{
 			MacroName = TEXT("GENERATED_UCLASS_BODY()");
 			DeprecationWarning = GetGeneratedMacroDeprecationWarning(MacroName);
-			GeneratedBodyLine = ClassData->GetGeneratedBodyLine();
+			GeneratedBodyLine = StructData.GetGeneratedBodyLine();
 			LegacyGeneratedBody = FString::Printf(TEXT("%s%s"), *ClassMacroCalls, *StandardUObjectConstructorsMacroCall);
 			GeneratedBody = FString::Printf(TEXT("%s%s"), *ClassNoPureDeclsMacroCalls, *EnhancedUObjectConstructorsMacroCall);
 		}
@@ -3201,9 +3200,9 @@ void ExportCopyConstructorDefinition(FOutputDevice& Out, const TCHAR* API, const
  * @param Class Class to generate for.
  * @param API API string.
  */
-void ExportVTableHelperCtorAndCaller(FOutputDevice& Out, FClassMetaData* ClassData, const TCHAR* API, const TCHAR* ClassCPPName)
+void ExportVTableHelperCtorAndCaller(FOutputDevice& Out, FStructMetaData& StructData, const TCHAR* API, const TCHAR* ClassCPPName)
 {
-	if (!ClassData->bCustomVTableHelperConstructorDeclared)
+	if (!StructData.bCustomVTableHelperConstructorDeclared)
 	{
 		Out.Logf(TEXT("\tDECLARE_VTABLE_PTR_HELPER_CTOR(%s_API, %s);" LINE_TERMINATOR), API, ClassCPPName);
 	}
@@ -3217,13 +3216,13 @@ void ExportVTableHelperCtorAndCaller(FOutputDevice& Out, FClassMetaData* ClassDa
  * @param Class Class to generate constructor for.
  * @param API API string for this constructor.
  */
-void ExportStandardConstructorsMacro(FOutputDevice& Out, FClass* Class, FClassMetaData* ClassData, const TCHAR* API, const TCHAR* ClassCPPName)
+void ExportStandardConstructorsMacro(FOutputDevice& Out, FClass* Class, FStructMetaData& StructData, const TCHAR* API, const TCHAR* ClassCPPName)
 {
 	if (!Class->HasAnyClassFlags(CLASS_CustomConstructor))
 	{
 		Out.Logf(TEXT("\t/** Standard constructor, called after all reflected properties have been initialized */\r\n"));
 		Out.Logf(TEXT("\t%s_API %s(const FObjectInitializer& ObjectInitializer%s);\r\n"), API, ClassCPPName,
-			ClassData->bDefaultConstructorDeclared ? TEXT("") : TEXT(" = FObjectInitializer::Get()"));
+			StructData.bDefaultConstructorDeclared ? TEXT("") : TEXT(" = FObjectInitializer::Get()"));
 	}
 	if (Class->HasAnyClassFlags(CLASS_Abstract))
 	{
@@ -3235,7 +3234,7 @@ void ExportStandardConstructorsMacro(FOutputDevice& Out, FClass* Class, FClassMe
 
 	}
 
-	ExportVTableHelperCtorAndCaller(Out, ClassData, API, ClassCPPName);
+	ExportVTableHelperCtorAndCaller(Out, StructData, API, ClassCPPName);
 	ExportCopyConstructorDefinition(Out, API, ClassCPPName);
 }
 
@@ -3246,9 +3245,9 @@ void ExportStandardConstructorsMacro(FOutputDevice& Out, FClass* Class, FClassMe
  * @param Class Class to generate constructor for.
  * @param API API string for this constructor.
  */
-void ExportConstructorDefinition(FOutputDevice& Out, FClass* Class, FClassMetaData* ClassData, const TCHAR* API, const TCHAR* ClassCPPName)
+void ExportConstructorDefinition(FOutputDevice& Out, FClass* Class, FStructMetaData& StructData, const TCHAR* API, const TCHAR* ClassCPPName)
 {
-	if (!ClassData->bConstructorDeclared)
+	if (!StructData.bConstructorDeclared)
 	{
 		Out.Logf(TEXT("\t/** Standard constructor, called after all reflected properties have been initialized */\r\n"));
 
@@ -3258,34 +3257,36 @@ void ExportConstructorDefinition(FOutputDevice& Out, FClass* Class, FClassMetaDa
 		FClass* SuperClass = Class->GetSuperClass();
 		if (SuperClass != nullptr)
 		{
-			FClassMetaData* SuperClassData = GScriptHelper.FindClassData(SuperClass);
-			if (SuperClassData)
+			FUnrealStructDefinitionInfo& StructDef = GTypeDefinitionInfoMap.FindChecked(SuperClass).AsStructChecked();
+			if (StructDef.HasSource()) // Don't consider the internal types generated from the engine
 			{
+				FStructMetaData& SuperClassData = GTypeDefinitionInfoMap.FindChecked(SuperClass).AsStructChecked().GetStructMetaData();
+
 				// Since we are dependent on our SuperClass having determined which constructors are defined, 
 				// if it is not yet determined we will need to wait on it becoming available. 
 				// Since the SourceFile array provided to the ParallelFor is in dependency order and does not allow cyclic dependencies, 
 				// we can be certain that another thread has started processing the file containing our SuperClass before this
 				// file would have been assigned out,  so we just have to wait
-				while (!SuperClassData->bConstructorDeclared)
+				while (!SuperClassData.bConstructorDeclared)
 				{
 					FPlatformProcess::Sleep(0.01);
 				}
 
-				bSuperClassObjectInitializerConstructorDeclared = SuperClassData->bObjectInitializerConstructorDeclared;
+				bSuperClassObjectInitializerConstructorDeclared = SuperClassData.bObjectInitializerConstructorDeclared;
 			}
 		}
 		if (bSuperClassObjectInitializerConstructorDeclared)
 		{
 			Out.Logf(TEXT("\t%s_API %s(const FObjectInitializer& ObjectInitializer = FObjectInitializer::Get()) : Super(ObjectInitializer) { };\r\n"), API, ClassCPPName);
-			ClassData->bObjectInitializerConstructorDeclared = true;
+			StructData.bObjectInitializerConstructorDeclared = true;
 		}
 		else
 		{
 			Out.Logf(TEXT("\t%s_API %s() { };\r\n"), API, ClassCPPName);
-			ClassData->bDefaultConstructorDeclared = true;
+			StructData.bDefaultConstructorDeclared = true;
 		}
 
-		ClassData->bConstructorDeclared = true;
+		StructData.bConstructorDeclared = true;
 	}
 	ExportCopyConstructorDefinition(Out, API, ClassCPPName);
 }
@@ -3296,9 +3297,9 @@ void ExportConstructorDefinition(FOutputDevice& Out, FClass* Class, FClassMetaDa
  * @param Out Output device to generate to.
  * @param Class Class to generate constructor call definition for.
  */
-void ExportDefaultConstructorCallDefinition(FOutputDevice& Out, FClass* Class, FClassMetaData* ClassData, const TCHAR* ClassCPPName)
+void ExportDefaultConstructorCallDefinition(FOutputDevice& Out, FClass* Class, FStructMetaData& StructData, const TCHAR* ClassCPPName)
 {
-	if (ClassData->bObjectInitializerConstructorDeclared)
+	if (StructData.bObjectInitializerConstructorDeclared)
 	{
 		if (Class->HasAnyClassFlags(CLASS_Abstract))
 		{
@@ -3309,7 +3310,7 @@ void ExportDefaultConstructorCallDefinition(FOutputDevice& Out, FClass* Class, F
 			Out.Logf(TEXT("\tDEFINE_DEFAULT_OBJECT_INITIALIZER_CONSTRUCTOR_CALL(%s)\r\n"), ClassCPPName);
 		}
 	}
-	else if (ClassData->bDefaultConstructorDeclared)
+	else if (StructData.bDefaultConstructorDeclared)
 	{
 		if (Class->HasAnyClassFlags(CLASS_Abstract))
 		{
@@ -3333,11 +3334,11 @@ void ExportDefaultConstructorCallDefinition(FOutputDevice& Out, FClass* Class, F
  * @param Class Class to generate constructor for.
  * @param API API string for this constructor.
  */
-void ExportEnhancedConstructorsMacro(FOutputDevice& Out, FClass* Class, FClassMetaData* ClassData, const TCHAR* API, const TCHAR* ClassCPPName)
+void ExportEnhancedConstructorsMacro(FOutputDevice& Out, FClass* Class, FStructMetaData& StructData, const TCHAR* API, const TCHAR* ClassCPPName)
 {
-	ExportConstructorDefinition(Out, Class, ClassData, API, ClassCPPName);
-	ExportVTableHelperCtorAndCaller(Out, ClassData, API, ClassCPPName);
-	ExportDefaultConstructorCallDefinition(Out, Class, ClassData, ClassCPPName);
+	ExportConstructorDefinition(Out, Class, StructData, API, ClassCPPName);
+	ExportVTableHelperCtorAndCaller(Out, StructData, API, ClassCPPName);
+	ExportDefaultConstructorCallDefinition(Out, Class, StructData, ClassCPPName);
 }
 
 /**
@@ -3360,18 +3361,17 @@ void FNativeClassHeaderGenerator::ExportConstructorsMacros(FOutputDevice& OutGen
 {
 	const FString ClassCPPName = FNameLookupCPP::GetNameCPP(Class);
 
-	FClassMetaData* ClassData = GScriptHelper.FindClassData(Class);
-	check(ClassData);
+	FStructMetaData& StructData = GTypeDefinitionInfoMap.FindChecked(Class).AsStructChecked().GetStructMetaData();
 
 	FUHTStringBuilder StdMacro;
 	FUHTStringBuilder EnhMacro;
 	FString StdMacroName = ConstructorsMacroPrefix + TEXT("_STANDARD_CONSTRUCTORS");
 	FString EnhMacroName = ConstructorsMacroPrefix + TEXT("_ENHANCED_CONSTRUCTORS");
 
-	ExportStandardConstructorsMacro(StdMacro, Class, ClassData, APIArg, *ClassCPPName);
-	ExportEnhancedConstructorsMacro(EnhMacro, Class, ClassData, APIArg, *ClassCPPName);
+	ExportStandardConstructorsMacro(StdMacro, Class, StructData, APIArg, *ClassCPPName);
+	ExportEnhancedConstructorsMacro(EnhMacro, Class, StructData, APIArg, *ClassCPPName);
 
-	if (!ClassData->bCustomVTableHelperConstructorDeclared)
+	if (!StructData.bCustomVTableHelperConstructorDeclared)
 	{
 		Out.Logf(TEXT("\tDEFINE_VTABLE_PTR_HELPER_CTOR(%s);" LINE_TERMINATOR), *ClassCPPName);
 	}
@@ -5181,7 +5181,7 @@ struct FNativeFunctionStringBuilder
 	FUHTStringBuilder AutogeneratedStaticDataFuncs;
 };
 
-void FNativeClassHeaderGenerator::ExportNativeFunctions(FOutputDevice& OutGeneratedHeaderText, FOutputDevice& OutGeneratedCPPText, FOutputDevice& OutMacroCalls, FOutputDevice& OutNoPureDeclsMacroCalls, FReferenceGatherers& OutReferenceGatherers, const FUnrealSourceFile& SourceFile, UClass* Class, FClassMetaData* ClassData) const
+void FNativeClassHeaderGenerator::ExportNativeFunctions(FOutputDevice& OutGeneratedHeaderText, FOutputDevice& OutGeneratedCPPText, FOutputDevice& OutMacroCalls, FOutputDevice& OutNoPureDeclsMacroCalls, FReferenceGatherers& OutReferenceGatherers, const FUnrealSourceFile& SourceFile, UClass* Class, FStructMetaData& StructData) const
 {
 	FNativeFunctionStringBuilder RuntimeStringBuilders;
 	FNativeFunctionStringBuilder EditorStringBuilders;
@@ -5366,7 +5366,7 @@ void FNativeClassHeaderGenerator::ExportNativeFunctions(FOutputDevice& OutGenera
 
 	// static class data
 	{
-		FString MacroName = SourceFile.GetGeneratedMacroName(ClassData, TEXT("_SPARSE_DATA"));
+		FString MacroName = SourceFile.GetGeneratedMacroName(StructData, TEXT("_SPARSE_DATA"));
 
 		WriteMacro(OutGeneratedHeaderText, MacroName, RuntimeStringBuilders.AutogeneratedStaticData + RuntimeStringBuilders.AutogeneratedStaticDataFuncs);
 		OutMacroCalls.Logf(TEXT("\t%s\r\n"), *MacroName);
@@ -5375,7 +5375,7 @@ void FNativeClassHeaderGenerator::ExportNativeFunctions(FOutputDevice& OutGenera
 
 	// Write runtime wrappers
 	{
-		FString MacroName = SourceFile.GetGeneratedMacroName(ClassData, TEXT("_RPC_WRAPPERS"));
+		FString MacroName = SourceFile.GetGeneratedMacroName(StructData, TEXT("_RPC_WRAPPERS"));
 
 		// WriteMacro has an assumption about what will be at the end of this block that is no longer true due to splitting the
 		// definition and implementation, so add on a line terminator to satisfy it
@@ -5388,7 +5388,7 @@ void FNativeClassHeaderGenerator::ExportNativeFunctions(FOutputDevice& OutGenera
 		OutMacroCalls.Logf(TEXT("\t%s\r\n"), *MacroName);
 
 		// Put static checks before RPCWrappers to get proper messages from static asserts before compiler errors.
-		FString NoPureDeclsMacroName = SourceFile.GetGeneratedMacroName(ClassData, TEXT("_RPC_WRAPPERS_NO_PURE_DECLS"));
+		FString NoPureDeclsMacroName = SourceFile.GetGeneratedMacroName(StructData, TEXT("_RPC_WRAPPERS_NO_PURE_DECLS"));
 		if (SourceFile.GetGeneratedCodeVersionForStruct(Class) > EGeneratedCodeVersion::V1)
 		{
 			WriteMacro(OutGeneratedHeaderText, NoPureDeclsMacroName, RuntimeStringBuilders.RPCWrappers);
@@ -5408,7 +5408,7 @@ void FNativeClassHeaderGenerator::ExportNativeFunctions(FOutputDevice& OutGenera
 	{
 		OutGeneratedHeaderText.Log( BeginEditorOnlyGuard );
 
-		FString MacroName = SourceFile.GetGeneratedMacroName(ClassData, TEXT("_EDITOR_ONLY_RPC_WRAPPERS"));
+		FString MacroName = SourceFile.GetGeneratedMacroName(StructData, TEXT("_EDITOR_ONLY_RPC_WRAPPERS"));
 
 		// WriteMacro has an assumption about what will be at the end of this block that is no longer true due to splitting the
 		// definition and implementation, so add on a line terminator to satisfy it
@@ -5421,7 +5421,7 @@ void FNativeClassHeaderGenerator::ExportNativeFunctions(FOutputDevice& OutGenera
 		OutMacroCalls.Logf(TEXT("\t%s\r\n"), *MacroName);
 
 		// Put static checks before RPCWrappers to get proper messages from static asserts before compiler errors.
-		FString NoPureDeclsMacroName = SourceFile.GetGeneratedMacroName(ClassData, TEXT("_EDITOR_ONLY_RPC_WRAPPERS_NO_PURE_DECLS"));
+		FString NoPureDeclsMacroName = SourceFile.GetGeneratedMacroName(StructData, TEXT("_EDITOR_ONLY_RPC_WRAPPERS_NO_PURE_DECLS"));
 		if (SourceFile.GetGeneratedCodeVersionForStruct(Class) > EGeneratedCodeVersion::V1)
 		{
 			WriteMacro(OutGeneratedHeaderText, NoPureDeclsMacroName, EditorStringBuilders.RPCWrappers);
@@ -6231,20 +6231,29 @@ bool FNativeClassHeaderGenerator::SaveHeaderIfChanged(FGeneratedFileInfo& FileIn
 
 	if (bWriteContents || bVerifyContents)
 	{
+		// UHT is getting fast enough that we can create an I/O storm in these large directories.
+		// The lock limits us to writing one file at a time.  It doesn't impact performance 
+		// significantly.
+		static FCriticalSection WritePacer;
+
 		const FString& ProjectSavedDir = FPaths::ProjectSavedDir();
 		const FString CleanFilename = FPaths::GetCleanFilename(FileInfo.GetFilename());
 		const FString Ref = ProjectSavedDir / TEXT("ReferenceGeneratedCode") / CleanFilename;
 
 		if (bWriteContents)
 		{
+			FScopeLock Lock(&WritePacer);
 			bool Written = FFileHelper::SaveStringToFile(InNewHeaderContents, *Ref);
 			check(Written);
 		}
 		else
 		{
-			const FString Verify = ProjectSavedDir / TEXT("VerifyGeneratedCode") / CleanFilename;
-			bool Written = FFileHelper::SaveStringToFile(InNewHeaderContents, *Verify);
-			check(Written);
+			{
+				FScopeLock Lock(&WritePacer);
+				const FString Verify = ProjectSavedDir / TEXT("VerifyGeneratedCode") / CleanFilename;
+				bool Written = FFileHelper::SaveStringToFile(InNewHeaderContents, *Verify);
+				check(Written);
+			}
 
 			FString RefHeader;
 			FString Message;
@@ -7106,9 +7115,6 @@ ECompilationResult::Type UnrealHeaderTool_Main(const FString& ModuleInfoFilename
 	double TotalPluginTime = FResults::TimedTry([&PackageDefs]() { ExportToScriptPlugins(PackageDefs); });
 	double TotalSummaryTime = FResults::TimedTry([&PackageDefs]() { GenerateSummary(PackageDefs); });
 
-	// Avoid TArray slack for meta data.
-	GScriptHelper.Shrink();
-
 	// Finish all async file tasks before stopping the clock
 	FTaskGraphInterface::Get().WaitUntilTasksComplete(GAsyncFileTasks);
 
diff --git a/Engine/Source/Programs/UnrealHeaderTool/Private/HeaderParser.cpp b/Engine/Source/Programs/UnrealHeaderTool/Private/HeaderParser.cpp
--- a/Engine/Source/Programs/UnrealHeaderTool/Private/HeaderParser.cpp
+++ b/Engine/Source/Programs/UnrealHeaderTool/Private/HeaderParser.cpp
@@ -1268,11 +1268,11 @@ namespace
 		ValidateMetaDataFormat(Field, InMetaData);
 		if (Field.IsUObject())
 		{
-			FClassMetaData::AddMetaData(CastChecked<UField>(Field.ToUObject()), MoveTemp(InMetaData));
+			FStructMetaData::AddMetaData(CastChecked<UField>(Field.ToUObject()), MoveTemp(InMetaData));
 		}
 		else
 		{
-			FClassMetaData::AddMetaData(Field.ToField(), MoveTemp(InMetaData));
+			FStructMetaData::AddMetaData(Field.ToField(), MoveTemp(InMetaData));
 		}
 	}
 
@@ -1706,17 +1706,12 @@ void AddModuleRelativePathToMetadata(FUnrealSourceFile& SourceFile, TMap<FName,
  * @param Type Type for which to add module's relative path.
  * @param MetaData Meta data to fill the information.
  */
-void AddModuleRelativePathToMetadata(UField* Type, TMap<FName, FString> &MetaData)
+void AddModuleRelativePathToMetadata(FUnrealTypeDefinitionInfo& Type, TMap<FName, FString> &MetaData)
 {
-	// Add metadata for the module relative path.
-	TSharedRef<FUnrealTypeDefinitionInfo>* TypeDefinitionPtr = GTypeDefinitionInfoMap.Find(Type);
-	if (TypeDefinitionPtr != nullptr)
+	// Don't add module relative paths to functions.
+	if (Type.AsFunction() == nullptr && Type.HasSource())
 	{
-		// Don't add module relative paths to functions.
-		if ((*TypeDefinitionPtr)->AsFunction() == nullptr)
-		{
-			MetaData.Add(NAME_ModuleRelativePath, (*TypeDefinitionPtr)->GetUnrealSourceFile().GetModuleRelativePath());
-		}
+		MetaData.Add(NAME_ModuleRelativePath, Type.GetUnrealSourceFile().GetModuleRelativePath());
 	}
 }
 
@@ -1795,7 +1790,7 @@ UEnum* FHeaderParser::CompileEnum()
 	FUnrealEnumDefinitionInfo& EnumDef = GTypeDefinitionInfoMap.FindByNameChecked(EnumToken.Identifier).AsEnumChecked();
 	UEnum* Enum = EnumDef.GetEnum();
 
-	Scope->AddType(Enum);
+	Scope->AddType(EnumDef);
 
 	for (const FPropertySpecifier& Specifier : SpecifiersFound)
 	{
@@ -1922,7 +1917,7 @@ UEnum* FHeaderParser::CompileEnum()
 	// List of all metadata generated for this enum
 	TMap<FName,FString> EnumValueMetaData = EnumToken.MetaData;
 
-	AddModuleRelativePathToMetadata(Enum, EnumValueMetaData);
+	AddModuleRelativePathToMetadata(EnumDef, EnumValueMetaData);
 	AddFormattedPrevCommentAsTooltipMetaData(EnumValueMetaData);
 
 	// Parse all enums tags.
@@ -2640,6 +2635,7 @@ UScriptStruct* FHeaderParser::CompileStructDeclaration()
 			}
 			
 			FString ParentStructNameStripped;
+			FUnrealFieldDefinitionInfo* TypeDef = nullptr;
 			const UField* Type = nullptr;
 			bool bOverrideParentStructName = false;
 
@@ -2652,17 +2648,17 @@ UScriptStruct* FHeaderParser::CompileStructDeclaration()
 			// If we're expecting a prefix, first try finding the correct field with the stripped struct name
 			if (bOverrideParentStructName)
 			{
-				Type = StructScope->FindTypeByName(*ParentStructNameStripped);
+				TypeDef = StructScope->FindTypeByName(*ParentStructNameStripped);
 			}
 
 			// If it wasn't found, try to find the literal name given
-			if (Type == nullptr)
+			if (TypeDef == nullptr)
 			{
-				Type = StructScope->FindTypeByName(*ParentStructNameInScript);
+				TypeDef = StructScope->FindTypeByName(*ParentStructNameInScript);
 			}
 
 			// Resolve structs declared in another class  //@TODO: UCREMOVAL: This seems extreme
-			if (Type == nullptr)
+			if (TypeDef == nullptr)
 			{
 				if (bOverrideParentStructName)
 				{
@@ -2674,6 +2670,10 @@ UScriptStruct* FHeaderParser::CompileStructDeclaration()
 					Type = FindObject<UScriptStruct>(ANY_PACKAGE, *ParentStructNameInScript);
 				}
 			}
+			else
+			{
+				Type = TypeDef->GetField();
+			}
 
 			// If the struct still wasn't found, throw an error
 			if (Type == NULL)
@@ -2708,9 +2708,9 @@ UScriptStruct* FHeaderParser::CompileStructDeclaration()
 		StructFlags |= (BaseStruct->StructFlags&STRUCT_Inherit);
 	}
 
-	Scope->AddType(Struct);
+	Scope->AddType(StructDef);
 
-	AddModuleRelativePathToMetadata(Struct, MetaData);
+	AddModuleRelativePathToMetadata(StructDef, MetaData);
 
 	// Check to make sure the syntactic native prefix was set-up correctly.
 	// If this check results in a false positive, it will be flagged as an identifier failure.
@@ -2746,14 +2746,6 @@ UScriptStruct* FHeaderParser::CompileStructDeclaration()
 	// Assume that, but restore the parser state once we finish parsing this struct
 	TGuardValue<EAccessSpecifier> HoldFromClass(CurrentAccessSpecifier, ACCESS_Public);
 
-	{
-		FToken StructToken;
-		StructToken.Struct = Struct;
-
-		// add this struct to the compiler's persistent tracking system
-		FClassMetaData* ClassMetaData = GScriptHelper.AddClassData(StructToken.Struct, &SourceFile);
-	}
-
 	int32 SavedLineNumber = InputLine;
 
 	// Clear comment before parsing body of the struct.
@@ -2772,7 +2764,7 @@ UScriptStruct* FHeaderParser::CompileStructDeclaration()
 		}
 		else if (Token.Matches(TEXT("UPROPERTY"), ESearchCase::CaseSensitive))
 		{
-			CompileVariableDeclaration(Struct);
+			CompileVariableDeclaration(StructDef);
 		}
 		else if (Token.Matches(TEXT("UFUNCTION"), ESearchCase::CaseSensitive))
 		{
@@ -3156,25 +3148,29 @@ void FHeaderParser::FixupDelegateProperties( UStruct* Struct, FScope& Scope, TMa
 			{
 				// this FDelegateProperty corresponds to an actual delegate variable (i.e. delegate<SomeDelegate> Foo); we need to lookup the token data for
 				// this property and verify that the delegate property's "type" is an actual delegate function
-				FClassMetaData* StructData = GScriptHelper.FindClassData(Struct);
-				check(StructData);
-				FToken* DelegatePropertyToken = StructData->FindTokenData(Property);
-				check(DelegatePropertyToken);
+				FUnrealPropertyDefinitionInfo& PropDef = GTypeDefinitionInfoMap.FindChecked(Property).AsPropertyChecked();
+				FToken& DelegatePropertyToken = PropDef.GetToken();
 
 				// attempt to find the delegate function in the map of functions we've already found
-				UFunction* SourceDelegateFunction = DelegateCache.FindRef(DelegatePropertyToken->DelegateName);
+				UFunction* SourceDelegateFunction = DelegateCache.FindRef(DelegatePropertyToken.DelegateName);
 				if (SourceDelegateFunction == nullptr)
 				{
-					FString NameOfDelegateFunction = DelegatePropertyToken->DelegateName.ToString() + FString( HEADER_GENERATED_DELEGATE_SIGNATURE_SUFFIX );
+					FString NameOfDelegateFunction = DelegatePropertyToken.DelegateName.ToString() + FString( HEADER_GENERATED_DELEGATE_SIGNATURE_SUFFIX );
 					if ( !NameOfDelegateFunction.Contains(TEXT(".")) )
 					{
 						// an unqualified delegate function name - search for a delegate function by this name within the current scope
-						SourceDelegateFunction = Cast<UFunction>(Scope.FindTypeByName(*NameOfDelegateFunction));
+						if (FUnrealFieldDefinitionInfo* FoundDef = Scope.FindTypeByName(*NameOfDelegateFunction))
+						{
+							if (FUnrealFunctionDefinitionInfo* SourceDelegateFunctionDef = FoundDef->AsFunction())
+							{
+								SourceDelegateFunction = SourceDelegateFunctionDef->GetFunction();
+							}
+						}
 						if (SourceDelegateFunction == nullptr)
 						{
 							// Try to find in other packages.
-							UObject* DelegateSignatureOuter = DelegatePropertyToken->DelegateSignatureOwnerClass 
-								? ((UObject*)DelegatePropertyToken->DelegateSignatureOwnerClass) 
+							UObject* DelegateSignatureOuter = DelegatePropertyToken.DelegateSignatureOwnerClass 
+								? ((UObject*)DelegatePropertyToken.DelegateSignatureOwnerClass) 
 								: ((UObject*)ANY_PACKAGE);
 							SourceDelegateFunction = Cast<UFunction>(StaticFindObject(UFunction::StaticClass(), DelegateSignatureOuter, *NameOfDelegateFunction));
 
@@ -3193,14 +3189,14 @@ void FHeaderParser::FixupDelegateProperties( UStruct* Struct, FScope& Scope, TMa
 						// verify that we got a valid string for the class name
 						if ( DelegateClassName.Len() == 0 )
 						{
-							UngetToken(*DelegatePropertyToken);
+							UngetToken(DelegatePropertyToken);
 							FError::Throwf(TEXT("Invalid scope specified in delegate property function reference: '%s'"), *NameOfDelegateFunction);
 						}
 
 						// verify that we got a valid string for the name of the function
 						if ( DelegateName.Len() == 0 )
 						{
-							UngetToken(*DelegatePropertyToken);
+							UngetToken(DelegatePropertyToken);
 							FError::Throwf(TEXT("Invalid delegate name specified in delegate property function reference '%s'"), *NameOfDelegateFunction);
 						}
 
@@ -3215,12 +3211,12 @@ void FHeaderParser::FixupDelegateProperties( UStruct* Struct, FScope& Scope, TMa
 
 					if ( SourceDelegateFunction == NULL )
 					{
-						UngetToken(*DelegatePropertyToken);
+						UngetToken(DelegatePropertyToken);
 						FError::Throwf(TEXT("Failed to find delegate function '%s'"), *NameOfDelegateFunction);
 					}
 					else if ( (SourceDelegateFunction->FunctionFlags&FUNC_Delegate) == 0 )
 					{
-						UngetToken(*DelegatePropertyToken);
+						UngetToken(DelegatePropertyToken);
 						FError::Throwf(TEXT("Only delegate functions can be used as the type for a delegate property; '%s' is not a delegate."), *NameOfDelegateFunction);
 					}
 				}
@@ -3228,14 +3224,14 @@ void FHeaderParser::FixupDelegateProperties( UStruct* Struct, FScope& Scope, TMa
 				// successfully found the delegate function that this delegate property corresponds to
 
 				// save this into the delegate cache for faster lookup later
-				DelegateCache.Add(DelegatePropertyToken->DelegateName, SourceDelegateFunction);
+				DelegateCache.Add(DelegatePropertyToken.DelegateName, SourceDelegateFunction);
 
 				// bind it to the delegate property
 				if( DelegateProperty != NULL )
 				{
 					if( !SourceDelegateFunction->HasAnyFunctionFlags( FUNC_MulticastDelegate ) )
 					{
-						DelegateProperty->SignatureFunction = DelegatePropertyToken->Function = SourceDelegateFunction;
+						DelegateProperty->SignatureFunction = DelegatePropertyToken.Function = SourceDelegateFunction;
 					}
 					else
 					{
@@ -3246,7 +3242,7 @@ void FHeaderParser::FixupDelegateProperties( UStruct* Struct, FScope& Scope, TMa
 				{
 					if( SourceDelegateFunction->HasAnyFunctionFlags( FUNC_MulticastDelegate ) )
 					{
-						MulticastDelegateProperty->SignatureFunction = DelegatePropertyToken->Function = SourceDelegateFunction;
+						MulticastDelegateProperty->SignatureFunction = DelegatePropertyToken.Function = SourceDelegateFunction;
 
 						if(MulticastDelegateProperty->HasAnyPropertyFlags(CPF_BlueprintAssignable | CPF_BlueprintCallable))
 						{
@@ -3526,13 +3522,11 @@ void FHeaderParser::VerifyPropertyMarkups( UClass* TargetClass )
 				return TargetFunc;
 			};
 
-			FClassMetaData* TargetClassData = GScriptHelper.FindClassData(TargetClass);
-			check(TargetClassData);
-			FToken* PropertyToken = TargetClassData->FindTokenData(Prop);
-			check(PropertyToken);
+			FUnrealPropertyDefinitionInfo& PropDef = GTypeDefinitionInfoMap.FindChecked(Prop).AsPropertyChecked();
+			FToken& PropertyToken = PropDef.GetToken();
 
-			TGuardValue<int32> GuardedInputPos(InputPos, PropertyToken->StartPos);
-			TGuardValue<int32> GuardedInputLine(InputLine, PropertyToken->StartLine);
+			TGuardValue<int32> GuardedInputPos(InputPos, PropertyToken.StartPos);
+			TGuardValue<int32> GuardedInputLine(InputLine, PropertyToken.StartLine);
 
 			if (Prop->HasAnyPropertyFlags(CPF_RepNotify))
 			{
@@ -4809,10 +4803,10 @@ void FHeaderParser::GetVarType(
 		bool bHandledType = false;
 		FString IdentifierStripped = GetClassNameWithPrefixRemoved(VarType.Identifier);
 		bool bStripped = false;
-		UScriptStruct* Struct = FindObject<UScriptStruct>( ANY_PACKAGE, VarType.Identifier );
+		UScriptStruct* Struct = FindObject<UScriptStruct>(ANY_PACKAGE, VarType.Identifier);
 		if (!Struct)
 		{
-			Struct = FindObject<UScriptStruct>( ANY_PACKAGE, *IdentifierStripped );
+			Struct = FindObject<UScriptStruct>(ANY_PACKAGE, *IdentifierStripped);
 			bStripped = true;
 		}
 
@@ -4840,10 +4834,13 @@ void FHeaderParser::GetVarType(
 				{
 					TSharedRef<FScope> LocScope = FScope::GetTypeScope(LocalOwnerClass);
 					const FString DelegateIdentifierStripped = GetClassNameWithPrefixRemoved(DelegateName.Identifier);
-					if (UFunction* DelegateFunc = Cast<UFunction>(LocScope->FindTypeByName(*(DelegateIdentifierStripped + HEADER_GENERATED_DELEGATE_SIGNATURE_SUFFIX))))
+					if (FUnrealFieldDefinitionInfo* FoundDef = LocScope->FindTypeByName(*(DelegateIdentifierStripped + HEADER_GENERATED_DELEGATE_SIGNATURE_SUFFIX)))
 					{
-						SetDelegateType(DelegateFunc, DelegateIdentifierStripped);
-						VarProperty.DelegateSignatureOwnerClass = LocalOwnerClass;
+						if (FUnrealFunctionDefinitionInfo* DelegateFuncDef = FoundDef->AsFunction())
+						{
+							SetDelegateType(DelegateFuncDef->GetFunction(), DelegateIdentifierStripped);
+							VarProperty.DelegateSignatureOwnerClass = LocalOwnerClass;
+						}
 					}
 				}
 				else
@@ -4861,243 +4858,255 @@ void FHeaderParser::GetVarType(
 			if (bStripped)
 			{
 				const TCHAR* PrefixCPP = UHTConfig.StructsWithTPrefix.Contains(IdentifierStripped) ? TEXT("T") : Struct->GetPrefixCPP();
-				FString ExpectedStructName = FString::Printf(TEXT("%s%s"), PrefixCPP, *Struct->GetName() );
-				if( FString(VarType.Identifier) != ExpectedStructName )
+				FString ExpectedStructName = FString::Printf(TEXT("%s%s"), PrefixCPP, *Struct->GetName());
+				if (FString(VarType.Identifier) != ExpectedStructName)
 				{
-					FError::Throwf( TEXT("Struct '%s' is missing or has an incorrect prefix, expecting '%s'"), VarType.Identifier, *ExpectedStructName );
+					FError::Throwf(TEXT("Struct '%s' is missing or has an incorrect prefix, expecting '%s'"), VarType.Identifier, *ExpectedStructName);
 				}
 			}
-			else if( !UHTConfig.StructsWithNoPrefix.Contains(VarType.Identifier) )
+			else if (!UHTConfig.StructsWithNoPrefix.Contains(VarType.Identifier))
 			{
 				const TCHAR* PrefixCPP = UHTConfig.StructsWithTPrefix.Contains(VarType.Identifier) ? TEXT("T") : Struct->GetPrefixCPP();
-				FError::Throwf(TEXT("Struct '%s' is missing a prefix, expecting '%s'"), VarType.Identifier, *FString::Printf(TEXT("%s%s"), PrefixCPP, *Struct->GetName()) );
+				FError::Throwf(TEXT("Struct '%s' is missing a prefix, expecting '%s'"), VarType.Identifier, *FString::Printf(TEXT("%s%s"), PrefixCPP, *Struct->GetName()));
 			}
 
 			bHandledType = true;
 
-			VarProperty = FPropertyBase( Struct );
-			if((Struct->StructFlags & STRUCT_HasInstancedReference) && !(Disallow & CPF_ContainsInstancedReference))
+			VarProperty = FPropertyBase(Struct);
+			if ((Struct->StructFlags & STRUCT_HasInstancedReference) && !(Disallow & CPF_ContainsInstancedReference))
 			{
 				Flags |= CPF_ContainsInstancedReference;
 			}
 			// Struct keyword in front of a struct is legal, we 'consume' it
 			bUnconsumedStructKeyword = false;
 		}
-		else if ( FindObject<UScriptStruct>( ANY_PACKAGE, *IdentifierStripped ) != nullptr)
+		else if (FindObject<UScriptStruct>(ANY_PACKAGE, *IdentifierStripped) != nullptr)
 		{
 			bHandledType = true;
 
 			// Struct keyword in front of a struct is legal, we 'consume' it
 			bUnconsumedStructKeyword = false;
 		}
-		else if (UFunction* DelegateFunc = Cast<UFunction>(Scope->FindTypeByName(*(IdentifierStripped + HEADER_GENERATED_DELEGATE_SIGNATURE_SUFFIX))))
-		{
-			SetDelegateType(DelegateFunc, IdentifierStripped);
-		}
 		else
 		{
-			// An object reference of some type (maybe a restricted class?)
-			UClass* TempClass = NULL;
-
-			EPropertyType PropertyType = CPT_ObjectReference;
-			const bool bIsSoftObjectPtrTemplate  = VarType.Matches(TEXT("TSoftObjectPtr"), ESearchCase::CaseSensitive);
-
-			bool bWeakIsAuto = false;
-
-			if (VarType.Matches(TEXT("TSubclassOf"), ESearchCase::CaseSensitive))
-			{
-				TempClass = UClass::StaticClass();
-			}
-			else if (VarType.Matches(TEXT("FScriptInterface"), ESearchCase::CaseSensitive))
+			UFunction* DelegateFunc = nullptr;
+			if (FUnrealFieldDefinitionInfo* FieldDef = Scope->FindTypeByName(*(IdentifierStripped + HEADER_GENERATED_DELEGATE_SIGNATURE_SUFFIX)))
 			{
-				TempClass = UInterface::StaticClass();
-				Flags |= CPF_UObjectWrapper;
+				if (FUnrealFunctionDefinitionInfo* FunctionDef = FieldDef->AsFunction())
+				{
+					DelegateFunc = FunctionDef->GetFunction();
+				}
 			}
-			else if (VarType.Matches(TEXT("TSoftClassPtr"), ESearchCase::CaseSensitive))
+
+			if (DelegateFunc)
 			{
-				TempClass = UClass::StaticClass();
-				PropertyType = CPT_SoftObjectReference;
+				SetDelegateType(DelegateFunc, IdentifierStripped);
 			}
 			else
 			{
-				const bool bIsLazyPtrTemplate        = VarType.Matches(TEXT("TLazyObjectPtr"), ESearchCase::CaseSensitive);
-				const bool bIsObjectPtrTemplate      = VarType.Matches(TEXT("TObjectPtr"), ESearchCase::CaseSensitive);
-				const bool bIsWeakPtrTemplate        = VarType.Matches(TEXT("TWeakObjectPtr"), ESearchCase::CaseSensitive);
-				const bool bIsAutoweakPtrTemplate    = VarType.Matches(TEXT("TAutoWeakObjectPtr"), ESearchCase::CaseSensitive);
-				const bool bIsScriptInterfaceWrapper = VarType.Matches(TEXT("TScriptInterface"), ESearchCase::CaseSensitive);
+				// An object reference of some type (maybe a restricted class?)
+				UClass* TempClass = NULL;
 
-				if (bIsLazyPtrTemplate || bIsObjectPtrTemplate || bIsWeakPtrTemplate || bIsAutoweakPtrTemplate || bIsScriptInterfaceWrapper || bIsSoftObjectPtrTemplate)
-				{
-					RequireSymbol(TEXT('<'), VarType.Identifier);
+				EPropertyType PropertyType = CPT_ObjectReference;
+				const bool bIsSoftObjectPtrTemplate = VarType.Matches(TEXT("TSoftObjectPtr"), ESearchCase::CaseSensitive);
 
-					// Also consume const
-					bNativeConstTemplateArg |= MatchIdentifier(TEXT("const"), ESearchCase::CaseSensitive);
+				bool bWeakIsAuto = false;
 
-					// Consume a forward class declaration 'class' if present
-					MatchIdentifier(TEXT("class"), ESearchCase::CaseSensitive);
-					
-					// Find the lazy/weak class
-					FToken InnerClass;
-					if (GetIdentifier(InnerClass))
+				if (VarType.Matches(TEXT("TSubclassOf"), ESearchCase::CaseSensitive))
+				{
+					TempClass = UClass::StaticClass();
+				}
+				else if (VarType.Matches(TEXT("FScriptInterface"), ESearchCase::CaseSensitive))
+				{
+					TempClass = UInterface::StaticClass();
+					Flags |= CPF_UObjectWrapper;
+				}
+				else if (VarType.Matches(TEXT("TSoftClassPtr"), ESearchCase::CaseSensitive))
+				{
+					TempClass = UClass::StaticClass();
+					PropertyType = CPT_SoftObjectReference;
+				}
+				else
+				{
+					const bool bIsLazyPtrTemplate = VarType.Matches(TEXT("TLazyObjectPtr"), ESearchCase::CaseSensitive);
+					const bool bIsObjectPtrTemplate = VarType.Matches(TEXT("TObjectPtr"), ESearchCase::CaseSensitive);
+					const bool bIsWeakPtrTemplate = VarType.Matches(TEXT("TWeakObjectPtr"), ESearchCase::CaseSensitive);
+					const bool bIsAutoweakPtrTemplate = VarType.Matches(TEXT("TAutoWeakObjectPtr"), ESearchCase::CaseSensitive);
+					const bool bIsScriptInterfaceWrapper = VarType.Matches(TEXT("TScriptInterface"), ESearchCase::CaseSensitive);
+
+					if (bIsLazyPtrTemplate || bIsObjectPtrTemplate || bIsWeakPtrTemplate || bIsAutoweakPtrTemplate || bIsScriptInterfaceWrapper || bIsSoftObjectPtrTemplate)
 					{
-						RedirectTypeIdentifier(InnerClass);
+						RequireSymbol(TEXT('<'), VarType.Identifier);
 
-						TempClass = FClasses::FindScriptClass(InnerClass.Identifier);
-						if (TempClass == nullptr)
-						{
-							FError::Throwf(TEXT("Unrecognized type '%s' (in expression %s<%s>) - type must be a UCLASS"), InnerClass.Identifier, VarType.Identifier, InnerClass.Identifier);
-						}
+						// Also consume const
+						bNativeConstTemplateArg |= MatchIdentifier(TEXT("const"), ESearchCase::CaseSensitive);
 
-						if (bIsAutoweakPtrTemplate)
-						{
-							PropertyType = CPT_WeakObjectReference;
-							bWeakIsAuto = true;
-						}
-						else if (bIsLazyPtrTemplate)
-						{
-							PropertyType = CPT_LazyObjectReference;
-						}
-						else if (bIsObjectPtrTemplate)
-						{
-							PropertyType = CPT_ObjectPtrReference;
-						}
-						else if (bIsWeakPtrTemplate)
+						// Consume a forward class declaration 'class' if present
+						MatchIdentifier(TEXT("class"), ESearchCase::CaseSensitive);
+
+						// Find the lazy/weak class
+						FToken InnerClass;
+						if (GetIdentifier(InnerClass))
 						{
-							PropertyType = CPT_WeakObjectReference;
+							RedirectTypeIdentifier(InnerClass);
+
+							TempClass = FClasses::FindScriptClass(InnerClass.Identifier);
+							if (TempClass == nullptr)
+							{
+								FError::Throwf(TEXT("Unrecognized type '%s' (in expression %s<%s>) - type must be a UCLASS"), InnerClass.Identifier, VarType.Identifier, InnerClass.Identifier);
+							}
+
+							if (bIsAutoweakPtrTemplate)
+							{
+								PropertyType = CPT_WeakObjectReference;
+								bWeakIsAuto = true;
+							}
+							else if (bIsLazyPtrTemplate)
+							{
+								PropertyType = CPT_LazyObjectReference;
+							}
+							else if (bIsObjectPtrTemplate)
+							{
+								PropertyType = CPT_ObjectPtrReference;
+							}
+							else if (bIsWeakPtrTemplate)
+							{
+								PropertyType = CPT_WeakObjectReference;
+							}
+							else if (bIsSoftObjectPtrTemplate)
+							{
+								PropertyType = CPT_SoftObjectReference;
+							}
+
+							Flags |= CPF_UObjectWrapper;
 						}
-						else if (bIsSoftObjectPtrTemplate)
+						else
 						{
-							PropertyType = CPT_SoftObjectReference;
+							FError::Throwf(TEXT("%s: Missing template type"), VarType.Identifier);
 						}
 
-						Flags |= CPF_UObjectWrapper;
+						// Const after template argument type but before end of template symbol
+						bNativeConstTemplateArg |= MatchIdentifier(TEXT("const"), ESearchCase::CaseSensitive);
+
+						RequireSymbol(TEXT('>'), VarType.Identifier, ESymbolParseOption::CloseTemplateBracket);
 					}
 					else
 					{
-						FError::Throwf(TEXT("%s: Missing template type"), VarType.Identifier);
+						TempClass = FClasses::FindScriptClass(VarType.Identifier);
 					}
-
-					// Const after template argument type but before end of template symbol
-					bNativeConstTemplateArg |= MatchIdentifier(TEXT("const"), ESearchCase::CaseSensitive);
-
-					RequireSymbol(TEXT('>'), VarType.Identifier, ESymbolParseOption::CloseTemplateBracket);
 				}
-				else
-				{
-					TempClass = FClasses::FindScriptClass(VarType.Identifier);
-				}
-			}
 
-			if (TempClass != NULL)
-			{
-				bHandledType = true;
-
-				if ((PropertyType == CPT_WeakObjectReference) && (Disallow & CPF_AutoWeak)) // if it is not allowing anything, force it strong. this is probably a function arg
+				if (TempClass != NULL)
 				{
-					PropertyType = CPT_ObjectReference;
-				}
+					bHandledType = true;
 
-				// if this is an interface class, we use the FInterfaceProperty class instead of FObjectProperty
-				if ((PropertyType == CPT_ObjectReference) && TempClass->HasAnyClassFlags(CLASS_Interface))
-				{
-					PropertyType = CPT_Interface;
-				}
+					if ((PropertyType == CPT_WeakObjectReference) && (Disallow & CPF_AutoWeak)) // if it is not allowing anything, force it strong. this is probably a function arg
+					{
+						PropertyType = CPT_ObjectReference;
+					}
 
-				VarProperty = FPropertyBase(TempClass, PropertyType, bWeakIsAuto);
-				if (TempClass->IsChildOf(UClass::StaticClass()))
-				{
-					if ( MatchSymbol(TEXT('<')) )
+					// if this is an interface class, we use the FInterfaceProperty class instead of FObjectProperty
+					if ((PropertyType == CPT_ObjectReference) && TempClass->HasAnyClassFlags(CLASS_Interface))
 					{
-						Flags |= CPF_UObjectWrapper;
+						PropertyType = CPT_Interface;
+					}
 
-						// Consume a forward class declaration 'class' if present
-						MatchIdentifier(TEXT("class"), ESearchCase::CaseSensitive);
+					VarProperty = FPropertyBase(TempClass, PropertyType, bWeakIsAuto);
+					if (TempClass->IsChildOf(UClass::StaticClass()))
+					{
+						if (MatchSymbol(TEXT('<')))
+						{
+							Flags |= CPF_UObjectWrapper;
+
+							// Consume a forward class declaration 'class' if present
+							MatchIdentifier(TEXT("class"), ESearchCase::CaseSensitive);
+
+							// Get the actual class type to restrict this to
+							FToken Limitor;
+							if (!GetIdentifier(Limitor))
+							{
+								FError::Throwf(TEXT("'class': Missing class limitor"));
+							}
+
+							RedirectTypeIdentifier(Limitor);
+
+							VarProperty.MetaClass = FClasses::FindScriptClassOrThrow(Limitor.Identifier);
 
-						// Get the actual class type to restrict this to
-						FToken Limitor;
-						if( !GetIdentifier(Limitor) )
+							RequireSymbol(TEXT('>'), TEXT("'class limitor'"), ESymbolParseOption::CloseTemplateBracket);
+						}
+						else
 						{
-							FError::Throwf(TEXT("'class': Missing class limitor"));
+							VarProperty.MetaClass = UObject::StaticClass();
 						}
 
-						RedirectTypeIdentifier(Limitor);
-
-						VarProperty.MetaClass = FClasses::FindScriptClassOrThrow(Limitor.Identifier);
+						if (PropertyType == CPT_WeakObjectReference)
+						{
+							FError::Throwf(TEXT("Class variables cannot be weak, they are always strong."));
+						}
+						else if (PropertyType == CPT_LazyObjectReference)
+						{
+							FError::Throwf(TEXT("Class variables cannot be lazy, they are always strong."));
+						}
 
-						RequireSymbol( TEXT('>'), TEXT("'class limitor'"), ESymbolParseOption::CloseTemplateBracket );
-					}
-					else
-					{
-						VarProperty.MetaClass = UObject::StaticClass();
+						if (bIsSoftObjectPtrTemplate)
+						{
+							FError::Throwf(TEXT("Class variables cannot be stored in TSoftObjectPtr, use TSoftClassPtr instead."));
+						}
 					}
 
-					if (PropertyType == CPT_WeakObjectReference)
-					{
-						FError::Throwf(TEXT("Class variables cannot be weak, they are always strong."));
-					}
-					else if (PropertyType == CPT_LazyObjectReference)
+					// Inherit instancing flags
+					if (DoesAnythingInHierarchyHaveDefaultToInstanced(TempClass))
 					{
-						FError::Throwf(TEXT("Class variables cannot be lazy, they are always strong."));
+						Flags |= ((CPF_InstancedReference | CPF_ExportObject) & (~Disallow));
 					}
 
-					if (bIsSoftObjectPtrTemplate)
+					// Eat the star that indicates this is a pointer to the UObject
+					if (!(Flags & CPF_UObjectWrapper))
 					{
-						FError::Throwf(TEXT("Class variables cannot be stored in TSoftObjectPtr, use TSoftClassPtr instead."));
-					}
-				}
+						// Const after variable type but before pointer symbol
+						bNativeConst |= MatchIdentifier(TEXT("const"), ESearchCase::CaseSensitive);
 
-				// Inherit instancing flags
-				if (DoesAnythingInHierarchyHaveDefaultToInstanced(TempClass))
-				{
-					Flags |= ((CPF_InstancedReference|CPF_ExportObject) & (~Disallow)); 
-				}
+						RequireSymbol(TEXT('*'), TEXT("Expected a pointer type"));
 
-				// Eat the star that indicates this is a pointer to the UObject
-				if (!(Flags & CPF_UObjectWrapper))
-				{
-					// Const after variable type but before pointer symbol
-					bNativeConst |= MatchIdentifier(TEXT("const"), ESearchCase::CaseSensitive);
+						// Optionally emit messages about native pointer members and swallow trailing 'const' after pointer properties
+						if (VariableCategory == EVariableCategory::Member)
+						{
+							ConditionalLogPointerUsage(UHTConfig.NativePointerMemberBehavior, TEXT("Native pointer"), FString(InputPos - VarStartPos, Input + VarStartPos).TrimStartAndEnd().ReplaceCharWithEscapedChar());
 
-					RequireSymbol(TEXT('*'), TEXT("Expected a pointer type"));
+							MatchIdentifier(TEXT("const"), ESearchCase::CaseSensitive);
+						}
 
-					// Optionally emit messages about native pointer members and swallow trailing 'const' after pointer properties
-					if (VariableCategory == EVariableCategory::Member)
+						VarProperty.PointerType = EPointerType::Native;
+					}
+					else if ((PropertyType == CPT_ObjectPtrReference) && (VariableCategory == EVariableCategory::Member))
 					{
-						ConditionalLogPointerUsage(UHTConfig.NativePointerMemberBehavior, TEXT("Native pointer"), FString(InputPos - VarStartPos, Input + VarStartPos).TrimStartAndEnd().ReplaceCharWithEscapedChar());
+						ConditionalLogPointerUsage(UHTConfig.ObjectPtrMemberBehavior, TEXT("ObjectPtr"), FString(InputPos - VarStartPos, Input + VarStartPos).TrimStartAndEnd().ReplaceCharWithEscapedChar());
+					}
 
-						MatchIdentifier(TEXT("const"), ESearchCase::CaseSensitive);
+					// Imply const if it's a parameter that is a pointer to a const class
+					if (VariableCategory != EVariableCategory::Member && (TempClass != NULL) && (TempClass->HasAnyClassFlags(CLASS_Const)))
+					{
+						Flags |= CPF_ConstParm;
 					}
 
-					VarProperty.PointerType = EPointerType::Native;
+					// Class keyword in front of a class is legal, we 'consume' it
+					bUnconsumedClassKeyword = false;
+					bUnconsumedConstKeyword = false;
 				}
-				else if ((PropertyType == CPT_ObjectPtrReference) && (VariableCategory == EVariableCategory::Member))
+			}
+
+			// Resolve delegates declared in another class  //@TODO: UCREMOVAL: This seems extreme
+			if (!bHandledType)
+			{
+				if (UFunction* ExternDelegateFunc = (UFunction*)StaticFindObject(UFunction::StaticClass(), ANY_PACKAGE, *(IdentifierStripped + HEADER_GENERATED_DELEGATE_SIGNATURE_SUFFIX)))
 				{
-					ConditionalLogPointerUsage(UHTConfig.ObjectPtrMemberBehavior, TEXT("ObjectPtr"), FString(InputPos - VarStartPos, Input + VarStartPos).TrimStartAndEnd().ReplaceCharWithEscapedChar());
+					SetDelegateType(ExternDelegateFunc, IdentifierStripped);
 				}
 
-				// Imply const if it's a parameter that is a pointer to a const class
-				if (VariableCategory != EVariableCategory::Member && (TempClass != NULL) && (TempClass->HasAnyClassFlags(CLASS_Const)))
+				if (!bHandledType)
 				{
-					Flags |= CPF_ConstParm;
+					FError::Throwf(TEXT("Unrecognized type '%s' - type must be a UCLASS, USTRUCT or UENUM"), VarType.Identifier);
 				}
-
-				// Class keyword in front of a class is legal, we 'consume' it
-				bUnconsumedClassKeyword = false;
-				bUnconsumedConstKeyword = false;
-			}
-		}
-
-		// Resolve delegates declared in another class  //@TODO: UCREMOVAL: This seems extreme
-		if (!bHandledType)
-		{
-			if (UFunction* DelegateFunc = (UFunction*)StaticFindObject(UFunction::StaticClass(), ANY_PACKAGE, *(IdentifierStripped + HEADER_GENERATED_DELEGATE_SIGNATURE_SUFFIX)))
-			{
-				SetDelegateType(DelegateFunc, IdentifierStripped);
-			}
-
-			if (!bHandledType)
-			{
-				FError::Throwf(TEXT("Unrecognized type '%s' - type must be a UCLASS, USTRUCT or UENUM"), VarType.Identifier );
 			}
 		}
 	}
@@ -5296,12 +5305,14 @@ void FHeaderParser::GetVarType(
 
 FProperty* FHeaderParser::GetVarNameAndDim
 (
-	UStruct*                Scope,
+	FUnrealStructDefinitionInfo& ParentStruct, 
 	FToken&                 VarProperty,
 	EVariableCategory::Type VariableCategory,
 	ELayoutMacroType        LayoutMacroType
 )
 {
+	UStruct* Scope = ParentStruct.GetStruct();
+
 	check(Scope);
 
 	EObjectFlags ObjectFlags = RF_Public;
@@ -5312,7 +5323,7 @@ FProperty* FHeaderParser::GetVarNameAndDim
 
 	const TCHAR* HintText = GetHintText(VariableCategory);
 
-	AddModuleRelativePathToMetadata(Scope, VarProperty.MetaData);
+	AddModuleRelativePathToMetadata(ParentStruct, VarProperty.MetaData);
 
 	// Get variable name.
 	if (VariableCategory == EVariableCategory::Return)
@@ -5679,9 +5690,11 @@ FProperty* FHeaderParser::GetVarNameAndDim
 	VarProperty.TokenProperty = Result;
 	VarProperty.StartLine = InputLine;
 	VarProperty.StartPos = InputPos;
-	FClassMetaData* ScopeData = GScriptHelper.FindClassData(Scope);
-	check(ScopeData);
-	ScopeData->AddProperty(FToken(VarProperty));
+
+	FUnrealPropertyDefinitionInfo& ResultDef = GTypeDefinitionInfoMap.FindChecked(Result).AsPropertyChecked();
+	ResultDef.GetToken() = VarProperty;
+	ParentStruct.AddProperty(ResultDef);
+	ParentStruct.GetStructMetaData().AddProperty(ResultDef);
 
 	// if we had any metadata, add it to the class
 	AddMetaDataToClassData(VarProperty.TokenProperty, TMap<FName,FString>(VarProperty.MetaData));
@@ -5741,24 +5754,24 @@ bool FHeaderParser::CompileDeclaration(TArray<UDelegateFunction*>& DelegatesToFi
 		CompileVersionDeclaration(GetCurrentClass());
 		RequireSymbol(TEXT(')'), Token.Identifier);
 
-		FClassMetaData* ClassData = GetCurrentClassData();
-		if (!ClassData || ClassData->ParsedInterface == EParsedInterface::NotAnInterface)
+		FStructMetaData& StructData = GetCurrentClassData();
+		if (StructData.ParsedInterface == EParsedInterface::NotAnInterface)
 		{
 			FString CurrentClassName = GetCurrentClass()->GetName();
 			FError::Throwf(TEXT("Could not find the associated 'U%s' class while parsing 'I%s' - it could be missing or malformed"), *CurrentClassName, *CurrentClassName);
 		}
 
-		if (ClassData->ParsedInterface == EParsedInterface::ParsedIInterface)
+		if (StructData.ParsedInterface == EParsedInterface::ParsedIInterface)
 		{
 			FString CurrentClassName = GetCurrentClass()->GetName();
 			FError::Throwf(TEXT("Duplicate IInterface definition found while parsing 'I%s'"), *CurrentClassName);
 		}
 
-		check(ClassData->ParsedInterface == EParsedInterface::ParsedUInterface);
+		check(StructData.ParsedInterface == EParsedInterface::ParsedUInterface);
 
-		ClassData->ParsedInterface = EParsedInterface::ParsedIInterface;
-		ClassData->GeneratedBodyMacroAccessSpecifier = CurrentAccessSpecifier;
-		ClassData->SetInterfaceGeneratedBodyLine(InputLine);
+		StructData.ParsedInterface = EParsedInterface::ParsedIInterface;
+		StructData.GeneratedBodyMacroAccessSpecifier = CurrentAccessSpecifier;
+		StructData.SetInterfaceGeneratedBodyLine(InputLine);
 
 		bClassHasGeneratedIInterfaceBody = true;
 
@@ -5784,10 +5797,10 @@ bool FHeaderParser::CompileDeclaration(TArray<UDelegateFunction*>& DelegatesToFi
 		CompileVersionDeclaration(GetCurrentClass());
 		RequireSymbol(TEXT(')'), Token.Identifier);
 
-		FClassMetaData* ClassData = GetCurrentClassData();
+		FStructMetaData& StructData = GetCurrentClassData();
 
-		ClassData->GeneratedBodyMacroAccessSpecifier = CurrentAccessSpecifier;
-		ClassData->SetGeneratedBodyLine(InputLine);
+		StructData.GeneratedBodyMacroAccessSpecifier = CurrentAccessSpecifier;
+		StructData.SetGeneratedBodyLine(InputLine);
 
 		bClassHasGeneratedUInterfaceBody = true;
 
@@ -5805,7 +5818,7 @@ bool FHeaderParser::CompileDeclaration(TArray<UDelegateFunction*>& DelegatesToFi
 			FError::Throwf(TEXT("%s must occur inside the class definition"), Token.Identifier);
 		}
 
-		FClassMetaData* ClassData = GetCurrentClassData();
+		FStructMetaData& StructData = GetCurrentClassData();
 
 		if (Token.Matches(TEXT("GENERATED_BODY"), ESearchCase::CaseSensitive))
 		{
@@ -5816,7 +5829,7 @@ bool FHeaderParser::CompileDeclaration(TArray<UDelegateFunction*>& DelegatesToFi
 
 			ClassDefinitionRanges[GetCurrentClass()].bHasGeneratedBody = true;
 
-			ClassData->GeneratedBodyMacroAccessSpecifier = CurrentAccessSpecifier;
+			StructData.GeneratedBodyMacroAccessSpecifier = CurrentAccessSpecifier;
 		}
 		else
 		{
@@ -5827,7 +5840,7 @@ bool FHeaderParser::CompileDeclaration(TArray<UDelegateFunction*>& DelegatesToFi
 		CompileVersionDeclaration(GetCurrentClass());
 		RequireSymbol(TEXT(')'), Token.Identifier);
 
-		ClassData->SetGeneratedBodyLine(InputLine);
+		StructData.SetGeneratedBodyLine(InputLine);
 
 		bClassHasGeneratedBody = true;
 		return true;
@@ -5874,7 +5887,7 @@ bool FHeaderParser::CompileDeclaration(TArray<UDelegateFunction*>& DelegatesToFi
 		CheckAllow(TEXT("'Member variable declaration'"), ENestAllowFlags::VarDecl);
 		check(TopNest->NestType == ENestType::Class);
 
-		CompileVariableDeclaration(GetCurrentClass());
+		CompileVariableDeclaration(GetCurrentClassDef());
 		return true;
 	}
 
@@ -5913,13 +5926,14 @@ bool FHeaderParser::CompileDeclaration(TArray<UDelegateFunction*>& DelegatesToFi
 		bEncounteredNewStyleClass_UnmatchedBrackets = false;
 
 		UClass* CurrentClass = GetCurrentClass();
+		FUnrealClassDefinitionInfo& CurrentClassDef = GetCurrentClassDef().AsClassChecked();
 
 		// Pop nesting here to allow other non UClass declarations in the header file.
 		if (CurrentClass->ClassFlags & CLASS_Interface)
 		{
 			checkf(TopNest->NestType == ENestType::Interface || TopNest->NestType == ENestType::NativeInterface, TEXT("Unexpected end of interface block."));
 			PopNest(TopNest->NestType, TEXT("'Interface'"));
-			PostPopNestInterface(CurrentClass);
+			PostPopNestInterface(CurrentClassDef);
 
 			// Ensure the UINTERFACE classes have a GENERATED_BODY declaration
 			if (bHaveSeenUClass && !bClassHasGeneratedUInterfaceBody)
@@ -5936,7 +5950,7 @@ bool FHeaderParser::CompileDeclaration(TArray<UDelegateFunction*>& DelegatesToFi
 		else
 		{
 			PopNest(ENestType::Class, TEXT("'Class'"));
-			PostPopNestClass(CurrentClass);
+			PostPopNestClass(CurrentClassDef);
 
 			// Ensure classes have a GENERATED_BODY declaration
 			if (bHaveSeenUClass && !bClassHasGeneratedBody)
@@ -5950,7 +5964,7 @@ bool FHeaderParser::CompileDeclaration(TArray<UDelegateFunction*>& DelegatesToFi
 		bClassHasGeneratedUInterfaceBody = false;
 		bClassHasGeneratedIInterfaceBody = false;
 
-		GetCurrentScope()->AddType(CurrentClass);
+		GetCurrentScope()->AddType(CurrentClassDef);
 		return true;
 	}
 
@@ -6098,9 +6112,7 @@ bool FHeaderParser::CompileDeclaration(TArray<UDelegateFunction*>& DelegatesToFi
 						{
 							FString EnclosingDefine = CompilerDirectiveStack.Num() > 0 ? TEXT("WITH_EDITORONLY_DATA") : TEXT("");
 
-							UClass* CurrentClass = GetCurrentClass();
-							FUnrealClassDefinitionInfo& ClassDef = GTypeDefinitionInfoMap.FindChecked(CurrentClass).AsClassChecked();
-
+							FUnrealClassDefinitionInfo& ClassDef = GetCurrentClassDef().AsClassChecked();
 							ClassDef.AddArchiveType(ArchiveType);
 							ClassDef.SetEnclosingDefine(MoveTemp(EnclosingDefine));
 						}
@@ -6264,8 +6276,7 @@ FClass* FHeaderParser::ParseClassNameDeclaration(FString& DeclaredClassName, FSt
 
 	FClass* FoundClass = FClasses::FindClass(*GetClassNameWithPrefixRemoved(*DeclaredClassName));
 	check(FoundClass);
-
-	FClassMetaData* ClassMetaData = GScriptHelper.AddClassData(FoundClass, &SourceFile);
+	FStructMetaData& StructMetaData = GTypeDefinitionInfoMap.FindChecked(FoundClass).AsStructChecked().GetStructMetaData();
 
 	// Get parent class.
 	bool bSpecifiesParentClass = false;
@@ -6363,7 +6374,7 @@ FClass* FHeaderParser::ParseClassNameDeclaration(FString& DeclaredClassName, FSt
 				break;
 			}
 
-			HandleOneInheritedClass(FoundClass, MoveTemp(InterfaceName));
+			HandleOneInheritedClass(StructMetaData, FoundClass, MoveTemp(InterfaceName));
 		}
 	}
 	else if (FoundClass->GetSuperClass())
@@ -6374,8 +6385,9 @@ FClass* FHeaderParser::ParseClassNameDeclaration(FString& DeclaredClassName, FSt
 	return FoundClass;
 }
 
-void FHeaderParser::HandleOneInheritedClass(UClass* Class, FString&& InterfaceName)
+void FHeaderParser::HandleOneInheritedClass(FStructMetaData& StructMetaData, UClass* Class, FString&& InterfaceName)
 {
+
 	// Check for UInterface derived interface inheritance
 	if (UClass* Interface = FClasses::FindScriptClass(InterfaceName))
 	{
@@ -6391,17 +6403,13 @@ void FHeaderParser::HandleOneInheritedClass(UClass* Class, FString&& InterfaceNa
 		new (Class->Interfaces) FImplementedInterface(Interface, 0, false);
 		if (Interface->HasAnyClassFlags(CLASS_Native))
 		{
-			FClassMetaData* ClassData = GScriptHelper.FindClassData(Class);
-			check(ClassData);
-			ClassData->AddInheritanceParent(Interface, &SourceFile);
+			StructMetaData.AddInheritanceParent(Interface, &SourceFile);
 		}
 	}
 	else
 	{
 		// Non-UObject inheritance
-		FClassMetaData* ClassData = GScriptHelper.FindClassData(Class);
-		check(ClassData);
-		ClassData->AddInheritanceParent(MoveTemp(InterfaceName), &SourceFile);
+		StructMetaData.AddInheritanceParent(MoveTemp(InterfaceName), &SourceFile);
 	}
 }
 
@@ -6517,9 +6525,8 @@ UClass* FHeaderParser::CompileClassDeclaration()
 	ClassDeclarationData.MergeClassCategories(Class);
 
 	// Class attributes.
-	FClassMetaData* ClassData = GScriptHelper.FindClassData(Class);
-	check(ClassData);
-	ClassData->SetPrologLine(PrologFinishLine);
+	FStructMetaData& StructData = GTypeDefinitionInfoMap.FindChecked(Class).AsStructChecked().GetStructMetaData();
+	StructData.SetPrologLine(PrologFinishLine);
 
 	ClassDeclarationData.MergeAndValidateClassFlags(DeclaredClassName, PrevClassFlags, Class);
 	Class->SetInternalFlags(EInternalObjectFlags::Native);
@@ -6535,7 +6542,7 @@ UClass* FHeaderParser::CompileClassDeclaration()
 	if (ClassDeclarationData.AutoExpandCategories.Num()) { MetaData.Add(NAME_AutoExpandCategories, FString::Join(ClassDeclarationData.AutoExpandCategories, TEXT(" "))); }
 
 	AddIncludePathToMetadata(Class, MetaData);
-	AddModuleRelativePathToMetadata(Class, MetaData);
+	AddModuleRelativePathToMetadata(ClassDeclarationData, MetaData);
 
 	// Register the metadata
 	AddMetaDataToClassData(Class, MoveTemp(MetaData));
@@ -6556,7 +6563,7 @@ UClass* FHeaderParser::CompileClassDeclaration()
 
 	// auto-create properties for all of the VFTables needed for the multiple inheritances
 	// get the inheritance parents
-	const TArray<FMultipleInheritanceBaseClass*>& InheritanceParents = ClassData->GetInheritanceParents();
+	const TArray<FMultipleInheritanceBaseClass*>& InheritanceParents = StructData.GetInheritanceParents();
 
 	// for all base class types, make a VfTable property
 	for (int32 ParentIndex = InheritanceParents.Num() - 1; ParentIndex >= 0; ParentIndex--)
@@ -6690,6 +6697,7 @@ void FHeaderParser::CompileInterfaceDeclaration()
 	// New style files have the interface name / extends afterwards
 	RequireIdentifier(TEXT("class"), ESearchCase::CaseSensitive, TEXT("Interface declaration"));
 	FClass* InterfaceClass = ParseInterfaceNameDeclaration(/*out*/ DeclaredInterfaceName, /*out*/ RequiredAPIMacroIfPresent);
+	FUnrealClassDefinitionInfo& ClassDef = GTypeDefinitionInfoMap.FindChecked(InterfaceClass).AsClassChecked();
 	ClassDefinitionRanges.Add(InterfaceClass, ClassDefinitionRange(&Input[InputPos], nullptr));
 
 	// Record that this interface is RequiredAPI if the CORE_API style macro was present
@@ -6753,13 +6761,11 @@ void FHeaderParser::CompileInterfaceDeclaration()
 	}
 
 	// Try parsing metadata for the interface
-	FClassMetaData* ClassData = GScriptHelper.AddInterfaceClassData(InterfaceClass, &SourceFile);
-	check(ClassData);
-
-	ClassData->SetPrologLine(PrologFinishLine);
+	FStructMetaData& StructData = GTypeDefinitionInfoMap.FindChecked(InterfaceClass).AsStructChecked().GetStructMetaData();
+	StructData.SetPrologLine(PrologFinishLine);
 
 	// Register the metadata
-	AddModuleRelativePathToMetadata(InterfaceClass, MetaData);
+	AddModuleRelativePathToMetadata(ClassDef, MetaData);
 	AddMetaDataToClassData(InterfaceClass, MoveTemp(MetaData));
 
 	// Handle the start of the rest of the interface
@@ -7011,8 +7017,10 @@ bool FHeaderParser::IsValidDelegateDeclaration(const FToken& Token) const
 }
 
 // Parse the parameter list of a function or delegate declaration
-void FHeaderParser::ParseParameterList(UFunction* Function, bool bExpectCommaBeforeName, TMap<FName, FString>* MetaData)
+void FHeaderParser::ParseParameterList(FUnrealFunctionDefinitionInfo& FunctionDef, bool bExpectCommaBeforeName, TMap<FName, FString>* MetaData)
 {
+	UFunction* Function = FunctionDef.GetFunction();
+
 	// Get parameter list.
 	if (MatchSymbol(TEXT(')')))
 	{
@@ -7033,7 +7041,7 @@ void FHeaderParser::ParseParameterList(UFunction* Function, bool bExpectCommaBef
 			RequireSymbol(TEXT(','), TEXT("Delegate definitions require a , between the parameter type and parameter name"));
 		}
 
-		FProperty* Prop = GetVarNameAndDim(Function, Property, VariableCategory);
+		FProperty* Prop = GetVarNameAndDim(FunctionDef, Property, VariableCategory);
 
 		Function->NumParms++;
 
@@ -7316,19 +7324,15 @@ UDelegateFunction* FHeaderParser::CompileDelegateDeclaration(const TCHAR* Delega
 
 	UDelegateFunction* DelegateSignatureFunction = (bIsSparse ? CreateDelegateFunction<USparseDelegateFunction>(FuncInfo) : CreateDelegateFunction<UDelegateFunction>(FuncInfo));
 	
-	FClassMetaData* ClassMetaData = GScriptHelper.AddClassData(DelegateSignatureFunction, &SourceFile);
-
 	DelegateSignatureFunction->FunctionFlags |= FuncInfo.FunctionFlags;
 
 	FuncInfo.FunctionReference = DelegateSignatureFunction;
 	FuncInfo.SetFunctionNames();
 	if (FuncInfo.FunctionReference->HasAnyFunctionFlags(FUNC_Delegate) && !GetCurrentScope()->IsFileScope())
 	{
-		GetCurrentClassData()->MarkContainsDelegate();
+		GetCurrentClassData().MarkContainsDelegate();
 	}
 
-	GetCurrentScope()->AddType(DelegateSignatureFunction);
-
 	// determine whether this function should be 'const'
 	if (bDeclaredConst)
 	{
@@ -7366,13 +7370,14 @@ UDelegateFunction* FHeaderParser::CompileDelegateDeclaration(const TCHAR* Delega
 	TSharedRef<FUnrealFunctionDefinitionInfo> FuncDef = MakeShareable(new FUnrealFunctionDefinitionInfo(SourceFile, InputLine, FString(FuncInfoFunctionIdentifier), MoveTemp(FuncInfo)));
 	FuncDef->SetObject(FuncDef->GetFunctionData().FunctionReference);
 	GTypeDefinitionInfoMap.Add(FuncDef->GetFunctionData().FunctionReference, FuncDef);
+	GetCurrentScope()->AddType(*FuncDef);
 
 	// Get parameter list.
 	if (FoundParamCount)
 	{
 		RequireSymbol(TEXT(','), CurrentScopeName);
 
-		ParseParameterList(DelegateSignatureFunction, /*bExpectCommaBeforeName=*/ true);
+		ParseParameterList(*FuncDef, /*bExpectCommaBeforeName=*/ true);
 
 		// Check the expected versus actual number of parameters
 		int32 ParamCount = UE_PTRDIFF_TO_INT32(FoundParamCount - UHTConfig.DelegateParameterCountStrings.GetData()) + 1;
@@ -7394,7 +7399,7 @@ UDelegateFunction* FHeaderParser::CompileDelegateDeclaration(const TCHAR* Delega
 	if (bHasReturnValue)
 	{
 		ReturnType.PropertyFlags |= CPF_Parm | CPF_OutParm | CPF_ReturnParm;
-		FProperty* ReturnProp = GetVarNameAndDim(DelegateSignatureFunction, ReturnType, EVariableCategory::Return);
+		FProperty* ReturnProp = GetVarNameAndDim(*FuncDef, ReturnType, EVariableCategory::Return);
 
 		DelegateSignatureFunction->NumParms++;
 	}
@@ -7416,10 +7421,11 @@ UDelegateFunction* FHeaderParser::CompileDelegateDeclaration(const TCHAR* Delega
 	PostPopFunctionDeclaration(DelegateSignatureFunction);
 
 	// Don't allow delegate signatures to be redefined.
-	auto FunctionIterator = GetCurrentScope()->GetTypeIterator<UFunction>();
+	auto FunctionIterator = GetCurrentScope()->GetTypeIterator<FUnrealFunctionDefinitionInfo>();
 	while (FunctionIterator.MoveNext())
 	{
-		UFunction* TestFunc = *FunctionIterator;
+		FUnrealFunctionDefinitionInfo* TestFuncDef = *FunctionIterator;
+		UFunction* TestFunc = TestFuncDef->GetFunction();
 		if ((TestFunc->GetFName() == DelegateSignatureFunction->GetFName()) && (TestFunc != DelegateSignatureFunction))
 		{
 			FError::Throwf(TEXT("Can't override delegate signature function '%s'"), FuncInfoFunctionIdentifier);
@@ -7517,9 +7523,6 @@ void FHeaderParser::CompileFunctionDeclaration()
 	TArray<FPropertySpecifier> SpecifiersFound;
 	ReadSpecifierSetInsideMacro(SpecifiersFound, TEXT("Function"), MetaData);
 
-	FScriptLocation FuncNameRetry;
-	InitScriptLocation(FuncNameRetry);
-
 	if (!GetCurrentClass()->HasAnyClassFlags(CLASS_Native))
 	{
 		FError::Throwf(TEXT("Should only be here for native classes!"));
@@ -7816,31 +7819,28 @@ void FHeaderParser::CompileFunctionDeclaration()
 
 	UFunction* TopFunction = CreateFunction(FuncInfo);
 
-	FClassMetaData* ClassMetaData = GScriptHelper.AddClassData(TopFunction, &SourceFile);
-
 	TopFunction->FunctionFlags |= FuncInfo.FunctionFlags;
 
 	FuncInfo.FunctionReference = TopFunction;
 	FuncInfo.SetFunctionNames();
 
-	GetCurrentScope()->AddType(TopFunction);
-
 	TSharedRef<FUnrealFunctionDefinitionInfo> FuncDef = MakeShareable(new FUnrealFunctionDefinitionInfo(SourceFile, InputLine, FString(FuncInfo.Function.Identifier), FFuncInfo(FuncInfo)));
 	FuncDef->SetObject(FuncDef->GetFunctionData().FunctionReference);
 	GTypeDefinitionInfoMap.Add(FuncDef->GetFunctionData().FunctionReference, FuncDef);
+	GetCurrentScope()->AddType(*FuncDef);
 	if (FuncInfo.FunctionReference->HasAnyFunctionFlags(FUNC_Delegate))
 	{
-		GetCurrentClassData()->MarkContainsDelegate();
+		GetCurrentClassData().MarkContainsDelegate();
 	}
 
 	// Get parameter list.
-	ParseParameterList(TopFunction, false, &MetaData);
+	ParseParameterList(*FuncDef, false, &MetaData);
 
 	// Get return type, if any.
 	if (bHasReturnValue)
 	{
 		ReturnType.PropertyFlags |= CPF_Parm | CPF_OutParm | CPF_ReturnParm;
-		FProperty* ReturnProp = GetVarNameAndDim(TopFunction, ReturnType, EVariableCategory::Return);
+		FProperty* ReturnProp = GetVarNameAndDim(*FuncDef, ReturnType, EVariableCategory::Return);
 
 		TopFunction->NumParms++;
 	}
@@ -8247,8 +8247,9 @@ struct FExposeOnSpawnValidator
 	}
 };
 
-void FHeaderParser::CompileVariableDeclaration(UStruct* Struct)
+void FHeaderParser::CompileVariableDeclaration(FUnrealStructDefinitionInfo& StructDef)
 {
+	UStruct* Struct = StructDef.GetStruct();
 	EPropertyFlags DisallowFlags = CPF_ParmFlags;
 	EPropertyFlags EdFlags       = CPF_None;
 
@@ -8256,7 +8257,7 @@ void FHeaderParser::CompileVariableDeclaration(UStruct* Struct)
 	FPropertyBase OriginalProperty(CPT_None);
 	FIndexRange TypeRange;
 	ELayoutMacroType LayoutMacroType = ELayoutMacroType::None;
-	GetVarType( &FScope::GetTypeScope(Struct).Get(), OriginalProperty, DisallowFlags, /*OuterPropertyType=*/ NULL, EPropertyDeclarationStyle::UPROPERTY, EVariableCategory::Member, &TypeRange, &LayoutMacroType);
+	GetVarType(&*StructDef.GetScope(), OriginalProperty, DisallowFlags, /*OuterPropertyType=*/ NULL, EPropertyDeclarationStyle::UPROPERTY, EVariableCategory::Member, &TypeRange, &LayoutMacroType);
 	OriginalProperty.PropertyFlags |= EdFlags;
 
 	FString* Category = OriginalProperty.MetaData.Find(NAME_Category);
@@ -8332,7 +8333,7 @@ void FHeaderParser::CompileVariableDeclaration(UStruct* Struct)
 	for (;;)
 	{
 		FToken     Property    = OriginalProperty;
-		FProperty* NewProperty = GetVarNameAndDim(Struct, Property, EVariableCategory::Member, LayoutMacroType);
+		FProperty* NewProperty = GetVarNameAndDim(StructDef, Property, EVariableCategory::Member, LayoutMacroType);
 
 		// Optionally consume the :1 at the end of a bitfield boolean declaration
 		if (Property.IsBool())
@@ -8701,7 +8702,8 @@ ECompilationResult::Type FHeaderParser::ParseHeader()
 		auto ScopeTypeIterator = SourceFile.GetScope()->GetTypeIterator();
 		while (ScopeTypeIterator.MoveNext())
 		{
-			UField* Type = *ScopeTypeIterator;
+			FUnrealFieldDefinitionInfo* TypeDef = *ScopeTypeIterator;
+			UField* Type = TypeDef->GetField();
 
 			if (!Type->IsA<UScriptStruct>() && !Type->IsA<UClass>())
 			{
@@ -8712,7 +8714,7 @@ ECompilationResult::Type FHeaderParser::ParseHeader()
 
 			// now validate all delegate variables declared in the class
 			TMap<FName, UFunction*> DelegateCache;
-			FixupDelegateProperties(Struct, FScope::GetTypeScope(Struct).Get(), DelegateCache);
+			FixupDelegateProperties(Struct, *TypeDef->GetScope(), DelegateCache);
 		}
 
 		// Fix up any delegates themselves, if they refer to other delegates
@@ -8776,8 +8778,6 @@ ECompilationResult::Type FHeaderParser::ParseHeader()
 			const FString& ExpectedHeaderName = SourceFile.GetGeneratedHeaderFilename();
 			FError::Throwf(TEXT("Expected an include at the top of the header: '#include \"%s\"'"), *ExpectedHeaderName);
 		}
-
-		GScriptHelper.CheckForNoIInterfaces();
 	}
 #if !PLATFORM_EXCEPTIONS_DISABLED
 	catch( TCHAR* ErrorMsg )
@@ -10145,18 +10145,17 @@ bool FHeaderParser::TryToMatchConstructorParameterList(FToken Token)
 		return false;
 	}
 
-	FClassMetaData* ClassData = GScriptHelper.FindClassData(GetCurrentClass());
-	check(ClassData);
+	FStructMetaData& StructData = GetCurrentClassData();
 
 	bool bOICtor = false;
 	bool bVTCtor = false;
 
-	if (!ClassData->bDefaultConstructorDeclared && MatchSymbol(TEXT(')')))
+	if (!StructData.bDefaultConstructorDeclared && MatchSymbol(TEXT(')')))
 	{
-		ClassData->bDefaultConstructorDeclared = true;
+		StructData.bDefaultConstructorDeclared = true;
 	}
-	else if (!ClassData->bObjectInitializerConstructorDeclared
-		|| !ClassData->bCustomVTableHelperConstructorDeclared
+	else if (!StructData.bObjectInitializerConstructorDeclared
+		|| !StructData.bCustomVTableHelperConstructorDeclared
 	)
 	{
 		FToken ObjectInitializerParamParsingToken;
@@ -10228,11 +10227,11 @@ bool FHeaderParser::TryToMatchConstructorParameterList(FToken Token)
 			}
 		}
 
-		ClassData->bObjectInitializerConstructorDeclared = ClassData->bObjectInitializerConstructorDeclared || (bOICtor && bIsRef && bIsConst);
-		ClassData->bCustomVTableHelperConstructorDeclared = ClassData->bCustomVTableHelperConstructorDeclared || (bVTCtor && bIsRef);
+		StructData.bObjectInitializerConstructorDeclared = StructData.bObjectInitializerConstructorDeclared || (bOICtor && bIsRef && bIsConst);
+		StructData.bCustomVTableHelperConstructorDeclared = StructData.bCustomVTableHelperConstructorDeclared || (bVTCtor && bIsRef);
 	}
 
-	ClassData->bConstructorDeclared = ClassData->bConstructorDeclared || !bVTCtor;
+	StructData.bConstructorDeclared = StructData.bConstructorDeclared || !bVTCtor;
 
 	// Optionally match semicolon.
 	if (!MatchSymbol(TEXT(';')))
@@ -10326,8 +10325,10 @@ void FHeaderParser::ResetClassData()
 	check(CurrentClass->ClassWithin);
 }
 
-void FHeaderParser::PostPopNestClass(UClass* CurrentClass)
+void FHeaderParser::PostPopNestClass(FUnrealClassDefinitionInfo& CurrentClassDef)
 {
+	UClass* CurrentClass = CurrentClassDef.GetClass();
+
 	// Validate all the rep notify events here, to make sure they're implemented
 	VerifyPropertyMarkups(CurrentClass);
 
@@ -10413,7 +10414,7 @@ void FHeaderParser::PostPopNestClass(UClass* CurrentClass)
 void FHeaderParser::PostPopFunctionDeclaration(UFunction* PoppedFunction)
 {
 	//@TODO: UCREMOVAL: Move this code to occur at delegate var declaration, and force delegates to be declared before variables that use them
-	if (!GetCurrentScope()->IsFileScope() && GetCurrentClassData()->ContainsDelegates())
+	if (!GetCurrentScope()->IsFileScope() && GetCurrentClassData().ContainsDelegates())
 	{
 		// now validate all delegate variables declared in the class
 		TMap<FName, UFunction*> DelegateCache;
@@ -10421,12 +10422,12 @@ void FHeaderParser::PostPopFunctionDeclaration(UFunction* PoppedFunction)
 	}
 }
 
-void FHeaderParser::PostPopNestInterface(UClass* CurrentInterface)
+void FHeaderParser::PostPopNestInterface(FUnrealClassDefinitionInfo& CurrentInterfaceDef)
 {
-	FClassMetaData* ClassData = GScriptHelper.FindClassData(CurrentInterface);
-	check(ClassData);
-	if (ClassData->ContainsDelegates())
+	FStructMetaData& StructData = CurrentInterfaceDef.GetStructMetaData();
+	if (StructData.ContainsDelegates())
 	{
+		UClass* CurrentInterface = CurrentInterfaceDef.GetClass();
 		TMap<FName, UFunction*> DelegateCache;
 		FixupDelegateProperties(CurrentInterface, FScope::GetTypeScope(ExactCast<UClass>(CurrentInterface)).Get(), DelegateCache);
 	}
@@ -10773,10 +10774,10 @@ TFunctionType* CreateFunctionImpl(const FFuncInfo& FuncInfo, UObject* Outer, FSc
 		auto TypeIterator = CurrentScope->GetTypeIterator();
 		while (TypeIterator.MoveNext())
 		{
-			UField* Type = *TypeIterator;
-			if (Type->GetFName() == FuncInfo.Function.Identifier)
+			FUnrealFieldDefinitionInfo* Type = *TypeIterator;
+			if (Type->GetField()->GetFName() == FuncInfo.Function.Identifier)
 			{
-				FError::Throwf(TEXT("'%s' conflicts with '%s'"), FuncInfo.Function.Identifier, *Type->GetFullName());
+				FError::Throwf(TEXT("'%s' conflicts with '%s'"), FuncInfo.Function.Identifier, *Type->GetField()->GetFullName());
 			}
 		}
 	}
diff --git a/Engine/Source/Programs/UnrealHeaderTool/Private/HeaderParser.h b/Engine/Source/Programs/UnrealHeaderTool/Private/HeaderParser.h
--- a/Engine/Source/Programs/UnrealHeaderTool/Private/HeaderParser.h
+++ b/Engine/Source/Programs/UnrealHeaderTool/Private/HeaderParser.h
@@ -9,6 +9,7 @@
 #include "Scope.h"
 #include "UnrealTypeDefinitionInfo.h"
 #include "GeneratedCodeVersion.h"
+#include "ClassMaps.h"
 
 class UClass;
 enum class EGeneratedCodeVersion : uint8;
@@ -551,12 +552,20 @@ protected:
 		return (UClass*)GetCurrentClassScope()->GetStruct();
 	}
 
+	/**
+	 * Gets current class definition.
+	 */
+	FUnrealStructDefinitionInfo& GetCurrentClassDef() const
+	{
+		return GetCurrentClassScope()->GetStructDef();
+	}
+
 	/**
 	 * Gets current class's metadata.
 	 */
-	FClassMetaData* GetCurrentClassData()
+	FStructMetaData& GetCurrentClassData()
 	{
-		return GScriptHelper.FindClassData(GetCurrentClass());
+		return GetCurrentClassDef().GetStructMetaData();
 	}
 
 	// Information about all nesting levels.
@@ -642,7 +651,7 @@ protected:
 	bool IsBitfieldProperty(ELayoutMacroType LayoutMacroType);
 
 	// Parse the parameter list of a function or delegate declaration
-	void ParseParameterList(UFunction* Function, bool bExpectCommaBeforeName = false, TMap<FName, FString>* MetaData = NULL);
+	void ParseParameterList(FUnrealFunctionDefinitionInfo& FunctionDef, bool bExpectCommaBeforeName = false, TMap<FName, FString>* MetaData = NULL);
 
 public:
 	// Throws if a specifier value wasn't provided
@@ -711,7 +720,7 @@ protected:
 	bool SkipDeclaration(FToken& Token);
 	/** Similar to MatchSymbol() but will return to the exact location as on entry if the symbol was not found. */
 	bool SafeMatchSymbol(const TCHAR Match);
-	void HandleOneInheritedClass(UClass* Class, FString&& InterfaceName);
+	void HandleOneInheritedClass(FStructMetaData& StructMetaData, UClass* Class, FString&& InterfaceName);
 	FClass* ParseClassNameDeclaration(FString& DeclaredClassName, FString& RequiredAPIMacroIfPresent);
 
 	/** The property style of a variable declaration being parsed */
@@ -743,7 +752,7 @@ protected:
 	UClass* CompileClassDeclaration();
 	UDelegateFunction* CompileDelegateDeclaration(const TCHAR* DelegateIdentifier, EDelegateSpecifierAction::Type SpecifierAction = EDelegateSpecifierAction::DontParse);
 	void CompileFunctionDeclaration();
-	void CompileVariableDeclaration (UStruct* Struct);
+	void CompileVariableDeclaration (FUnrealStructDefinitionInfo& StructDef);
 	void CompileInterfaceDeclaration();
 	void CompileRigVMMethodDeclaration(UStruct* Struct);
 	void ParseRigVMMethodParameters(UStruct* Struct);
@@ -789,14 +798,14 @@ protected:
 	/**
 	 * Parses a variable name declaration and creates a new FProperty object.
 	 *
-	 * @param	Scope				struct to create the property in
+	 * @param	ParentStruct		struct to create the property in
 	 * @param	VarProperty			type and propertyflag info for the new property (inout)
 	 * @param   VariableCategory	what kind of variable is being created
 	 *
 	 * @return	a pointer to the new FProperty if successful, or NULL if there was no property to parse
 	 */
 	FProperty* GetVarNameAndDim(
-		UStruct* Struct,
+		FUnrealStructDefinitionInfo& ParentStruct,
 		FToken& VarProperty,
 		EVariableCategory::Type VariableCategory,
 		ELayoutMacroType LayoutMacroType = ELayoutMacroType::None);
@@ -867,17 +876,17 @@ protected:
 	 * Tasks that need to be done after popping interface definition
 	 * from parsing stack.
 	 *
-	 * @param CurrentInterface Interface that have just been popped.
+	 * @param CurrentInterfaceDef Interface that have just been popped.
 	 */
-	void PostPopNestInterface(UClass* CurrentInterface);
+	void PostPopNestInterface(FUnrealClassDefinitionInfo& CurrentInterfaceDef);
 
 	/**
 	 * Tasks that need to be done after popping class definition
 	 * from parsing stack.
 	 *
-	 * @param CurrentClass Class that have just been popped.
+	 * @param CurrentClassDef Class that have just been popped.
 	 */
-	void PostPopNestClass(UClass* CurrentClass);
+	void PostPopNestClass(FUnrealClassDefinitionInfo& CurrentClassDef);
 
 	/**
 	 * Binds all delegate properties declared in ValidationScope the delegate functions specified in the variable declaration, verifying that the function is a valid delegate
diff --git a/Engine/Source/Programs/UnrealHeaderTool/Private/NativeClassExporter.h b/Engine/Source/Programs/UnrealHeaderTool/Private/NativeClassExporter.h
--- a/Engine/Source/Programs/UnrealHeaderTool/Private/NativeClassExporter.h
+++ b/Engine/Source/Programs/UnrealHeaderTool/Private/NativeClassExporter.h
@@ -16,7 +16,7 @@ class UClass;
 class UEnum;
 class UScriptStruct;
 class UDelegateFunction;
-class FClassMetaData;
+class FStructMetaData;
 class FOutputDevice;
 class FUnrealPackageDefinitionInfo;
 struct FFuncInfo;
@@ -335,6 +335,9 @@ private:
 	/** Return the address of the singleton function */
 	static const FString& GetPackageSingletonNameFuncAddr(UPackage* Item, TSet<FString>* UniqueCrossModuleReferences);
 
+	/** Return the address of the singleton function - handles nullptr */
+	static const FString& GetSingletonNameFuncAddr(FUnrealFieldDefinitionInfo& Item, TSet<FString>* UniqueCrossModuleReferences, bool bRequiresValidObject = true);
+
 	/** Return the address of the singleton function - handles nullptr */
 	static const FString& GetSingletonNameFuncAddr(UField* Item, TSet<FString>* UniqueCrossModuleReferences, bool bRequiresValidObject = true);
 
@@ -522,7 +525,7 @@ private:
 	 * @param Class			class
 	 * @param ClassData		class data
 	 */
-	void ExportNativeFunctions(FOutputDevice& OutGeneratedHeaderText, FOutputDevice& OutGeneratedCPPText, FOutputDevice& OutMacroCalls, FOutputDevice& OutNoPureDeclsMacroCalls, FReferenceGatherers& OutReferenceGatherers, const FUnrealSourceFile& SourceFile, UClass* Class, FClassMetaData* ClassData) const;
+	void ExportNativeFunctions(FOutputDevice& OutGeneratedHeaderText, FOutputDevice& OutGeneratedCPPText, FOutputDevice& OutMacroCalls, FOutputDevice& OutNoPureDeclsMacroCalls, FReferenceGatherers& OutReferenceGatherers, const FUnrealSourceFile& SourceFile, UClass* Class, FStructMetaData& StructData) const;
 
 	/**
 	 * Export the actual internals to a standard thunk function
diff --git a/Engine/Source/Programs/UnrealHeaderTool/Private/ParserHelper.cpp b/Engine/Source/Programs/UnrealHeaderTool/Private/ParserHelper.cpp
--- a/Engine/Source/Programs/UnrealHeaderTool/Private/ParserHelper.cpp
+++ b/Engine/Source/Programs/UnrealHeaderTool/Private/ParserHelper.cpp
@@ -9,41 +9,14 @@
 #include "ClassMaps.h"
 
 /////////////////////////////////////////////////////
-// FClassMetaData
+// FStructMetaData
 
-/**
- * Finds the metadata for the property specified
- *
- * @param	Prop	the property to search for
- *
- * @return	pointer to the metadata for the property specified, or NULL
- *			if the property doesn't exist in the list (for example, if it
- *			is declared in a package that is already compiled and has had its
- *			source stripped)
- */
-FToken* FClassMetaData::FindTokenData(FProperty* Prop)
-{
-	check(Prop);
-	FUnrealPropertyDefinitionInfo& PropDef = GTypeDefinitionInfoMap.FindChecked(Prop).AsPropertyChecked();
-	return &PropDef.GetToken();
-}
-
-void FClassMetaData::AddProperty(FToken&& PropertyToken)
+void FStructMetaData::AddProperty(FUnrealPropertyDefinitionInfo& PropertyDef)
 {
-	FProperty* Prop = PropertyToken.TokenProperty;
-	check(Prop);
-
-	FUnrealPropertyDefinitionInfo& PropDef = GTypeDefinitionInfoMap.FindChecked(Prop).AsPropertyChecked();
-	PropDef.GetToken() = MoveTemp(PropertyToken);
-
-	// Add the property to the owner
-	UObject* Outer = Prop->GetOwner<UObject>();
-	check(Outer);
-	GTypeDefinitionInfoMap.FindChecked(Outer).AsStructChecked().AddProperty(PropDef);
-
 	// update the optimization flags
 	if (!bContainsDelegates)
 	{
+		FProperty* Prop = PropertyDef.GetProperty();
 		if (Prop->IsA(FDelegateProperty::StaticClass()) || Prop->IsA(FMulticastDelegateProperty::StaticClass()))
 		{
 			bContainsDelegates = true;
@@ -62,12 +35,12 @@ void FClassMetaData::AddProperty(FToken&& PropertyToken)
 	}
 }
 
-void FClassMetaData::AddInheritanceParent(FString&& InParent, FUnrealSourceFile* UnrealSourceFile)
+void FStructMetaData::AddInheritanceParent(FString&& InParent, FUnrealSourceFile* UnrealSourceFile)
 {
 	MultipleInheritanceParents.Add(new FMultipleInheritanceBaseClass(MoveTemp(InParent)));
 }
 
-void FClassMetaData::AddInheritanceParent(UClass* ImplementedInterfaceClass, FUnrealSourceFile* UnrealSourceFile)
+void FStructMetaData::AddInheritanceParent(UClass* ImplementedInterfaceClass, FUnrealSourceFile* UnrealSourceFile)
 {
 	MultipleInheritanceParents.Add(new FMultipleInheritanceBaseClass(ImplementedInterfaceClass));
 }
@@ -214,35 +187,3 @@ bool FAdvancedDisplayParameterHandler::CanMarkMore() const
 {
 	return bUseNumber ? (NumberLeaveUnmarked > 0) : (0 != ParametersNames.Num());
 }
-
-FClassMetaData* FCompilerMetadataManager::AddClassData(UStruct* Struct, FUnrealSourceFile* UnrealSourceFile)
-{
-	TUniquePtr<FClassMetaData>* pClassData = Find(Struct);
-	if (pClassData == NULL)
-	{
-		pClassData = &Emplace(Struct, new FClassMetaData());
-	}
-
-	return pClassData->Get();
-}
-
-FClassMetaData* FCompilerMetadataManager::AddInterfaceClassData(UStruct* Struct, FUnrealSourceFile* UnrealSourceFile)
-{
-	FClassMetaData* ClassData = AddClassData(Struct, UnrealSourceFile);
-	ClassData->ParsedInterface = EParsedInterface::ParsedUInterface;
-	InterfacesToVerify.Emplace(Struct, ClassData);
-	return ClassData;
-}
-
-void FCompilerMetadataManager::CheckForNoIInterfaces()
-{
-	for (const TPair<UStruct*, FClassMetaData*>& StructDataPair : InterfacesToVerify)
-	{
-		if (StructDataPair.Value->ParsedInterface == EParsedInterface::ParsedUInterface)
-		{
-			FString Name = StructDataPair.Key->GetName();
-			FError::Throwf(TEXT("UInterface 'U%s' parsed without a corresponding 'I%s'"), *Name, *Name);
-		}
-	}
-	InterfacesToVerify.Reset();
-}
\ No newline at end of file
diff --git a/Engine/Source/Programs/UnrealHeaderTool/Private/ParserHelper.h b/Engine/Source/Programs/UnrealHeaderTool/Private/ParserHelper.h
--- a/Engine/Source/Programs/UnrealHeaderTool/Private/ParserHelper.h
+++ b/Engine/Source/Programs/UnrealHeaderTool/Private/ParserHelper.h
@@ -22,8 +22,7 @@ class UObject;
 class UField;
 class UMetaData;
 class FHeaderParser;
-
-extern class FCompilerMetadataManager GScriptHelper;
+class FUnrealPropertyDefinitionInfo;
 
 /*-----------------------------------------------------------------------------
 	FPropertyBase.
@@ -1214,7 +1213,7 @@ enum class EParsedInterface
 /**
  * Class for storing compiler metadata about a class's properties.
  */
-class FClassMetaData
+class FStructMetaData
 {
 	/** base classes to multiply inherit from (other than the main base class */
 	TArray<FMultipleInheritanceBaseClass*>					MultipleInheritanceParents;
@@ -1233,7 +1232,7 @@ class FClassMetaData
 
 public:
 	/** Default constructor */
-	FClassMetaData()
+	FStructMetaData()
 		: bContainsDelegates(false)
 		, PrologLine(-1)
 		, GeneratedBodyLine(-1)
@@ -1312,9 +1311,9 @@ public:
 	 * Adds a new property to be tracked.  Determines the correct list for the property based on
 	 * its owner (function, struct, etc).
 	 * 
-	 * @param	PropertyToken	the property to add
+	 * @param	PropertyDef	the property to add
 	 */
-	void AddProperty(FToken&& PropertyToken);
+	void AddProperty(FUnrealPropertyDefinitionInfo& PropertyDef);
 
 	/**
 	 * Adds new editor-only metadata (key/value pairs) to the class or struct that
@@ -1367,18 +1366,6 @@ public:
 		}
 	}
 
-	/**
-	 * Finds the metadata for the property specified
-	 * 
-	 * @param	Prop	the property to search for
-	 *
-	 * @return	pointer to the metadata for the property specified, or NULL
-	 *			if the property doesn't exist in the list (for example, if it
-	 *			is declared in a package that is already compiled and has had its
-	 *			source stripped)
-	 */
-	FToken* FindTokenData( FProperty* Prop );
-
 	/**
 	 * Add a string to the list of inheritance parents for this class.
 	 *
@@ -1411,14 +1398,6 @@ public:
 		return bContainsDelegates;
 	}
 
-	/**
-	 * Shrink TMaps to avoid slack in Pairs array.
-	 */
-	void Shrink()
-	{
-		MultipleInheritanceParents.Shrink();
-	}
-
 	// Is constructor declared?
 	bool bConstructorDeclared;
 
@@ -1438,71 +1417,6 @@ public:
 	EParsedInterface ParsedInterface = EParsedInterface::NotAnInterface;
 };
 
-/**
- * Class for storing and linking data about properties and functions that is only required by the compiler.
- * The type of data tracked by this class is data that would otherwise only be accessible by adding a 
- * member property to UFunction/FProperty.  
- */
-class FCompilerMetadataManager : protected TMap<UStruct*, TUniquePtr<FClassMetaData> >
-{
-	using Super = TMap<UStruct*, TUniquePtr<FClassMetaData>>;
-
-public:
-	/**
-	 * Adds a new class to be tracked
-	 * 
-	 * @param	Struct	the UStruct to add
-	 *
-	 * @return	a pointer to the newly added metadata for the class specified
-	 */
-	FClassMetaData* AddClassData(UStruct* Struct, FUnrealSourceFile* UnrealSourceFile);
-
-	FClassMetaData* AddInterfaceClassData(UStruct* Struct, FUnrealSourceFile* UnrealSourceFile);
-
-	/**
-	 * Find the metadata associated with the class specified
-	 * 
-	 * @param	Struct	the UStruct to add
-	 *
-	 * @return	a pointer to the newly added metadata for the class specified
-	 */
-	FClassMetaData* FindClassData(UStruct* Struct)
-	{
-		FClassMetaData* Result = nullptr;
-
-		TUniquePtr<FClassMetaData>* pClassData = Find(Struct);
-		if (pClassData)
-		{
-			Result = pClassData->Get();
-		}
-
-		return Result;
-	}
-
-	/**
-	 * Shrink TMaps to avoid slack in Pairs array.
-	 */
-	void Shrink()
-	{
-		TMap<UStruct*, TUniquePtr<FClassMetaData> >::Shrink();
-		for (TMap<UStruct*, TUniquePtr<FClassMetaData> >::TIterator It(*this); It; ++It)
-		{
-			FClassMetaData* MetaData = It->Value.Get();
-			MetaData->Shrink();
-		}
-	}
-
-	/**
-	 * Throws an exception if a UInterface was parsed but not the corresponding IInterface.
-	 */
-	void CheckForNoIInterfaces();
-
-	friend struct FCompilerMetadataManagerArchiveProxy;
-
-private:
-	TArray<TPair<UStruct*, FClassMetaData*>> InterfacesToVerify;
-};
-
 /*-----------------------------------------------------------------------------
 	Retry points.
 -----------------------------------------------------------------------------*/
diff --git a/Engine/Source/Programs/UnrealHeaderTool/Private/Scope.cpp b/Engine/Source/Programs/UnrealHeaderTool/Private/Scope.cpp
--- a/Engine/Source/Programs/UnrealHeaderTool/Private/Scope.cpp
+++ b/Engine/Source/Programs/UnrealHeaderTool/Private/Scope.cpp
@@ -7,8 +7,6 @@
 #include "UnrealTypeDefinitionInfo.h"
 #include "ClassMaps.h"
 
-extern FCompilerMetadataManager GScriptHelper;
-
 FScope::FScope(FScope* InParent)
 	: Parent(InParent)
 { }
@@ -19,9 +17,9 @@ FScope::FScope()
 
 }
 
-void FScope::AddType(UField* Type)
+void FScope::AddType(FUnrealFieldDefinitionInfo& Type)
 {
-	TypeMap.Add(Type->GetFName(), Type);
+	TypeMap.Add(Type.GetField()->GetFName(), &Type);
 }
 
 /**
@@ -68,9 +66,9 @@ void DispatchType(UField* Type, TArray<UEnum*> &Enums, TArray<UScriptStruct*> &S
 
 void FScope::SplitTypesIntoArrays(TArray<UEnum*>& Enums, TArray<UScriptStruct*>& Structs, TArray<UDelegateFunction*>& DelegateFunctions)
 {
-	for (TPair<FName, UField*>& TypePair : TypeMap)
+	for (TPair<FName, FUnrealFieldDefinitionInfo*>& TypePair : TypeMap)
 	{
-		UField* Type = TypePair.Value;
+		UField* Type = TypePair.Value->GetField();
 		DispatchType(Type, Enums, Structs, DelegateFunctions);
 	}
 }
@@ -86,16 +84,16 @@ TSharedRef<FScope> FScope::GetTypeScope(UStruct* Type)
 	return (*TypeDef)->GetScope();
 }
 
-UField* FScope::FindTypeByName(FName Name)
+FUnrealFieldDefinitionInfo* FScope::FindTypeByName(FName Name)
 {
 	if (!Name.IsNone())
 	{
-		TDeepScopeTypeIterator<UField, false> TypeIterator(this);
+		TDeepScopeTypeIterator<FUnrealFieldDefinitionInfo, false> TypeIterator(this);
 
 		while (TypeIterator.MoveNext())
 		{
-			UField* Type = *TypeIterator;
-			if (Type->GetFName() == Name)
+			FUnrealFieldDefinitionInfo* Type = *TypeIterator;
+			if (Type->GetField()->GetFName() == Name)
 			{
 				return Type;
 			}
@@ -105,16 +103,16 @@ UField* FScope::FindTypeByName(FName Name)
 	return nullptr;
 }
 
-const UField* FScope::FindTypeByName(FName Name) const
+const FUnrealFieldDefinitionInfo* FScope::FindTypeByName(FName Name) const
 {
 	if (!Name.IsNone())
 	{
-		TScopeTypeIterator<UField, true> TypeIterator = GetTypeIterator();
+		TScopeTypeIterator<FUnrealFieldDefinitionInfo, true> TypeIterator = GetTypeIterator();
 
 		while (TypeIterator.MoveNext())
 		{
-			UField* Type = *TypeIterator;
-			if (Type->GetFName() == Name)
+			FUnrealFieldDefinitionInfo* Type = *TypeIterator;
+			if (Type->GetField()->GetFName() == Name)
 			{
 				return Type;
 			}
@@ -124,9 +122,9 @@ const UField* FScope::FindTypeByName(FName Name) const
 	return nullptr;
 }
 
-bool FScope::ContainsType(UField* Type)
+bool FScope::ContainsType(FUnrealFieldDefinitionInfo* Type)
 {
-	return FindTypeByName(Type->GetFName()) != nullptr;
+	return FindTypeByName(Type->GetField()->GetFName()) != nullptr;
 }
 
 bool FScope::IsFileScope() const
@@ -169,18 +167,7 @@ FName FFileScope::GetName() const
 	return Name;
 }
 
-UStruct* FStructScope::GetStruct() const
-{
-	return Struct;
-}
-
 FName FStructScope::GetName() const
 {
 	return Struct->GetFName();
-}
-
-FStructScope::FStructScope(UStruct* InStruct, FScope* InParent)
-	: FScope(InParent), Struct(InStruct)
-{
-
-}
+}
\ No newline at end of file
diff --git a/Engine/Source/Programs/UnrealHeaderTool/Private/Scope.h b/Engine/Source/Programs/UnrealHeaderTool/Private/Scope.h
--- a/Engine/Source/Programs/UnrealHeaderTool/Private/Scope.h
+++ b/Engine/Source/Programs/UnrealHeaderTool/Private/Scope.h
@@ -4,6 +4,7 @@
 
 #include "CoreMinimal.h"
 #include "UObject/Class.h"
+#include "UnrealTypeDefinitionInfo.h"
 
 class UEnum;
 class UScriptStruct;
@@ -11,7 +12,7 @@ class UDelegateFunction;
 
 // Forward declarations.
 class UStruct;
-class FClassMetaData;
+class FStructMetaData;
 class FUnrealSourceFile;
 class FFileScope;
 class FStructScope;
@@ -53,7 +54,7 @@ public:
 	 *
 	 * @param Type Type to add.
 	 */
-	void AddType(UField* Type);
+	void AddType(FUnrealFieldDefinitionInfo& Type);
 
 	/**
 	 * Finds type by its name.
@@ -62,7 +63,7 @@ public:
 	 *
 	 * @returns Found type or nullptr on failure.
 	 */
-	UField* FindTypeByName(FName Name);
+	FUnrealFieldDefinitionInfo* FindTypeByName(FName Name);
 
 	/**
 	 * Finds type by its name.
@@ -72,7 +73,7 @@ public:
 	 * @returns Found type or nullptr on failure.
 	 */
 	template <typename CharType>
-	UField* FindTypeByName(const CharType* Name)
+	FUnrealFieldDefinitionInfo* FindTypeByName(const CharType* Name)
 	{
 		return FindTypeByName(FName(Name, FNAME_Find));
 	}
@@ -86,7 +87,7 @@ public:
 	 *
 	 * @returns Found type or nullptr on failure.
 	 */
-	const UField* FindTypeByName(FName Name) const;
+	const FUnrealFieldDefinitionInfo* FindTypeByName(FName Name) const;
 
 	/**
 	 * Finds type by its name.
@@ -98,7 +99,7 @@ public:
 	 * @returns Found type or nullptr on failure.
 	 */
 	template <typename CharType>
-	const UField* FindTypeByName(const CharType* Name) const
+	const FUnrealFieldDefinitionInfo* FindTypeByName(const CharType* Name) const
 	{
 		return FindTypeByName(FName(Name, FNAME_Find));
 	}
@@ -107,7 +108,7 @@ public:
 	 *
 	 * @param Type Type to look for.
 	 */
-	bool ContainsType(UField* Type);
+	bool ContainsType(FUnrealFieldDefinitionInfo* Type);
 
 	/**
 	 * Checks if scope contains type that satisfies a predicate.
@@ -161,7 +162,7 @@ public:
 	{
 		for (const auto& NameTypePair : TypeMap)
 		{
-			if (NameTypePair.Value->IsA<TTypeFilter>())
+			if (NameTypePair.Value->GetField()->IsA<TTypeFilter>())
 			{
 				OutArray.Add((TTypeFilter*)NameTypePair.Value);
 			}
@@ -184,7 +185,7 @@ public:
 	{
 	public:
 		// Conditional typedefs.
-		typedef typename TConditionalType<TIsConst, TMap<FName, UField*>::TConstIterator, TMap<FName, UField*>::TIterator>::Type MapIteratorType;
+		typedef typename TConditionalType<TIsConst, TMap<FName, FUnrealFieldDefinitionInfo*>::TConstIterator, TMap<FName, FUnrealFieldDefinitionInfo*>::TIterator>::Type MapIteratorType;
 		typedef typename TConditionalType<TIsConst, const FScope, FScope>::Type ScopeType;
 
 		// Constructor
@@ -201,7 +202,7 @@ public:
 		 */
 		TType* operator*()
 		{
-			return bBeforeStart ? nullptr : (TType*)(*TypeIterator).Value;
+			return bBeforeStart ? nullptr : static_cast<TType*>((*TypeIterator).Value);
 		}
 
 		/**
@@ -234,7 +235,7 @@ public:
 	 *
 	 * @return Type iterator.
 	 */
-	template <class TType = UField>
+	template <class TType = FUnrealFieldDefinitionInfo>
 	TScopeTypeIterator<TType, false> GetTypeIterator()
 	{
 		return TScopeTypeIterator<TType, false>(this);
@@ -245,7 +246,7 @@ public:
 	 *
 	 * @return Const type iterator.
 	 */
-	template <class TType = UField>
+	template <class TType = FUnrealFieldDefinitionInfo>
 	TScopeTypeIterator<TType, true> GetTypeIterator() const
 	{
 		return TScopeTypeIterator<TType, true>(this);
@@ -271,10 +272,7 @@ private:
 	const FScope* Parent;
 
 	// Map of types in this scope.
-	TMap<FName, UField*> TypeMap;
-
-	friend struct FScopeArchiveProxy;
-	friend struct FStructScopeArchiveProxy;
+	TMap<FName, FUnrealFieldDefinitionInfo*> TypeMap;
 };
 
 using FScopeSet = TSet<FScope*, DefaultKeyFuncs<FScope*>, TInlineSetAllocator<1024>>;
@@ -330,7 +328,6 @@ public:
 		}
 	}
 
-
 	const TArray<FFileScope*>& GetIncludedScopes() const
 	{
 		return IncludedScopes;
@@ -340,6 +337,7 @@ public:
 	{
 		SourceFile = InSourceFile;
 	}
+
 private:
 	// Source file.
 	FUnrealSourceFile* SourceFile;
@@ -349,8 +347,6 @@ private:
 
 	// Included scopes list.
 	TArray<FFileScope*> IncludedScopes;
-
-	friend struct FFileScopeArchiveProxy;
 };
 
 
@@ -361,13 +357,30 @@ class FStructScope : public FScope
 {
 public:
 	// Constructor.
-	FStructScope(UStruct* Struct, FScope* Parent);
+	FStructScope(FUnrealStructDefinitionInfo& InStructDef, UStruct* InStruct, FScope* InParent)
+		: FScope(InParent)
+		, Struct(InStruct)
+		, StructDef(InStructDef)
+	{
+	}
 
 	virtual FStructScope* AsStructScope() override { return this; }
+
 	/**
 	 * Gets struct associated with this scope.
 	 */
-	UStruct* GetStruct() const;
+	UStruct* GetStruct() const
+	{
+		return Struct;
+	}
+
+	/**
+	 * Get the structure definition associated with this scope
+	 */
+	FUnrealStructDefinitionInfo& GetStructDef()
+	{
+		return StructDef;
+	}
 
 	/**
 	 * Gets scope name.
@@ -377,8 +390,7 @@ public:
 private:
 	// Struct associated with this scope.
 	UStruct* Struct;
-
-	friend struct FStructScopeArchiveProxy;
+	FUnrealStructDefinitionInfo& StructDef;
 };
 
 /**
diff --git a/Engine/Source/Programs/UnrealHeaderTool/Private/UnrealSourceFile.cpp b/Engine/Source/Programs/UnrealHeaderTool/Private/UnrealSourceFile.cpp
--- a/Engine/Source/Programs/UnrealHeaderTool/Private/UnrealSourceFile.cpp
+++ b/Engine/Source/Programs/UnrealHeaderTool/Private/UnrealSourceFile.cpp
@@ -78,9 +78,9 @@ const FString& FUnrealSourceFile::GetStrippedFilename() const
 	return StrippedFilename;
 }
 
-FString FUnrealSourceFile::GetGeneratedMacroName(FClassMetaData* ClassData, const TCHAR* Suffix) const
+FString FUnrealSourceFile::GetGeneratedMacroName(FStructMetaData& StructData, const TCHAR* Suffix) const
 {
-	return GetGeneratedMacroName(ClassData->GetGeneratedBodyLine(), Suffix);
+	return GetGeneratedMacroName(StructData.GetGeneratedBodyLine(), Suffix);
 }
 
 FString FUnrealSourceFile::GetGeneratedMacroName(int32 LineNumber, const TCHAR* Suffix) const
diff --git a/Engine/Source/Programs/UnrealHeaderTool/Private/UnrealSourceFile.h b/Engine/Source/Programs/UnrealHeaderTool/Private/UnrealSourceFile.h
--- a/Engine/Source/Programs/UnrealHeaderTool/Private/UnrealSourceFile.h
+++ b/Engine/Source/Programs/UnrealHeaderTool/Private/UnrealSourceFile.h
@@ -13,7 +13,7 @@ class UPackage;
 class UClass;
 class UStruct;
 class FArchive;
-class FClassMetaData;
+class FStructMetaData;
 
 enum class ETopologicalState : uint8
 {
@@ -194,10 +194,10 @@ public:
 	/**
 	 * Gets file-wise generated body macro name.
 	 *
-	 * @param ClassData Class metadata for which to get generated body macro name.
+	 * @param StructData Struct metadata for which to get generated body macro name.
 	 * @param Suffix Suffix to add to generated body macro name.
 	 */
-	FString GetGeneratedMacroName(FClassMetaData* ClassData, const TCHAR* Suffix = nullptr) const;
+	FString GetGeneratedMacroName(FStructMetaData& StructData, const TCHAR* Suffix = nullptr) const;
 
 	/**
 	 * Gets scope for this file.
diff --git a/Engine/Source/Programs/UnrealHeaderTool/Private/UnrealTypeDefinitionInfo.cpp b/Engine/Source/Programs/UnrealHeaderTool/Private/UnrealTypeDefinitionInfo.cpp
--- a/Engine/Source/Programs/UnrealHeaderTool/Private/UnrealTypeDefinitionInfo.cpp
+++ b/Engine/Source/Programs/UnrealHeaderTool/Private/UnrealTypeDefinitionInfo.cpp
@@ -245,7 +245,7 @@ void FUnrealStructDefinitionInfo::SetObject(UObject* InObject)
 	// Don't create a scope for things without a source.  Those are builtin types
 	if (HasSource())
 	{
-		StructScope = MakeShared<FStructScope>(static_cast<UStruct*>(InObject), &GetUnrealSourceFile().GetScope().Get());
+		StructScope = MakeShared<FStructScope>(*this, static_cast<UStruct*>(InObject), &GetUnrealSourceFile().GetScope().Get());
 	}
 }
 
@@ -262,3 +262,15 @@ uint32 FUnrealClassDefinitionInfo::GetHash(bool bIncludeNoExport) const
 	return FUnrealStructDefinitionInfo::GetHash(bIncludeNoExport);
 }
 
+
+void FUnrealClassDefinitionInfo::PostParseFinalize()
+{
+	if (IsInterface() && GetStructMetaData().ParsedInterface == EParsedInterface::ParsedUInterface)
+	{
+		FString UName = GetNameCPP();
+		FString IName = TEXT("I") + UName.RightChop(1);
+		FError::Throwf(TEXT("UInterface '%s' parsed without a corresponding '%s'"), *UName, *IName);
+	}
+
+	FUnrealStructDefinitionInfo::PostParseFinalize();
+}
\ No newline at end of file
diff --git a/Engine/Source/Programs/UnrealHeaderTool/Private/UnrealTypeDefinitionInfo.h b/Engine/Source/Programs/UnrealHeaderTool/Private/UnrealTypeDefinitionInfo.h
--- a/Engine/Source/Programs/UnrealHeaderTool/Private/UnrealTypeDefinitionInfo.h
+++ b/Engine/Source/Programs/UnrealHeaderTool/Private/UnrealTypeDefinitionInfo.h
@@ -756,11 +756,21 @@ public:
 		return Properties;
 	}
 
+	/**
+	 * Return the class meta data information
+	 */
+	FStructMetaData& GetStructMetaData()
+	{
+		return StructMetaData;
+	}
+
 private:
 	TSharedPtr<FScope> StructScope;
 
 	/** Properties of the structure */
 	TArray<FUnrealPropertyDefinitionInfo*> Properties;
+
+	FStructMetaData StructMetaData;
 };
 
 /**
@@ -925,7 +935,12 @@ public:
 		: FUnrealStructDefinitionInfo(InSourceFile, InLineNumber, MoveTemp(InNameCPP))
 		, BaseClassNameCPP(MoveTemp(InBaseClassNameCPP))
 		, bIsInterface(bInIsInterface)
-	{ }
+	{
+		if (bInIsInterface)
+		{
+			GetStructMetaData().ParsedInterface = EParsedInterface::ParsedUInterface;
+		}
+	}
 
 	FUnrealClassDefinitionInfo(FString&& InNameCPP)
 		: FUnrealStructDefinitionInfo(MoveTemp(InNameCPP))
@@ -946,6 +961,11 @@ public:
 
 	virtual uint32 GetHash(bool bIncludeNoExport = true) const override;
 
+	/**
+	 * Perform any post parsing finalization and validation
+	 */
+	virtual void PostParseFinalize() override;
+
 	/**
 	 * Return the Engine instance associated with the compiler instance
 	 */
