Merging the base class structure into the main type info.
Produced a common inheritance parser.
Added error return for when verification fails

#rnx

diff --git a/Engine/Source/Programs/UnrealHeaderTool/Private/BaseParser.h b/Engine/Source/Programs/UnrealHeaderTool/Private/BaseParser.h
--- a/Engine/Source/Programs/UnrealHeaderTool/Private/BaseParser.h
+++ b/Engine/Source/Programs/UnrealHeaderTool/Private/BaseParser.h
@@ -3,8 +3,10 @@
 #pragma once
 
 #include "CoreMinimal.h"
-#include "Containers/UnrealString.h"
 #include "GeneratedCodeVersion.h"
+#include "ParserHelper.h"
+#include "Containers/UnrealString.h"
+#include "UObject/ErrorException.h"
 
 class FToken;
 
@@ -230,8 +232,100 @@ public:
 	// Validates and inserts one key-value pair into the meta data map
 	static void InsertMetaDataPair(TMap<FName, FString>& MetaData, FName InKey, FString InValue);
 
+	/**
+	 * Parse class/struct inheritance.
+	 *
+	 * @param What				The name of the statement we are parsing.  (i.e. 'class')
+	 * @param InLambda			Function to call for every parent.  Must be in the form of
+	 *							Lambda(const TCHAR* Identifier, bool bSuperClass)
+	 */
+	template <typename Lambda>
+	void ParseInheritance(const TCHAR* What, Lambda&& InLambda);
+
 	//////////////
 
 	// Initialize the metadata keywords prior to parsing
 	static void InitMetadataKeywords();
 };
+
+template <typename Lambda>
+void FBaseParser::ParseInheritance(const TCHAR* What, Lambda&& InLambda)
+{
+
+	if (!MatchSymbol(TEXT(':')))
+	{
+		return;
+	}
+
+	// Process the super class 
+	{
+		FToken Token;
+		RequireIdentifier(TEXT("public"), ESearchCase::CaseSensitive, TEXT("inheritance"));
+		if (!GetIdentifier(Token))
+		{
+			FError::Throwf(TEXT("Missing %s name"), What);
+		}
+		RedirectTypeIdentifier(Token);
+		InLambda(Token.Identifier, true);
+	}
+
+	// Handle additional inherited interface classes
+	while (MatchSymbol(TEXT(',')))
+	{
+		RequireIdentifier(TEXT("public"), ESearchCase::CaseSensitive, TEXT("Interface inheritance must be public"));
+
+		FString InterfaceName;
+
+		for (;;)
+		{
+			FToken Token;
+			if (!GetIdentifier(Token, true))
+			{
+				FError::Throwf(TEXT("Failed to get interface class identifier"));
+			}
+
+			InterfaceName += Token.Identifier;
+
+			// Handle templated native classes
+			if (MatchSymbol(TEXT('<')))
+			{
+				InterfaceName += TEXT('<');
+
+				int32 NestedScopes = 1;
+				while (NestedScopes)
+				{
+					if (!GetToken(Token))
+					{
+						FError::Throwf(TEXT("Unexpected end of file"));
+					}
+
+					if (Token.TokenType == TOKEN_Symbol)
+					{
+						if (Token.Matches(TEXT('<')))
+						{
+							++NestedScopes;
+						}
+						else if (Token.Matches(TEXT('>')))
+						{
+							--NestedScopes;
+						}
+					}
+
+					InterfaceName += Token.Identifier;
+				}
+			}
+
+			// Handle scoped native classes
+			if (MatchSymbol(TEXT("::")))
+			{
+				InterfaceName += TEXT("::");
+
+				// Keep reading nested identifiers
+				continue;
+			}
+
+			break;
+		}
+		InLambda(*InterfaceName, false);
+	}
+}
diff --git a/Engine/Source/Programs/UnrealHeaderTool/Private/CodeGenerator.cpp b/Engine/Source/Programs/UnrealHeaderTool/Private/CodeGenerator.cpp
--- a/Engine/Source/Programs/UnrealHeaderTool/Private/CodeGenerator.cpp
+++ b/Engine/Source/Programs/UnrealHeaderTool/Private/CodeGenerator.cpp
@@ -6426,28 +6426,44 @@ void ResolveSuperClasses(UPackage* Package)
 			continue;
 		}
 
-		const FUnrealClassDefinitionInfo& ClassDef = GTypeDefinitionInfoMap.FindChecked(DefinedClass).AsClassChecked();
+		FUnrealClassDefinitionInfo& ClassDef = GTypeDefinitionInfoMap.FindChecked(DefinedClass).AsClassChecked();
 
-		const FString& BaseClassName         = ClassDef.GetBaseClassNameCPP();
-		const FString& BaseClassNameStripped = GetClassNameWithPrefixRemoved(BaseClassName);
-
-		if (!BaseClassNameStripped.IsEmpty() && !DefinedClass->GetSuperClass())
+		// Resolve the base class
 		{
-			UClass* FoundBaseClass = FindObject<UClass>(Package, *BaseClassNameStripped);
+			FUnrealStructDefinitionInfo::FBaseClassInfo& SuperClassInfo = ClassDef.GetSuperClassInfo();
+			const FString& BaseClassName = SuperClassInfo.Name;
+			const FString& BaseClassNameStripped = GetClassNameWithPrefixRemoved(BaseClassName);
 
-			if (FoundBaseClass == nullptr)
+			if (!BaseClassNameStripped.IsEmpty() && !DefinedClass->GetSuperClass())
 			{
-				FoundBaseClass = FindObject<UClass>(ANY_PACKAGE, *BaseClassNameStripped);
+				UClass* FoundBaseClass = FindObject<UClass>(Package, *BaseClassNameStripped);
+
+				if (FoundBaseClass == nullptr)
+				{
+					FoundBaseClass = FindObject<UClass>(ANY_PACKAGE, *BaseClassNameStripped);
+				}
+
+				if (FoundBaseClass == nullptr)
+				{
+					// Don't know its parent class. Raise error.
+					FError::Throwf(TEXT("Couldn't find parent type for '%s' named '%s' in current module (Package: %s) or any other module parsed so far."), *DefinedClass->GetName(), *BaseClassName, *GetNameSafe(Package));
+				}
+
+				DefinedClass->SetSuperStruct(FoundBaseClass);
+				DefinedClass->ClassCastFlags |= FoundBaseClass->ClassCastFlags;
+				SuperClassInfo.Struct = GTypeDefinitionInfoMap.FindChecked(FoundBaseClass).AsStruct();
 			}
+		}
 
-			if (FoundBaseClass == nullptr)
+		// Resolve the inherited classes
+		{
+			for (FUnrealStructDefinitionInfo::FBaseClassInfo& BaseClassInfo : ClassDef.GetBaseClassInfo())
 			{
-				// Don't know its parent class. Raise error.
-				FError::Throwf(TEXT("Couldn't find parent type for '%s' named '%s' in current module (Package: %s) or any other module parsed so far."), *DefinedClass->GetName(), *BaseClassName, *GetNameSafe(Package));
+				if (UClass* FoundClass = FClasses::FindScriptClass(BaseClassInfo.Name))
+				{
+					BaseClassInfo.Struct = GTypeDefinitionInfoMap.FindChecked(FoundClass).AsStruct();
+				}
 			}
-
-			DefinedClass->SetSuperStruct(FoundBaseClass);
-			DefinedClass->ClassCastFlags |= FoundBaseClass->ClassCastFlags;
 		}
 	}
 }
@@ -6699,7 +6715,7 @@ void DefineTypes(TArray<FUnrealPackageDefinitionInfo*>& PackageDefs)
 				for (TSharedRef<FUnrealTypeDefinitionInfo>& TypeDef : SourceFile.GetDefinedClasses())
 				{
 					FUnrealClassDefinitionInfo& ClassDef = TypeDef->AsClassChecked();
-					UClass* ResultClass = ProcessParsedClass(ClassDef.IsInterface(), ClassDef.GetNameCPP(), ClassDef.GetBaseClassNameCPP(), Package, RF_Public | RF_Standalone);
+					UClass* ResultClass = ProcessParsedClass(ClassDef.IsInterface(), ClassDef.GetNameCPP(), ClassDef.GetSuperClassInfo().Name, Package, RF_Public | RF_Standalone);
 					GTypeDefinitionInfoMap.Add(ResultClass, TypeDef);
 					ClassDef.SetObject(ResultClass);
 					AllClasses.Add(ResultClass);
@@ -7168,6 +7184,10 @@ ECompilationResult::Type UnrealHeaderTool_Main(const FString& ModuleInfoFilename
 		{
 			UE_LOG(LogCompile, Error, TEXT("Number of generated files mismatch ref=%d, ver=%d"), RefFileNames.Num(), VerFileNames.Num());
 		}
+		if (ChangeMessages.Num() > 0 || RefFileNames.Num() != VerFileNames.Num())
+		{
+			FResults::SetResult(ECompilationResult::OtherCompilationError);
+		}
 	}
 
 	RequestEngineExit(TEXT("UnrealHeaderTool finished"));
diff --git a/Engine/Source/Programs/UnrealHeaderTool/Private/HeaderParser.cpp b/Engine/Source/Programs/UnrealHeaderTool/Private/HeaderParser.cpp
--- a/Engine/Source/Programs/UnrealHeaderTool/Private/HeaderParser.cpp
+++ b/Engine/Source/Programs/UnrealHeaderTool/Private/HeaderParser.cpp
@@ -2584,128 +2584,83 @@ UScriptStruct* FHeaderParser::CompileStructDeclaration()
 		}
 	}
 
-	// Get optional superstruct.
-	bool bExtendsBaseStruct = false;
-	
-	if (MatchSymbol(TEXT(':')))
-	{
-		RequireIdentifier(TEXT("public"), ESearchCase::CaseSensitive, TEXT("struct inheritance"));
-		bExtendsBaseStruct = true;
-	}
+	// Parse the inheritance list
+	ParseInheritance(TEXT("struct"), [](const TCHAR* StructName, bool bIsSuperClass) {}); // Eat the results, already been parsed
 
-	UScriptStruct* BaseStruct = NULL;
-	if (bExtendsBaseStruct)
+	// if we have a base struct, propagate inherited struct flags now
+	//ETSTODO - Move base class resolution to the post declaration phase
+	FUnrealFieldDefinitionInfo* BaseStructDef = nullptr;
+	if (!StructDef.GetSuperClassInfo().Name.IsEmpty())
 	{
-		FToken ParentScope, ParentName;
-		if (GetIdentifier( ParentScope ))
-		{
-			RedirectTypeIdentifier(ParentScope);
+		const FString& ParentStructNameInScript = StructDef.GetSuperClassInfo().Name;
+		FString ParentStructNameStripped;
+		bool bOverrideParentStructName = false;
 
-			TSharedPtr<FScope> StructScope = Scope;
-			FString ParentStructNameInScript = FString(ParentScope.Identifier);
-			if (MatchSymbol(TEXT('.')))
-			{
-				if (GetIdentifier(ParentName))
-				{
-					RedirectTypeIdentifier(ParentName);
-
-					ParentStructNameInScript = FString(ParentName.Identifier);
-					FString ParentNameStripped = GetClassNameWithPrefixRemoved(ParentScope.Identifier);
-					FClass* StructClass = FClasses::FindClass(*ParentNameStripped);
-					if( !StructClass )
-					{
-						// If we find the literal class name, the user didn't use a prefix
-						StructClass = FClasses::FindClass(ParentScope.Identifier);
-						if( StructClass )
-						{
-							FError::Throwf(TEXT("'struct': Parent struct class '%s' is missing a prefix, expecting '%s'"), ParentScope.Identifier, *FString::Printf(TEXT("%s%s"),StructClass->GetPrefixCPP(),ParentScope.Identifier) );
-						}
-						else
-						{
-							FError::Throwf(TEXT("'struct': Can't find parent struct class '%s'"), ParentScope.Identifier );
-						}
-					}
+		if (!UHTConfig.StructsWithNoPrefix.Contains(ParentStructNameInScript))
+		{
+			bOverrideParentStructName = true;
+			ParentStructNameStripped = GetClassNameWithPrefixRemoved(ParentStructNameInScript);
+		}
 
-					StructScope = FScope::GetTypeScope(StructClass);
-				}
-				else
-				{
-					FError::Throwf( TEXT("'struct': Missing parent struct type after '%s.'"), ParentScope.Identifier );
-				}
-			}
-			
-			FString ParentStructNameStripped;
-			FUnrealFieldDefinitionInfo* TypeDef = nullptr;
-			const UField* Type = nullptr;
-			bool bOverrideParentStructName = false;
+		// If we're expecting a prefix, first try finding the correct field with the stripped struct name
+		if (bOverrideParentStructName)
+		{
+			BaseStructDef = Scope->FindTypeByName(*ParentStructNameStripped);
+		}
 
-			if( !UHTConfig.StructsWithNoPrefix.Contains(ParentStructNameInScript) )
-			{
-				bOverrideParentStructName = true;
-				ParentStructNameStripped = GetClassNameWithPrefixRemoved(ParentStructNameInScript);
-			}
+		// If it wasn't found, try to find the literal name given
+		if (BaseStructDef == nullptr)
+		{
+			BaseStructDef = Scope->FindTypeByName(*ParentStructNameInScript);
+		}
 
-			// If we're expecting a prefix, first try finding the correct field with the stripped struct name
+		// Resolve structs declared in another class  //@TODO: UCREMOVAL: This seems extreme
+		if (BaseStructDef == nullptr)
+		{
+			UScriptStruct* Type = nullptr;
 			if (bOverrideParentStructName)
 			{
-				TypeDef = StructScope->FindTypeByName(*ParentStructNameStripped);
+				Type = FindObject<UScriptStruct>(ANY_PACKAGE, *ParentStructNameStripped);
 			}
 
-			// If it wasn't found, try to find the literal name given
-			if (TypeDef == nullptr)
+			if (Type == nullptr)
 			{
-				TypeDef = StructScope->FindTypeByName(*ParentStructNameInScript);
+				Type = FindObject<UScriptStruct>(ANY_PACKAGE, *ParentStructNameInScript);
 			}
 
-			// Resolve structs declared in another class  //@TODO: UCREMOVAL: This seems extreme
-			if (TypeDef == nullptr)
-			{
-				if (bOverrideParentStructName)
-				{
-					Type = FindObject<UScriptStruct>(ANY_PACKAGE, *ParentStructNameStripped);
-				}
-
-				if (Type == nullptr)
-				{
-					Type = FindObject<UScriptStruct>(ANY_PACKAGE, *ParentStructNameInScript);
-				}
-			}
-			else
+			if (Type != nullptr)
 			{
-				Type = TypeDef->GetField();
+				BaseStructDef = &GTypeDefinitionInfoMap.FindChecked(Type).AsFieldChecked();
 			}
+		}
 
-			// If the struct still wasn't found, throw an error
-			if (Type == NULL)
-			{
-				FError::Throwf(TEXT("'struct': Can't find struct '%s'"), *ParentStructNameInScript );
-			}
-			else
+		// If the struct still wasn't found, throw an error
+		if (BaseStructDef == NULL)
+		{
+			FError::Throwf(TEXT("'struct': Can't find struct '%s'"), *ParentStructNameInScript);
+		}
+		else
+		{
+			// If the struct was found, confirm it adheres to the correct syntax. This should always fail if we were expecting an override that was not found.
+			if (bOverrideParentStructName)
 			{
-				// If the struct was found, confirm it adheres to the correct syntax. This should always fail if we were expecting an override that was not found.
-				BaseStruct = ((UScriptStruct*)Type);
-				if( bOverrideParentStructName )
+				UScriptStruct* BaseStruct = BaseStructDef->AsScriptStructChecked().GetScriptStruct();
+				const TCHAR* PrefixCPP = UHTConfig.StructsWithTPrefix.Contains(ParentStructNameStripped) ? TEXT("T") : BaseStruct->GetPrefixCPP();
+				if (ParentStructNameInScript != FString::Printf(TEXT("%s%s"), PrefixCPP, *ParentStructNameStripped))
 				{
-					const TCHAR* PrefixCPP = UHTConfig.StructsWithTPrefix.Contains(ParentStructNameStripped) ? TEXT("T") : BaseStruct->GetPrefixCPP();
-					if( ParentStructNameInScript != FString::Printf(TEXT("%s%s"), PrefixCPP, *ParentStructNameStripped) )
-					{
-						BaseStruct = nullptr;
-						FError::Throwf(TEXT("Parent Struct '%s' is missing a valid Unreal prefix, expecting '%s'"), *ParentStructNameInScript, *FString::Printf(TEXT("%s%s"), PrefixCPP, *Type->GetName()));
-					}
+					BaseStructDef = nullptr;
+					FError::Throwf(TEXT("Parent Struct '%s' is missing a valid Unreal prefix, expecting '%s'"), *ParentStructNameInScript, *FString::Printf(TEXT("%s%s"), PrefixCPP, *BaseStruct->GetName()));
 				}
 			}
 		}
-		else
-		{
-			FError::Throwf(TEXT("'struct': Missing parent struct after ': public'") );
-		}
 	}
 
-	// if we have a base struct, propagate inherited struct flags now
-	if (BaseStruct != NULL)
+	if (BaseStructDef != NULL)
 	{
+		UScriptStruct* BaseStruct = BaseStructDef->AsScriptStructChecked().GetScriptStruct();
 		Struct->SetSuperStruct(BaseStruct);
-		StructFlags |= (BaseStruct->StructFlags&STRUCT_Inherit);
+		StructFlags |= (BaseStruct->StructFlags & STRUCT_Inherit);
+		StructDef.GetSuperClassInfo().Struct = BaseStructDef->AsStruct();
 	}
 
 	Scope->AddType(StructDef);
@@ -6270,147 +6225,44 @@ bool FHeaderParser::SafeMatchSymbol( const TCHAR Match )
 	return false;
 }
 
-FClass* FHeaderParser::ParseClassNameDeclaration(FString& DeclaredClassName, FString& RequiredAPIMacroIfPresent)
+FUnrealClassDefinitionInfo& FHeaderParser::ParseClassNameDeclaration(FString& DeclaredClassName, FString& RequiredAPIMacroIfPresent)
 {
 	ParseNameWithPotentialAPIMacroPrefix(/*out*/ DeclaredClassName, /*out*/ RequiredAPIMacroIfPresent, TEXT("class"));
 
 	FClass* FoundClass = FClasses::FindClass(*GetClassNameWithPrefixRemoved(*DeclaredClassName));
 	check(FoundClass);
-	FStructMetaData& StructMetaData = GTypeDefinitionInfoMap.FindChecked(FoundClass).AsStructChecked().GetStructMetaData();
+	FoundClass->ClassCastFlags |= ClassCastFlagMap::Get().GetCastFlag(DeclaredClassName);
 
-	// Get parent class.
-	bool bSpecifiesParentClass = false;
+	FUnrealClassDefinitionInfo& ClassDef = GTypeDefinitionInfoMap.FindChecked(FoundClass).AsClassChecked();
 
 	// Skip optional final keyword
 	MatchIdentifier(TEXT("final"), ESearchCase::CaseSensitive);
 
-	if (MatchSymbol(TEXT(':')))
+	// Parse the inheritance list
+	ParseInheritance(TEXT("class"), [](const TCHAR* ClassName, bool bIsSuperClass) {}); // Eat the results, already been parsed
+
+	// Collect data about the super and base classes
+	if (FUnrealStructDefinitionInfo* SuperStruct = ClassDef.GetSuperClassInfo().Struct)
 	{
-		RequireIdentifier(TEXT("public"), ESearchCase::CaseSensitive, TEXT("class inheritance"));
-		bSpecifiesParentClass = true;
+		UClass* SuperClass = SuperStruct->AsClassChecked().GetClass();
+		FoundClass->ClassCastFlags |= SuperClass->ClassCastFlags;
 	}
-
-	// Add class cast flag
-	FoundClass->ClassCastFlags |= ClassCastFlagMap::Get().GetCastFlag(DeclaredClassName);
-
-	if (bSpecifiesParentClass)
+	for (FUnrealStructDefinitionInfo::FBaseClassInfo& BaseClassInfo : ClassDef.GetBaseClassInfo())
 	{
-		// Set the base class.
-		UClass* TempClass = GetQualifiedClass(TEXT("'extends'"));
-		check(TempClass);
-		// a class cannot 'extends' an interface, use 'implements'
-		if (TempClass->ClassFlags & CLASS_Interface)
-		{
-			FError::Throwf(TEXT("Class '%s' cannot extend interface '%s', use 'implements'"), *FoundClass->GetName(), *TempClass->GetName());
-		}
-
-		UClass* SuperClass = FoundClass->GetSuperClass();
-		if( SuperClass == NULL )
-		{
-			FoundClass->SetSuperStruct(TempClass);
-		}
-		else if( SuperClass != TempClass )
-		{
-			FError::Throwf(TEXT("%s's superclass must be %s, not %s"), *FoundClass->GetPathName(), *SuperClass->GetPathName(), *TempClass->GetPathName());
-		}
-
-		FoundClass->ClassCastFlags |= FoundClass->GetSuperClass()->ClassCastFlags;
-
-		// Handle additional inherited interface classes
-		while (MatchSymbol(TEXT(',')))
+		if (BaseClassInfo.Struct != nullptr)
 		{
-			RequireIdentifier(TEXT("public"), ESearchCase::CaseSensitive, TEXT("Interface inheritance must be public"));
-
-			FString InterfaceName;
-
-			FToken Token;
-			for (;;)
+			UClass* BaseClass = BaseClassInfo.Struct->AsClassChecked().GetClass();
+			if (!BaseClass->HasAnyClassFlags(CLASS_Interface))
 			{
-				if (!GetIdentifier(Token, true))
-				{
-					FError::Throwf(TEXT("Failed to get interface class identifier"));
-				}
-
-				InterfaceName += Token.Identifier;
-
-				// Handle templated native classes
-				if (MatchSymbol(TEXT('<')))
-				{
-					InterfaceName += TEXT('<');
-
-					int32 NestedScopes = 1;
-					while (NestedScopes)
-					{
-						if (!GetToken(Token))
-						{
-							FError::Throwf(TEXT("Unexpected end of file"));
-						}
-
-						if (Token.TokenType == TOKEN_Symbol)
-						{
-							if (Token.Matches(TEXT('<')))
-							{
-								++NestedScopes;
-							}
-							else if (Token.Matches(TEXT('>')))
-							{
-								--NestedScopes;
-							}
-						}
-
-						InterfaceName += Token.Identifier;
-					}
-				}
-
-				// Handle scoped native classes
-				if (MatchSymbol(TEXT("::")))
-				{
-					InterfaceName += TEXT("::");
-
-					// Keep reading nested identifiers
-					continue;
-				}
-
-				break;
+				FError::Throwf(TEXT("Implements: Class %s is not an interface; Can only inherit from non-UObjects or UInterface derived interfaces"), *BaseClass->GetName());
 			}
 
-			HandleOneInheritedClass(StructMetaData, FoundClass, MoveTemp(InterfaceName));
+			// Propagate the inheritable ClassFlags
+			FoundClass->ClassFlags |= (BaseClass->ClassFlags & CLASS_ScriptInherit);
+			FoundClass->Interfaces.Emplace(BaseClass, 0, false);
 		}
 	}
-	else if (FoundClass->GetSuperClass())
-	{
-		FError::Throwf(TEXT("class: missing 'Extends %s'"), *FoundClass->GetSuperClass()->GetName());
-	}
-
-	return FoundClass;
-}
-
-void FHeaderParser::HandleOneInheritedClass(FStructMetaData& StructMetaData, UClass* Class, FString&& InterfaceName)
-{
-
-	// Check for UInterface derived interface inheritance
-	if (UClass* Interface = FClasses::FindScriptClass(InterfaceName))
-	{
-		// Try to find the interface
-		if ( !Interface->HasAnyClassFlags(CLASS_Interface) )
-		{
-			FError::Throwf(TEXT("Implements: Class %s is not an interface; Can only inherit from non-UObjects or UInterface derived interfaces"), *Interface->GetName() );
-		}
-
-		// Propagate the inheritable ClassFlags
-		Class->ClassFlags |= (Interface->ClassFlags) & CLASS_ScriptInherit;
-
-		new (Class->Interfaces) FImplementedInterface(Interface, 0, false);
-		if (Interface->HasAnyClassFlags(CLASS_Native))
-		{
-			StructMetaData.AddInheritanceParent(Interface, &SourceFile);
-		}
-	}
-	else
-	{
-		// Non-UObject inheritance
-		StructMetaData.AddInheritanceParent(MoveTemp(InterfaceName), &SourceFile);
-	}
+	return ClassDef;
 }
 
 /**
@@ -6488,10 +6340,8 @@ UClass* FHeaderParser::CompileClassDeclaration()
 	FString DeclaredClassName;
 	FString RequiredAPIMacroIfPresent;
 	
-	FClass* Class = ParseClassNameDeclaration(/*out*/ DeclaredClassName, /*out*/ RequiredAPIMacroIfPresent);
-	check(Class);
-
-	FUnrealClassDefinitionInfo& ClassDeclarationData = GTypeDefinitionInfoMap.FindChecked(Class).AsClassChecked();
+	FUnrealClassDefinitionInfo& ClassDeclarationData = ParseClassNameDeclaration(/*out*/ DeclaredClassName, /*out*/ RequiredAPIMacroIfPresent);
+	FClass* Class = static_cast<FClass*>(ClassDeclarationData.GetClass());
 
 	ClassDefinitionRanges.Add(Class, ClassDefinitionRange(&Input[InputPos], nullptr));
 
@@ -6563,22 +6413,26 @@ UClass* FHeaderParser::CompileClassDeclaration()
 
 	// auto-create properties for all of the VFTables needed for the multiple inheritances
 	// get the inheritance parents
-	const TArray<FMultipleInheritanceBaseClass*>& InheritanceParents = StructData.GetInheritanceParents();
+	const TArray<FUnrealStructDefinitionInfo::FBaseClassInfo>& BaseClassInfos = ClassDeclarationData.GetBaseClassInfo();
 
 	// for all base class types, make a VfTable property
-	for (int32 ParentIndex = InheritanceParents.Num() - 1; ParentIndex >= 0; ParentIndex--)
+	for (int32 ParentIndex = BaseClassInfos.Num() - 1; ParentIndex >= 0; ParentIndex--)
 	{
 		// if this base class corresponds to an interface class, assign the vtable FProperty in the class's Interfaces map now...
-		if (UClass* InheritedInterface = InheritanceParents[ParentIndex]->InterfaceClass)
+		if (FUnrealStructDefinitionInfo* BaseClassStruct = BaseClassInfos[ParentIndex].Struct)
 		{
-			FImplementedInterface* Found = Class->Interfaces.FindByPredicate([=](const FImplementedInterface& Impl) { return Impl.Class == InheritedInterface; });
-			if (Found)
-			{
-				Found->PointerOffset = 1;
-			}
-			else
+			UClass* BaseClass = BaseClassStruct->AsClassChecked().GetClass();
+			if (BaseClass->HasAnyClassFlags(CLASS_Native))
 			{
-				Class->Interfaces.Add(FImplementedInterface(InheritedInterface, 1, false));
+				FImplementedInterface* Found = Class->Interfaces.FindByPredicate([&BaseClass](const FImplementedInterface& Impl) { return Impl.Class == BaseClass; });
+				if (Found)
+				{
+					Found->PointerOffset = 1;
+				}
+				else
+				{
+					Class->Interfaces.Add(FImplementedInterface(BaseClass, 1, false));
+				}
 			}
 		}
 	}
@@ -9591,40 +9445,26 @@ TSharedRef<FUnrealTypeDefinitionInfo> FHeaderPreParser::ParseClassDeclaration(FU
 	// Skip optional final keyword
 	MatchIdentifier(TEXT("final"), ESearchCase::CaseSensitive);
 
-	// Handle inheritance
-	FString BaseClassName;
-	if (MatchSymbol(TEXT(':')))
-	{
-		// Require 'public'
-		RequireIdentifier(TEXT("public"), ESearchCase::CaseSensitive, ErrorMsg);
 
-		// Inherits from something
-		FToken BaseClassNameToken;
-		if (!GetIdentifier(BaseClassNameToken, true))
+	// Create the definition
+	TSharedRef<FUnrealClassDefinitionInfo> ClassDef = MakeShareable(new FUnrealClassDefinitionInfo(SourceFile, InLineNumber, MoveTemp(ClassName), bClassIsAnInterface));
+
+	// Parse the inheritance list
+	ParseInheritance(TEXT("class"), [this, &SourceFile, &ClassNameWithoutPrefixStr, &ClassDef = *ClassDef](const TCHAR* ClassName, bool bIsSuperClass)
+	{
+		FString ClassNameStr(ClassName);
+		SourceFile.AddClassIncludeIfNeeded(InputLine, ClassNameWithoutPrefixStr, ClassNameStr);
+		if (bIsSuperClass)
 		{
-			FError::Throwf(TEXT("Expected a base class name"));
+			ClassDef.GetSuperClassInfo().Name = MoveTemp(ClassNameStr);
 		}
-
-		BaseClassName = BaseClassNameToken.Identifier;
-		SourceFile.AddClassIncludeIfNeeded(InputLine, ClassNameWithoutPrefixStr, BaseClassName);
-
-		// Get additional inheritance links and rack them up as dependencies if they're UObject derived
-		while (MatchSymbol(TEXT(',')))
+		else
 		{
-			// Require 'public'
-			RequireIdentifier(TEXT("public"), ESearchCase::CaseSensitive, ErrorMsg);
-
-			FToken InterfaceClassNameToken;
-			if (!GetIdentifier(InterfaceClassNameToken, true))
-			{
-				FUHTException::Throwf(SourceFile, InputLine, TEXT("Expected an interface class name"));
-			}
-
-			SourceFile.AddClassIncludeIfNeeded(InputLine, ClassNameWithoutPrefixStr, FString(InterfaceClassNameToken.Identifier));
+			ClassDef.GetBaseClassInfo().Emplace(FUnrealStructDefinitionInfo::FBaseClassInfo{ MoveTemp(ClassNameStr) });
 		}
 	}
+	);
 
-	TSharedRef<FUnrealClassDefinitionInfo> ClassDef = MakeShareable(new FUnrealClassDefinitionInfo(SourceFile, InLineNumber, MoveTemp(ClassName), MoveTemp(BaseClassName), bClassIsAnInterface));
 	ClassDef->MetaData = MoveTemp(MetaData);
 	ClassDef->ParseClassProperties(MoveTemp(SpecifiersFound), RequiredAPIMacroIfPresent);
 	return ClassDef;
@@ -9730,35 +9570,24 @@ TSharedRef<FUnrealTypeDefinitionInfo> FHeaderPreParser::ParseStructDeclaration(F
 	FString StructNameStripped = GetClassNameWithPrefixRemoved(StructNameInScript);
 
 
-	// Get optional superstruct.
-	FString ParentScopeCPP;
-	FString ParentNameCPP;
-	if (MatchSymbol(TEXT(':')))
+	// Create the structure definition
+	TSharedRef<FUnrealScriptStructDefinitionInfo> StructDef = MakeShareable(new FUnrealScriptStructDefinitionInfo(SourceFile, InLineNumber, MoveTemp(StructNameInScript)));
+
+	// Parse the inheritance list
+	ParseInheritance(TEXT("struct"), [this, &SourceFile, &StructNameStripped, &StructDef = *StructDef](const TCHAR* StructName, bool bIsSuperClass)
 	{
-		RequireIdentifier(TEXT("public"), ESearchCase::CaseSensitive, TEXT("struct inheritance"));
-		FToken ParentScope, ParentName;
-		if (GetIdentifier(ParentScope))
+		FString StructNameStr(StructName);
+		SourceFile.AddScriptStructIncludeIfNeeded(InputLine, StructNameStripped, StructNameStr);
+		if (bIsSuperClass)
 		{
-			RedirectTypeIdentifier(ParentScope);
-			if (MatchSymbol(TEXT('.')))
-			{
-				if (GetIdentifier(ParentName))
-				{
-					RedirectTypeIdentifier(ParentName);
-					ParentScopeCPP = ParentName.Identifier;
-					ParentNameCPP = ParentScope.Identifier;
-					SourceFile.AddClassIncludeIfNeeded(InputLine, TEXT(""), ParentScopeCPP);
-				}
-			}
-			else
-			{
-				ParentNameCPP = ParentScope.Identifier;
-				SourceFile.AddScriptStructIncludeIfNeeded(InputLine, StructNameStripped, ParentNameCPP);
-			}
+			StructDef.GetSuperClassInfo().Name = MoveTemp(StructNameStr);
+		}
+		else
+		{
+			StructDef.GetBaseClassInfo().Emplace(FUnrealStructDefinitionInfo::FBaseClassInfo{ MoveTemp(StructNameStr) });
 		}
 	}
-	TSharedRef<FUnrealScriptStructDefinitionInfo> StructDef = MakeShareable(new FUnrealScriptStructDefinitionInfo(SourceFile, 
-		InLineNumber, MoveTemp(StructNameInScript), MoveTemp(ParentScopeCPP), MoveTemp(ParentNameCPP)));
+	);
 	return StructDef;
 }
 
diff --git a/Engine/Source/Programs/UnrealHeaderTool/Private/HeaderParser.h b/Engine/Source/Programs/UnrealHeaderTool/Private/HeaderParser.h
--- a/Engine/Source/Programs/UnrealHeaderTool/Private/HeaderParser.h
+++ b/Engine/Source/Programs/UnrealHeaderTool/Private/HeaderParser.h
@@ -720,8 +720,7 @@ protected:
 	bool SkipDeclaration(FToken& Token);
 	/** Similar to MatchSymbol() but will return to the exact location as on entry if the symbol was not found. */
 	bool SafeMatchSymbol(const TCHAR Match);
-	void HandleOneInheritedClass(FStructMetaData& StructMetaData, UClass* Class, FString&& InterfaceName);
-	FClass* ParseClassNameDeclaration(FString& DeclaredClassName, FString& RequiredAPIMacroIfPresent);
+	FUnrealClassDefinitionInfo& ParseClassNameDeclaration(FString& DeclaredClassName, FString& RequiredAPIMacroIfPresent);
 
 	/** The property style of a variable declaration being parsed */
 	struct EPropertyDeclarationStyle
diff --git a/Engine/Source/Programs/UnrealHeaderTool/Private/ParserHelper.cpp b/Engine/Source/Programs/UnrealHeaderTool/Private/ParserHelper.cpp
--- a/Engine/Source/Programs/UnrealHeaderTool/Private/ParserHelper.cpp
+++ b/Engine/Source/Programs/UnrealHeaderTool/Private/ParserHelper.cpp
@@ -35,16 +35,6 @@ void FStructMetaData::AddProperty(FUnrealPropertyDefinitionInfo& PropertyDef)
 	}
 }
 
-void FStructMetaData::AddInheritanceParent(FString&& InParent, FUnrealSourceFile* UnrealSourceFile)
-{
-	MultipleInheritanceParents.Add(new FMultipleInheritanceBaseClass(MoveTemp(InParent)));
-}
-
-void FStructMetaData::AddInheritanceParent(UClass* ImplementedInterfaceClass, FUnrealSourceFile* UnrealSourceFile)
-{
-	MultipleInheritanceParents.Add(new FMultipleInheritanceBaseClass(ImplementedInterfaceClass));
-}
-
 /////////////////////////////////////////////////////
 // FPropertyBase
 
diff --git a/Engine/Source/Programs/UnrealHeaderTool/Private/ParserHelper.h b/Engine/Source/Programs/UnrealHeaderTool/Private/ParserHelper.h
--- a/Engine/Source/Programs/UnrealHeaderTool/Private/ParserHelper.h
+++ b/Engine/Source/Programs/UnrealHeaderTool/Private/ParserHelper.h
@@ -1172,37 +1172,6 @@ struct FFuncInfo
 	}
 };
 
-/**
- * Tracks information about a multiple inheritance parent declaration for native script classes.
- */
-struct FMultipleInheritanceBaseClass
-{
-	/**
-	 * The name to use for the base class when exporting the script class to header file.
-	 */
-	FString ClassName;
-
-	/**
-	 * For multiple inheritance parents declared using 'Implements', corresponds to the UClass for the interface.  For multiple inheritance parents declared
-	 * using 'Inherits', this value will be NULL.
-	 */
-	UClass* InterfaceClass = nullptr;
-
-	/**
-	 * Constructors
-	 */
-	explicit FMultipleInheritanceBaseClass(FString&& BaseClassName)
-		: ClassName(MoveTemp(BaseClassName))
-	{
-	}
-
-	explicit FMultipleInheritanceBaseClass(UClass* ImplementedInterfaceClass)
-		: InterfaceClass(ImplementedInterfaceClass)
-	{
-		ClassName = FString::Printf(TEXT("I%s"), *ImplementedInterfaceClass->GetName());
-	}
-};
-
 enum class EParsedInterface
 {
 	NotAnInterface,
@@ -1215,9 +1184,6 @@ enum class EParsedInterface
  */
 class FStructMetaData
 {
-	/** base classes to multiply inherit from (other than the main base class */
-	TArray<FMultipleInheritanceBaseClass*>					MultipleInheritanceParents;
-
 	/** whether this class declares delegate functions or properties */
 	bool													bContainsDelegates;
 
@@ -1366,30 +1332,6 @@ public:
 		}
 	}
 
-	/**
-	 * Add a string to the list of inheritance parents for this class.
-	 *
-	 * @param Inparent The C++ class name to add to the multiple inheritance list
-	 * @param UnrealSourceFile Currently parsed source file.
-	 */
-	void AddInheritanceParent(FString&& InParent, FUnrealSourceFile* UnrealSourceFile);
-
-	/**
-	 * Add a string to the list of inheritance parents for this class.
-	 *
-	 * @param Inparent	The C++ class name to add to the multiple inheritance list
-	 * @param UnrealSourceFile Currently parsed source file.
-	 */
-	void AddInheritanceParent(UClass* ImplementedInterfaceClass, FUnrealSourceFile* UnrealSourceFile);
-
-	/**
-	 * Return the list of inheritance parents
-	 */
-	const TArray<FMultipleInheritanceBaseClass*>& GetInheritanceParents() const
-	{
-		return MultipleInheritanceParents;
-	}
-
 	/**
 	 * Returns whether this class contains any delegate properties which need to be fixed up.
 	 */
diff --git a/Engine/Source/Programs/UnrealHeaderTool/Private/UnrealTypeDefinitionInfo.h b/Engine/Source/Programs/UnrealHeaderTool/Private/UnrealTypeDefinitionInfo.h
--- a/Engine/Source/Programs/UnrealHeaderTool/Private/UnrealTypeDefinitionInfo.h
+++ b/Engine/Source/Programs/UnrealHeaderTool/Private/UnrealTypeDefinitionInfo.h
@@ -716,6 +716,13 @@ private:
  */
 class FUnrealStructDefinitionInfo : public FUnrealFieldDefinitionInfo
 {
+public:
+	struct FBaseClassInfo
+	{
+		FString Name;
+		FUnrealStructDefinitionInfo* Struct = nullptr;
+	};
+
 public:
 	using FUnrealFieldDefinitionInfo::FUnrealFieldDefinitionInfo;
 
@@ -764,6 +771,19 @@ public:
 		return StructMetaData;
 	}
 
+	/**
+	 * Return the super class information
+	 */
+	FBaseClassInfo& GetSuperClassInfo()
+	{
+		return SuperClassInfo;
+	}
+
+	TArray<FBaseClassInfo>& GetBaseClassInfo()
+	{
+		return BaseClassInfo;
+	}
+
 private:
 	TSharedPtr<FScope> StructScope;
 
@@ -771,6 +791,8 @@ private:
 	TArray<FUnrealPropertyDefinitionInfo*> Properties;
 
 	FStructMetaData StructMetaData;
+	FBaseClassInfo SuperClassInfo;
+	TArray<FBaseClassInfo> BaseClassInfo;
 };
 
 /**
@@ -779,10 +801,8 @@ private:
 class FUnrealScriptStructDefinitionInfo : public FUnrealStructDefinitionInfo
 {
 public:
-	FUnrealScriptStructDefinitionInfo(FUnrealSourceFile& InSourceFile, int32 InLineNumber, FString&& InNameCPP, FString&& InParentScopeCPP, FString&& InParentNameCPP)
+	FUnrealScriptStructDefinitionInfo(FUnrealSourceFile& InSourceFile, int32 InLineNumber, FString&& InNameCPP)
 		: FUnrealStructDefinitionInfo(InSourceFile, InLineNumber, MoveTemp(InNameCPP))
-		, ParentScopeCPP(MoveTemp(InParentScopeCPP))
-		, ParentNameCPP(MoveTemp(InParentNameCPP))
 	{ }
 
 	virtual FUnrealScriptStructDefinitionInfo* AsScriptStruct() override
@@ -807,26 +827,6 @@ public:
 	{
 		return static_cast<UScriptStruct*>(GetObject());
 	}
-
-	/**
-	 * Return the parent structure scoped name (currently blank)
-	 */
-	const FString& GetParentScopeCPP() const
-	{
-		return ParentScopeCPP;
-	}
-
-	/**
-	 * Return the name of the parent structure.  If not derived from another struct, it will be empty
-	 */
-	const FString& GetParentNameCPP() const
-	{
-		return ParentNameCPP;
-	}
-
-private:
-	FString ParentScopeCPP;
-	FString ParentNameCPP;
 };
 
 /**
@@ -931,9 +931,8 @@ class FUnrealClassDefinitionInfo
 	, public FClassDeclarationMetaData
 {
 public:
-	FUnrealClassDefinitionInfo(FUnrealSourceFile& InSourceFile, int32 InLineNumber, FString&& InNameCPP, FString&& InBaseClassNameCPP, bool bInIsInterface)
+	FUnrealClassDefinitionInfo(FUnrealSourceFile& InSourceFile, int32 InLineNumber, FString&& InNameCPP, bool bInIsInterface)
 		: FUnrealStructDefinitionInfo(InSourceFile, InLineNumber, MoveTemp(InNameCPP))
-		, BaseClassNameCPP(MoveTemp(InBaseClassNameCPP))
 		, bIsInterface(bInIsInterface)
 	{
 		if (bInIsInterface)
@@ -1014,16 +1013,7 @@ public:
 		return bIsInterface;
 	}
 
-	/** 
-	 * Return the CPP name of the base class or blank if there is none.
-	 */
-	const FString& GetBaseClassNameCPP() const
-	{
-		return BaseClassNameCPP;
-	}
-
 private:
-	FString BaseClassNameCPP;
 	FString EnclosingDefine;
 	ESerializerArchiveType ArchiveType = ESerializerArchiveType::None;
 	bool bIsInterface = false;
