Adding the interface's hash code to classes to make sure they re-instance properly when the interface is changed.

#rnx

diff --git a/Engine/Source/Programs/UnrealHeaderTool/Private/CodeGenerator.cpp b/Engine/Source/Programs/UnrealHeaderTool/Private/CodeGenerator.cpp
--- a/Engine/Source/Programs/UnrealHeaderTool/Private/CodeGenerator.cpp
+++ b/Engine/Source/Programs/UnrealHeaderTool/Private/CodeGenerator.cpp
@@ -2260,11 +2260,15 @@ void FNativeClassHeaderGenerator::ExportNativeGeneratedInitCode(FOutputDevice& O
 				{
 					OffsetString = TEXT("0");
 				}
+
+				FUHTStringBuilder IntHash;
+				InterClassDef.GetHashTag(IntHash);
 				StaticDefinitions.Logf(
-					TEXT("\t\t\t{ %s, %s, %s },\r\n"),
+					TEXT("\t\t\t{ %s, %s, %s }, %s\r\n"),
 					*InterClassDef.GetSingletonNameChopped(false),
 					*OffsetString,
-					Inter.bImplementedByK2 ? TEXT("true") : TEXT("false")
+					Inter.bImplementedByK2 ? TEXT("true") : TEXT("false"),
+					*IntHash
 				);
 			}
 			StaticDefinitions.Log(TEXT("\t\t};\r\n"));
