Added new PropertyTypes which in the long run will be where different property types are implemented.
Added new casting system betwen the different type definitions.  This allows some patterns that makes the code easier.
Moved some direct calls to engine API to FClasses.
More property data moved to the property definition.
Replacing a many more instances of referencing engine type directly with referencing UHT types.
Removed the FToken from the property data.  Saved 600MB of memory on some titles.

#rnx

diff --git a/Engine/Source/Programs/UnrealHeaderTool/Private/BaseParser.cpp b/Engine/Source/Programs/UnrealHeaderTool/Private/BaseParser.cpp
--- a/Engine/Source/Programs/UnrealHeaderTool/Private/BaseParser.cpp
+++ b/Engine/Source/Programs/UnrealHeaderTool/Private/BaseParser.cpp
@@ -377,7 +377,6 @@ bool FBaseParser::IsWhitespace( TCHAR c )
 // Gets the next token from the input stream, advancing the variables which keep track of the current input position and line.
 bool FBaseParser::GetToken( FToken& Token, bool bNoConsts/*=false*/, ESymbolParseOption bParseTemplateClosingBracket/*=ESymbolParseOption::Normal*/ )
 {
-	Token.ClearTokenName();
 	TCHAR c = GetLeadingChar();
 	if( c == 0 )
 	{
@@ -831,27 +830,6 @@ bool FBaseParser::MatchSymbol(const TCHAR* Match, ESymbolParseOption bParseTempl
 	return false;
 }
 
-//
-// Get a specific identifier and return 1 if gotten, 0 if not.
-// This is used primarily for checking for required symbols during compilation.
-//
-bool FBaseParser::MatchIdentifierByName( FName Match )
-{
-	FToken Token;
-	if (!GetToken(Token))
-	{
-		return false;
-	}
-
-	if ((Token.TokenType == TOKEN_Identifier) && (Token.GetTokenName() == Match))
-	{
-		return true;
-	}
-
-	UngetToken(Token);
-	return false;
-}
-
 bool FBaseParser::MatchIdentifier( const TCHAR* Match, ESearchCase::Type SearchCase)
 {
 	FToken Token;
@@ -875,7 +853,7 @@ bool FBaseParser::MatchConstInt( const TCHAR* Match )
 	FToken Token;
 	if (GetToken(Token))
 	{
-		if( Token.TokenType==TOKEN_Const && (Token.Type == CPT_Int || Token.Type == CPT_Int64) && FCString::Strcmp(Token.Identifier,Match)==0 )
+		if( Token.TokenType==TOKEN_Const && (Token.ConstantType == CPT_Int || Token.ConstantType == CPT_Int64) && FCString::Strcmp(Token.Identifier,Match)==0 )
 		{
 			return true;
 		}
@@ -893,7 +871,7 @@ bool FBaseParser::MatchAnyConstInt()
 	FToken Token;
 	if (GetToken(Token))
 	{
-		if( Token.TokenType==TOKEN_Const && (Token.Type == CPT_Int || Token.Type == CPT_Int64) )
+		if( Token.TokenType==TOKEN_Const && (Token.ConstantType == CPT_Int || Token.ConstantType == CPT_Int64) )
 		{
 			return true;
 		}
@@ -941,17 +919,6 @@ bool FBaseParser::PeekSymbol( const TCHAR Match )
 //
 // Peek ahead and see if an identifier follows in the stream.
 //
-bool FBaseParser::PeekIdentifierByName( FName Match )
-{
-	FToken Token;
-	if (!GetToken(Token, true))
-	{
-		return false;
-	}
-	UngetToken(Token);
-	return Token.TokenType==TOKEN_Identifier && Token.GetTokenName()==Match;
-}
-
 bool FBaseParser::PeekIdentifier( const TCHAR* Match, ESearchCase::Type SearchCase)
 {
 	FToken Token;
@@ -972,6 +939,12 @@ void FBaseParser::UngetToken( const FToken& Token )
 	InputLine = Token.StartLine;
 }
 
+void FBaseParser::UngetToken(int32 StartLine, int32 StartPos)
+{
+	InputLine = StartLine;
+	InputPos = StartPos;
+}
+
 //
 // Require a symbol.
 //
diff --git a/Engine/Source/Programs/UnrealHeaderTool/Private/BaseParser.h b/Engine/Source/Programs/UnrealHeaderTool/Private/BaseParser.h
--- a/Engine/Source/Programs/UnrealHeaderTool/Private/BaseParser.h
+++ b/Engine/Source/Programs/UnrealHeaderTool/Private/BaseParser.h
@@ -17,6 +17,8 @@ enum class EPointerMemberBehavior
 	AllowAndLog,
 };
 
+using FMetaData = TMap<FName, FString>;
+
 /////////////////////////////////////////////////////
 // UHTConfig
 
@@ -171,6 +173,7 @@ public:
 	bool GetRawTokenRespectingQuotes( FToken& Token, TCHAR StopChar = TCHAR('\n') );
 
 	void UngetToken( const FToken& Token );
+	void UngetToken(int32 StartLine, int32 StartPos);
 	bool GetIdentifier( FToken& Token, bool bNoConsts = false );
 	bool GetSymbol( FToken& Token );
 
@@ -185,11 +188,9 @@ public:
 	bool GetConstInt64(int64& Result, const TCHAR* Tag = NULL);
 
 	// Matching predefined text.
-	bool MatchIdentifierByName( FName Match );
 	bool MatchIdentifier( const TCHAR* Match, ESearchCase::Type SearchCase);
 	bool MatchConstInt( const TCHAR* Match );
 	bool MatchAnyConstInt();
-	bool PeekIdentifierByName( FName Match );
 	bool PeekIdentifier( const TCHAR* Match, ESearchCase::Type SearchCase);
 	bool MatchSymbol( const TCHAR Match, ESymbolParseOption bParseTemplateClosingBracket = ESymbolParseOption::Normal );
 	bool MatchSymbol(const TCHAR* Match, ESymbolParseOption bParseTemplateClosingBracket = ESymbolParseOption::Normal);
diff --git a/Engine/Source/Programs/UnrealHeaderTool/Private/ClassMaps.h b/Engine/Source/Programs/UnrealHeaderTool/Private/ClassMaps.h
--- a/Engine/Source/Programs/UnrealHeaderTool/Private/ClassMaps.h
+++ b/Engine/Source/Programs/UnrealHeaderTool/Private/ClassMaps.h
@@ -52,9 +52,31 @@ struct FTypeDefinitionInfoMap : public FFreezableContainer
 		DefinitionsByName.Add(Object->GetFName(), Definition);
 	}
 
-	bool Contains(const UObject* Object) { check(bFrozen); return DefinitionsByUObject.Contains(Object); }
-	TSharedRef<FUnrealTypeDefinitionInfo>* Find(const UObject* Object) { check(bFrozen); return DefinitionsByUObject.Find(Object); }
-	TSharedRef<FUnrealTypeDefinitionInfo>& operator[](const UObject* Object) { check(bFrozen); return DefinitionsByUObject[Object]; }
+	bool Contains(const UObject* Object)
+	{
+		check(bFrozen); 
+		return DefinitionsByUObject.Contains(Object); 
+	}
+
+	TSharedRef<FUnrealTypeDefinitionInfo>* Find(const UObject* Object)
+	{
+		check(bFrozen);
+		return DefinitionsByUObject.Find(Object);
+	}
+
+	template<typename To>
+	To* Find(const UObject* Object)
+	{
+		check(bFrozen);
+		return UHTCast<To>(DefinitionsByUObject.Find(Object));
+	}
+
+	TSharedRef<FUnrealTypeDefinitionInfo>& operator[](const UObject* Object)
+	{
+		check(bFrozen);
+		return DefinitionsByUObject[Object]; 
+	}
+
 	FUnrealTypeDefinitionInfo& FindChecked(const UObject* Object)
 	{
 		check(bFrozen);
@@ -62,15 +84,46 @@ struct FTypeDefinitionInfoMap : public FFreezableContainer
 		check(TypeDef);
 		return **TypeDef;
 	}
-	TSharedRef<FUnrealTypeDefinitionInfo>* FindByName(const FName Name) { check(bFrozen); return DefinitionsByName.Find(Name); }
+
+	template<typename To>
+	To& FindChecked(const UObject* Object)
+	{
+		check(bFrozen);
+		To* TypeDef = Find<To>(Object);
+		check(TypeDef);
+		return *TypeDef;
+	}
+
+	TSharedRef<FUnrealTypeDefinitionInfo>* FindByName(const FName Name)
+	{
+		check(bFrozen);
+		return DefinitionsByName.Find(Name);
+	}
+
+	template<typename To>
+	To* FindByName(const FName Name)
+	{
+		check(bFrozen);
+		return UHTCast<To>(DefinitionsByName.Find(Name));
+	}
+
 	FUnrealTypeDefinitionInfo& FindByNameChecked(const FName Name)
-	{ 
-		check(bFrozen); 
+	{
+		check(bFrozen);
 		TSharedRef<FUnrealTypeDefinitionInfo>* TypeDef = DefinitionsByName.Find(Name);
 		check(TypeDef);
 		return **TypeDef;
 	}
 
+	template <typename To>
+	To& FindByNameChecked(const FName Name)
+	{
+		check(bFrozen);
+		To* TypeDef = FindByName<To>(Name);
+		check(TypeDef);
+		return *TypeDef;
+	}
+
 	template <typename Lambda>
 	void ForAllTypes(Lambda&& InLambda)
 	{
@@ -93,18 +146,42 @@ struct FTypeDefinitionInfoMap : public FFreezableContainer
 	{
 		DefinitionsByFField.Add(Field, MoveTemp(Definition));
 	}
-	bool Contains(const FField* Field) { return DefinitionsByFField.Contains(Field); }
+
+	bool Contains(const FField* Field) 
+	{
+		return DefinitionsByFField.Contains(Field);
+	}
+
 	TSharedRef<FUnrealTypeDefinitionInfo>* Find(const FField* Field)
 	{
 		return DefinitionsByFField.Find(Field);
 	}
+
+	template <typename To>
+	To* Find(const FField* Field)
+	{
+		return UHTCast<To>(DefinitionsByFField.Find(Field));
+	}
+
 	FUnrealTypeDefinitionInfo& FindChecked(const FField* Field)
 	{
 		TSharedRef<FUnrealTypeDefinitionInfo>* TypeDef = DefinitionsByFField.Find(Field);
 		check(TypeDef);
 		return **TypeDef;
 	}
-	TSharedRef<FUnrealTypeDefinitionInfo>& operator[](const FField* Field) { return DefinitionsByFField[Field]; }
+
+	template <typename To>
+	To& FindChecked(const FField* Field)
+	{
+		To* TypeDef = Find<To>(Field);
+		check(TypeDef);
+		return *TypeDef;
+	}
+
+	TSharedRef<FUnrealTypeDefinitionInfo>& operator[](const FField* Field)
+	{ 
+		return DefinitionsByFField[Field];
+	}
 
 private:
 
diff --git a/Engine/Source/Programs/UnrealHeaderTool/Private/Classes.h b/Engine/Source/Programs/UnrealHeaderTool/Private/Classes.h
--- a/Engine/Source/Programs/UnrealHeaderTool/Private/Classes.h
+++ b/Engine/Source/Programs/UnrealHeaderTool/Private/Classes.h
@@ -3,6 +3,7 @@
 #pragma once
 
 #include "CoreMinimal.h"
+#include "UObject\ObjectMacros.h"
 
 class UPackage;
 class FString;
@@ -32,4 +33,34 @@ public:
 	 * @return               The found class, or NULL if the class was not found.
 	 */
 	static FClass* FindScriptClass(const FString& InClassName, FString* OutErrorMsg = nullptr);
+
+	/**
+	 * Find an optional object.
+	 * @see StaticFindObject()
+	 */
+	template< class T >
+	static inline T* FindObject(UObject* Outer, const TCHAR* Name, bool ExactClass = false)
+	{
+		return (T*)StaticFindObject(T::StaticClass(), Outer, Name, ExactClass);
+	}
+
+	/**
+	 * Find an optional object, relies on the name being unqualified
+	 * @see StaticFindObjectFast()
+	 */
+	template< class T >
+	static inline T* FindObjectFast(UObject* Outer, FName Name, bool ExactClass = false, bool AnyPackage = false, EObjectFlags ExclusiveFlags = RF_NoFlags)
+	{
+		return (T*)StaticFindObjectFast(T::StaticClass(), Outer, Name, ExactClass, AnyPackage, ExclusiveFlags);
+	}
+
+	/**
+	 * Find an optional object, no failure allowed
+	 * @see StaticFindObjectChecked()
+	 */
+	template< class T >
+	static inline T* FindObjectChecked(UObject* Outer, const TCHAR* Name, bool ExactClass = false)
+	{
+		return (T*)StaticFindObjectChecked(T::StaticClass(), Outer, Name, ExactClass);
+	}
 };
