Added new PropertyTypes which in the long run will be where different property types are implemented.

diff --git a/Engine/Source/Programs/UnrealHeaderTool/Private/CodeGenerator.cpp b/Engine/Source/Programs/UnrealHeaderTool/Private/CodeGenerator.cpp
--- a/Engine/Source/Programs/UnrealHeaderTool/Private/CodeGenerator.cpp
+++ b/Engine/Source/Programs/UnrealHeaderTool/Private/CodeGenerator.cpp
@@ -167,19 +167,19 @@ namespace
 	const TCHAR EndEditorOnlyGuard[] = TEXT("#endif //WITH_EDITOR") LINE_TERMINATOR;
 
 	/** Whether or not the given class has any replicated properties. */
-	static bool ClassHasReplicatedProperties(UClass* Class)
+	static bool ClassHasReplicatedProperties(FUnrealClassDefinitionInfo& ClassDef)
 	{
+		UClass* Class = ClassDef.GetClass();
 		if (!Class->HasAnyClassFlags(CLASS_ReplicationDataIsSetUp))
 		{
-			for (TFieldIterator<FProperty> It(Class, EFieldIteratorFlags::ExcludeSuper); It; ++It)
-		{
-			if ((It->PropertyFlags & CPF_Net) != 0)
+			for (FUnrealPropertyDefinitionInfo* PropertyDef : ClassDef.GetProperties())
 			{
-				return true;
+				if ((PropertyDef->GetProperty()->PropertyFlags & CPF_Net) != 0)
+				{
+					return true;
+				}
 			}
 		}
-		}
-
 		return Class->FirstOwnedClassRep < Class->ClassReps.Num();
 	}
 
@@ -216,7 +216,7 @@ namespace
 			}
 			else
 			{
-				FUnrealPropertyDefinitionInfo& PropDef = GTypeDefinitionInfoMap.FindChecked(ClassReps[ClassRepIndex].Property).AsPropertyChecked();
+				FUnrealPropertyDefinitionInfo& PropDef = GTypeDefinitionInfoMap.FindChecked<FUnrealPropertyDefinitionInfo>(ClassReps[ClassRepIndex].Property);
 
 				bAnyStaticArrays = true;
 				ArrayDimBuilder.Logf(TEXT("\t\t%s=%s,\r\n"), *PropertyName, PropDef.GetArrayDimensions());
@@ -262,8 +262,8 @@ namespace
 		const ClassDefinitionRange& ClassRange,
 		const TCHAR* ClassCPPName,
 		const TCHAR* API,
-		FClass* Class,
-		FClass* SuperClass,
+		UClass* Class,
+		UClass* SuperClass,
 		FOutputDevice& Writer,
 		const FUnrealSourceFile& SourceFile,
 		EExportClassOutFlags& OutFlags)
@@ -491,18 +491,13 @@ FString Macroize(const TCHAR* MacroName, FString&& StringToMacroize)
 
 struct FParmsAndReturnProperties
 {
-	FParmsAndReturnProperties()
-		: Return(nullptr)
-	{
-	}
-
 	bool HasParms() const
 	{
 		return Parms.Num() || Return;
 	}
 
-	TArray<FProperty*> Parms;
-	FProperty*         Return;
+	TArray<FUnrealPropertyDefinitionInfo*> Parms;
+	FUnrealPropertyDefinitionInfo* Return = nullptr;
 };
 
 /**
@@ -511,20 +506,19 @@ struct FParmsAndReturnProperties
  * @param  Function The function to get the parameters for.
  * @return An aggregate containing the parameters and return type of that function.
  */
-FParmsAndReturnProperties GetFunctionParmsAndReturn(UFunction* Function)
+FParmsAndReturnProperties GetFunctionParmsAndReturn(FUnrealFunctionDefinitionInfo& FunctionDef)
 {
 	FParmsAndReturnProperties Result;
-	for ( TFieldIterator<FProperty> It(Function); It; ++It)
+	for (FUnrealPropertyDefinitionInfo* PropertyDef : FunctionDef.GetProperties())
 	{
-		FProperty* Field = *It;
-
-		if ((It->PropertyFlags & (CPF_Parm | CPF_ReturnParm)) == CPF_Parm)
+		FProperty* Field = PropertyDef->GetProperty();
+		if ((Field->PropertyFlags & (CPF_Parm | CPF_ReturnParm)) == CPF_Parm)
 		{
-			Result.Parms.Add(Field);
+			Result.Parms.Add(PropertyDef);
 		}
-		else if (It->PropertyFlags & CPF_ReturnParm)
+		else if (Field->PropertyFlags & CPF_ReturnParm)
 		{
-			Result.Return = Field;
+			Result.Return = PropertyDef;
 		}
 	}
 	return Result;
@@ -749,63 +743,25 @@ static FString OutputMetaDataCodeForObject(FOutputDevice& OutDeclaration, FOutpu
 	return Result;
 }
 
-void FNativeClassHeaderGenerator::ExportProperties(FOutputDevice& Out, UStruct* Struct, int32 TextIndent)
+void FNativeClassHeaderGenerator::ExportProperties(FOutputDevice& Out, FUnrealStructDefinitionInfo& StructDef, int32 TextIndent)
 {
-	FProperty*	Previous			= NULL;
-	FProperty*	PreviousNonEditorOnly = NULL;
-	FProperty*	LastInSuper			= NULL;
-	UStruct*	InheritanceSuper	= Struct->GetInheritanceSuper();
-
-	// Find last property in the lowest base class that has any properties
-	UStruct* CurrentSuper = InheritanceSuper;
-	while (LastInSuper == NULL && CurrentSuper)
-	{
-		for( TFieldIterator<FProperty> It(CurrentSuper,EFieldIteratorFlags::ExcludeSuper); It; ++It )
-		{
-			FProperty* Current = *It;
-
-			// Disregard properties with 0 size like functions.
-			if( It.GetStruct() == CurrentSuper && Current->ElementSize )
-			{
-				LastInSuper = Current;
-			}
-		}
-		// go up a layer in the hierarchy
-		CurrentSuper = CurrentSuper->GetSuperStruct();
-	}
-
 	FMacroBlockEmitter WithEditorOnlyData(Out, TEXT("WITH_EDITORONLY_DATA"));
 
 	// Iterate over all properties in this struct.
-	for( TFieldIterator<FProperty> It(Struct, EFieldIteratorFlags::ExcludeSuper); It; ++It )
+	for (FUnrealPropertyDefinitionInfo* PropertyDef : StructDef.GetProperties())
 	{
-		FProperty* Current = *It;
-
-		// Disregard properties with 0 size like functions.
-		if (It.GetStruct() == Struct)
-		{
-			WithEditorOnlyData(Current->IsEditorOnlyProperty());
-
-			// Export property specifiers
-			// Indent code and export CPP text.
-			{
-				FUHTStringBuilder JustPropertyDecl;
+		FProperty* Property = PropertyDef->GetProperty();
 
-				FUnrealPropertyDefinitionInfo& PropDef = GTypeDefinitionInfoMap.FindChecked(Current).AsPropertyChecked();
-				Current->ExportCppDeclaration( JustPropertyDecl, EExportedDeclaration::Member, PropDef.GetArrayDimensions());
-				ApplyAlternatePropertyExportText(*It, JustPropertyDecl, EExportingState::TypeEraseDelegates);
+		WithEditorOnlyData(Property->IsEditorOnlyProperty());
 
-				// Finish up line.
-				Out.Logf(TEXT("%s%s;\r\n"), FCString::Tab(TextIndent + 1), *JustPropertyDecl);
-			}
+		// Export property specifiers
+		// Indent code and export CPP text.
+		FUHTStringBuilder JustPropertyDecl;
+		Property->ExportCppDeclaration( JustPropertyDecl, EExportedDeclaration::Member, PropertyDef->GetArrayDimensions());
+		ApplyAlternatePropertyExportText(Property, JustPropertyDecl, EExportingState::TypeEraseDelegates);
 
-			LastInSuper	= NULL;
-			Previous = Current;
-			if (!Current->IsEditorOnlyProperty())
-			{
-				PreviousNonEditorOnly = Current;
-			}
-		}
+		// Finish up line.
+		Out.Logf(TEXT("%s%s;\r\n"), FCString::Tab(TextIndent + 1), *JustPropertyDecl);
 	}
 }
 
@@ -814,7 +770,7 @@ static FString GNullPtr(TEXT("nullptr"));
 const FString& FNativeClassHeaderGenerator::GetPackageSingletonName(UPackage* Item, TSet<FString>* UniqueCrossModuleReferences)
 {
 	check(Item);
-	FUnrealPackageDefinitionInfo& Package = GTypeDefinitionInfoMap.FindChecked(Item).AsPackageChecked();
+	FUnrealPackageDefinitionInfo& Package = GTypeDefinitionInfoMap.FindChecked<FUnrealPackageDefinitionInfo>(Item);
 	Package.AddCrossModuleReference(UniqueCrossModuleReferences);
 	return Package.GetSingletonName();
 }
@@ -822,7 +778,7 @@ const FString& FNativeClassHeaderGenerator::GetPackageSingletonName(UPackage* It
 const FString& FNativeClassHeaderGenerator::GetPackageSingletonNameFuncAddr(UPackage* Item, TSet<FString>* UniqueCrossModuleReferences)
 {
 	check(Item);
-	FUnrealPackageDefinitionInfo& Package = GTypeDefinitionInfoMap.FindChecked(Item).AsPackageChecked();
+	FUnrealPackageDefinitionInfo& Package = GTypeDefinitionInfoMap.FindChecked<FUnrealPackageDefinitionInfo>(Item);
 	Package.AddCrossModuleReference(UniqueCrossModuleReferences);
 	return Package.GetSingletonNameChopped();
 }
@@ -841,7 +797,7 @@ const FString& FNativeClassHeaderGenerator::GetSingletonNameFuncAddr(UField* Ite
 	}
 	else
 	{
-		FUnrealFieldDefinitionInfo& Field = GTypeDefinitionInfoMap.FindChecked(Item).AsFieldChecked();
+		FUnrealFieldDefinitionInfo& Field = GTypeDefinitionInfoMap.FindChecked<FUnrealFieldDefinitionInfo>(Item);
 		return GetSingletonNameFuncAddr(Field, UniqueCrossModuleReferences, bRequiresValidObject);
 	}
 }
@@ -956,8 +912,9 @@ void FNativeClassHeaderGenerator::GetPropertyTag(FUHTStringBuilder& Out, FProper
 #endif
 }
 
-void FNativeClassHeaderGenerator::PropertyNew(FOutputDevice& DeclOut, FOutputDevice& Out, FReferenceGatherers& OutReferenceGatherers, FProperty* Prop, const TCHAR* OffsetStr, const TCHAR* Name, const TCHAR* DeclSpaces, const TCHAR* Spaces, const TCHAR* SourceStruct) const
+void FNativeClassHeaderGenerator::PropertyNew(FOutputDevice& DeclOut, FOutputDevice& Out, FReferenceGatherers& OutReferenceGatherers, FUnrealPropertyDefinitionInfo& PropertyDef, const TCHAR* OffsetStr, const TCHAR* Name, const TCHAR* DeclSpaces, const TCHAR* Spaces, const TCHAR* SourceStruct) const
 {
+	FProperty*	   Prop                 = PropertyDef.GetProperty();
 	FString        PropName             = CreateUTF8LiteralString(FNativeClassHeaderGenerator::GetOverriddenName(Prop));
 	FString        PropNameDep          = Prop->HasAllPropertyFlags(CPF_Deprecated) ? Prop->GetName() + TEXT("_DEPRECATED") : Prop->GetName();
 	const TCHAR*   FPropertyObjectFlags = FClass::IsOwnedByDynamicType(Prop) ? TEXT("RF_Public|RF_Transient") : TEXT("RF_Public|RF_Transient|RF_MarkAsNative");
@@ -981,8 +938,6 @@ void FNativeClassHeaderGenerator::PropertyNew(FOutputDevice& DeclOut, FOutputDev
 		NameWithoutScope.RightChopInline(ScopeIndex + 2, false);
 	}
 
-	FUnrealPropertyDefinitionInfo& PropDef = GTypeDefinitionInfoMap.FindChecked(Prop).AsPropertyChecked();
-
 	if (FByteProperty* TypedProp = CastField<FByteProperty>(Prop))
 	{
 		DeclOut.Logf(TEXT("%sstatic const UECodeGen_Private::FBytePropertyParams %s;\r\n"), DeclSpaces, *NameWithoutScope);
@@ -1049,7 +1004,7 @@ void FNativeClassHeaderGenerator::PropertyNew(FOutputDevice& DeclOut, FOutputDev
 
 	if (FIntProperty* TypedProp = CastField<FIntProperty>(Prop))
 	{
-		const TCHAR* PropTypeName = PropDef.IsUnsized() ? TEXT("FUnsizedIntPropertyParams") : TEXT("FIntPropertyParams");
+		const TCHAR* PropTypeName = PropertyDef.IsUnsized() ? TEXT("FUnsizedIntPropertyParams") : TEXT("FIntPropertyParams");
 
 		DeclOut.Logf(TEXT("%sstatic const UECodeGen_Private::%s %s;\r\n"), DeclSpaces, PropTypeName, *NameWithoutScope);
 
@@ -1115,7 +1070,7 @@ void FNativeClassHeaderGenerator::PropertyNew(FOutputDevice& DeclOut, FOutputDev
 
 	if (FUInt32Property* TypedProp = CastField<FUInt32Property>(Prop))
 	{
-		const TCHAR* PropTypeName = PropDef.IsUnsized() ? TEXT("FUnsizedFIntPropertyParams") : TEXT("FUInt32PropertyParams");
+		const TCHAR* PropTypeName = PropertyDef.IsUnsized() ? TEXT("FUnsizedFIntPropertyParams") : TEXT("FUInt32PropertyParams");
 
 		DeclOut.Logf(TEXT("%sstatic const UECodeGen_Private::%s %s;\r\n"), DeclSpaces, PropTypeName, *NameWithoutScope);
 
@@ -1522,7 +1477,7 @@ void FNativeClassHeaderGenerator::PropertyNew(FOutputDevice& DeclOut, FOutputDev
 			FPropertyObjectFlags,
 			*ArrayDim,
 			OffsetStr,
-			PropDef.GetAllocatorType() == EAllocatorType::MemoryImage ? TEXT("EArrayPropertyFlags::UsesMemoryImageAllocator") : TEXT("EArrayPropertyFlags::None"),
+			PropertyDef.GetAllocatorType() == EAllocatorType::MemoryImage ? TEXT("EArrayPropertyFlags::UsesMemoryImageAllocator") : TEXT("EArrayPropertyFlags::None"),
 			*MetaDataParams,
 			*PropTag
 		);
@@ -1544,7 +1499,7 @@ void FNativeClassHeaderGenerator::PropertyNew(FOutputDevice& DeclOut, FOutputDev
 			FPropertyObjectFlags,
 			*ArrayDim,
 			OffsetStr,
-			PropDef.GetAllocatorType() == EAllocatorType::MemoryImage ? TEXT("EMapPropertyFlags::UsesMemoryImageAllocator") : TEXT("EMapPropertyFlags::None"),
+			PropertyDef.GetAllocatorType() == EAllocatorType::MemoryImage ? TEXT("EMapPropertyFlags::UsesMemoryImageAllocator") : TEXT("EMapPropertyFlags::None"),
 			*MetaDataParams,
 			*PropTag
 		);
@@ -1748,7 +1703,7 @@ TTuple<FString, FString> FNativeClassHeaderGenerator::OutputProperties(FOutputDe
 
 			WithEditorOnlyMacroEmitter(bRequiresHasEditorOnlyMacro);
 			WithEditorOnlyMacroEmitterDecl(bRequiresHasEditorOnlyMacro);
-			OutputProperty(DeclOut, Out, OutReferenceGatherers, Scope, PropertyNamesAndPointers, Prop, DeclSpaces, Spaces);
+			OutputProperty(DeclOut, Out, OutReferenceGatherers, Scope, PropertyNamesAndPointers, GTypeDefinitionInfoMap.FindChecked<FUnrealPropertyDefinitionInfo>(Prop), DeclSpaces, Spaces);
 		}
 
 		WithEditorOnlyMacroEmitter(bHasAllEditorOnlyDataProperties);
@@ -1758,7 +1713,7 @@ TTuple<FString, FString> FNativeClassHeaderGenerator::OutputProperties(FOutputDe
 
 		for (const FPropertyNamePointerPair& PropNameAndPtr : PropertyNamesAndPointers)
 		{
-			bool bRequiresHasEditorOnlyMacro = IsEditorOnlyDataProperty(PropNameAndPtr.Prop);
+			bool bRequiresHasEditorOnlyMacro = IsEditorOnlyDataProperty(PropNameAndPtr.PropDef->GetProperty());
 
 			WithEditorOnlyMacroEmitter(bRequiresHasEditorOnlyMacro);
 			WithEditorOnlyMacroEmitterDecl(bRequiresHasEditorOnlyMacro);
@@ -1811,53 +1766,73 @@ inline FString GetEventStructParamsName(UObject* Outer, const TCHAR* FunctionNam
 	return Result;
 }
 
-void FNativeClassHeaderGenerator::OutputProperty(FOutputDevice& DeclOut, FOutputDevice& Out, FReferenceGatherers& OutReferenceGatherers, const TCHAR* Scope, TArray<FPropertyNamePointerPair>& PropertyNamesAndPointers, FProperty* Prop, const TCHAR* DeclSpaces, const TCHAR* Spaces) const
+void FNativeClassHeaderGenerator::OutputProperty(FOutputDevice& DeclOut, FOutputDevice& Out, FReferenceGatherers& OutReferenceGatherers, const TCHAR* Scope, TArray<FPropertyNamePointerPair>& PropertyNamesAndPointers, FUnrealPropertyDefinitionInfo& PropertyDef, const TCHAR* DeclSpaces, const TCHAR* Spaces) const
 {
+	FProperty* Prop = PropertyDef.GetProperty();
+
 	// Helper to handle the creation of the underlying properties if they're enum properties
-	auto HandleUnderlyingEnumProperty = [this, &PropertyNamesAndPointers, &DeclOut, &Out, &OutReferenceGatherers, DeclSpaces, Spaces](FProperty* LocalProp, FString&& InOuterName)
+	auto HandleUnderlyingEnumProperty = [this, &PropertyNamesAndPointers, &DeclOut, &Out, &OutReferenceGatherers, DeclSpaces, Spaces](FUnrealPropertyDefinitionInfo& LocalPropDef, FString&& InOuterName)
 	{
+		FProperty* LocalProp = LocalPropDef.GetProperty();
 		if (FEnumProperty* EnumProp = CastField<FEnumProperty>(LocalProp))
 		{
 			FString PropVarName = InOuterName + TEXT("_Underlying");
 
-			PropertyNew(DeclOut, Out, OutReferenceGatherers, EnumProp->UnderlyingProp, TEXT("0"), *PropVarName, DeclSpaces, Spaces);
-			PropertyNamesAndPointers.Emplace(MoveTemp(PropVarName), EnumProp->UnderlyingProp);
+			FUnrealPropertyDefinitionInfo& UnderlyingDef = LocalPropDef.GetValuePropDef();
+			FProperty* Underlying = UnderlyingDef.GetProperty();
+
+			checkSlow(Underlying == EnumProp->UnderlyingProp);
+			PropertyNew(DeclOut, Out, OutReferenceGatherers, UnderlyingDef, TEXT("0"), *PropVarName, DeclSpaces, Spaces);
+			PropertyNamesAndPointers.Emplace(MoveTemp(PropVarName), UnderlyingDef);
 		}
 
-		PropertyNamesAndPointers.Emplace(MoveTemp(InOuterName), LocalProp);
+		PropertyNamesAndPointers.Emplace(MoveTemp(InOuterName), LocalPropDef);
 	};
 
 	if (FArrayProperty* ArrayProperty = CastField<FArrayProperty>(Prop))
 	{
-		FString InnerVariableName = FString::Printf(TEXT("%sNewProp_%s_Inner"), Scope, *ArrayProperty->Inner->GetName());
+		FUnrealPropertyDefinitionInfo& InnerDef = PropertyDef.GetValuePropDef();
+		FProperty* Inner = InnerDef.GetProperty();
+
+		FString InnerVariableName = FString::Printf(TEXT("%sNewProp_%s_Inner"), Scope, *Inner->GetName());
+
+		checkSlow(PropertyDef.GetValuePropDef().GetProperty() == ArrayProperty->Inner);
 
-		HandleUnderlyingEnumProperty(ArrayProperty->Inner, CopyTemp(InnerVariableName));
-		PropertyNew(DeclOut, Out, OutReferenceGatherers, ArrayProperty->Inner, TEXT("0"), *InnerVariableName, DeclSpaces, Spaces);
+		HandleUnderlyingEnumProperty(InnerDef, CopyTemp(InnerVariableName));
+		PropertyNew(DeclOut, Out, OutReferenceGatherers, PropertyDef.GetValuePropDef(), TEXT("0"), *InnerVariableName, DeclSpaces, Spaces);
 	}
 
 	else if (FMapProperty* MapProperty = CastField<FMapProperty>(Prop))
 	{
-		FProperty* Key   = MapProperty->KeyProp;
-		FProperty* Value = MapProperty->ValueProp;
+		FUnrealPropertyDefinitionInfo& KeyDef = PropertyDef.GetKeyPropDef();
+		FProperty* Key = KeyDef.GetProperty();
+		FUnrealPropertyDefinitionInfo& ValueDef = PropertyDef.GetValuePropDef();
+		FProperty* Value = ValueDef.GetProperty();
+
+		checkSlow(PropertyDef.GetValuePropDef().GetProperty() == MapProperty->ValueProp);
+		checkSlow(PropertyDef.GetKeyPropDef().GetProperty() == MapProperty->KeyProp);
 
 		FString KeyVariableName   = FString::Printf(TEXT("%sNewProp_%s_KeyProp"), Scope, *Key->GetName());
 		FString ValueVariableName = FString::Printf(TEXT("%sNewProp_%s_ValueProp"), Scope, *Value->GetName());
 
-		HandleUnderlyingEnumProperty(Value, CopyTemp(ValueVariableName));
-		PropertyNew(DeclOut, Out, OutReferenceGatherers, Value, TEXT("1"), *ValueVariableName, DeclSpaces, Spaces);
+		HandleUnderlyingEnumProperty(ValueDef, CopyTemp(ValueVariableName));
+		PropertyNew(DeclOut, Out, OutReferenceGatherers, ValueDef, TEXT("1"), *ValueVariableName, DeclSpaces, Spaces);
 
-		HandleUnderlyingEnumProperty(Key, CopyTemp(KeyVariableName));
-		PropertyNew(DeclOut, Out, OutReferenceGatherers, Key, TEXT("0"), *KeyVariableName, DeclSpaces, Spaces);
+		HandleUnderlyingEnumProperty(KeyDef, CopyTemp(KeyVariableName));
+		PropertyNew(DeclOut, Out, OutReferenceGatherers, KeyDef, TEXT("0"), *KeyVariableName, DeclSpaces, Spaces);
 	}
 
 	else if (FSetProperty* SetProperty = CastField<FSetProperty>(Prop))
 	{
-		FProperty* Inner = SetProperty->ElementProp;
+		FUnrealPropertyDefinitionInfo& InnerDef = PropertyDef.GetValuePropDef();
+		FProperty* Inner = InnerDef.GetProperty();
 
 		FString ElementVariableName = FString::Printf(TEXT("%sNewProp_%s_ElementProp"), Scope, *Inner->GetName());
 
-		HandleUnderlyingEnumProperty(Inner, CopyTemp(ElementVariableName));
-		PropertyNew(DeclOut, Out, OutReferenceGatherers, Inner, TEXT("0"), *ElementVariableName, DeclSpaces, Spaces);
+		checkSlow(PropertyDef.GetValuePropDef().GetProperty() == SetProperty->ElementProp);
+
+		HandleUnderlyingEnumProperty(InnerDef, CopyTemp(ElementVariableName));
+		PropertyNew(DeclOut, Out, OutReferenceGatherers, InnerDef, TEXT("0"), *ElementVariableName, DeclSpaces, Spaces);
 	}
 
 	{
@@ -1891,8 +1866,8 @@ void FNativeClassHeaderGenerator::OutputProperty(FOutputDevice& DeclOut, FOutput
 
 		FString PropMacroOuterClass = FString::Printf(TEXT("STRUCT_OFFSET(%s, %s)"), *SourceStruct, *PropName);
 
-		HandleUnderlyingEnumProperty(Prop, CopyTemp(PropVariableName));
-		PropertyNew(DeclOut, Out, OutReferenceGatherers, Prop, *PropMacroOuterClass, *PropVariableName, DeclSpaces, Spaces, *SourceStruct);
+		HandleUnderlyingEnumProperty(PropertyDef, CopyTemp(PropVariableName));
+		PropertyNew(DeclOut, Out, OutReferenceGatherers, PropertyDef, *PropMacroOuterClass, *PropVariableName, DeclSpaces, Spaces, *SourceStruct);
 	}
 }
 
@@ -1916,10 +1891,11 @@ static bool IsAlwaysAccessible(UScriptStruct* Script)
 	return Result;
 }
 
-static void FindNoExportStructsRecursive(TArray<UScriptStruct*>& Structs, UStruct* Start)
+static void FindNoExportStructsRecursive(TArray<FUnrealScriptStructDefinitionInfo*>& StructDefs, FUnrealStructDefinitionInfo* StartDef)
 {
-	while (Start)
+	for (; StartDef; StartDef = StartDef->GetSuperClassInfo().Struct)
 	{
+		UStruct* Start = StartDef->GetStruct();
 		if (UScriptStruct* StartScript = Cast<UScriptStruct>(Start))
 		{
 			if (StartScript->StructFlags & STRUCT_Native)
@@ -1930,51 +1906,52 @@ static void FindNoExportStructsRecursive(TArray<UScriptStruct*>& Structs, UStruc
 			if (!IsAlwaysAccessible(StartScript)) // these are a special cases that already exists and if wrong if exported naively
 			{
 				// this will topologically sort them in reverse order
-				Structs.Remove(StartScript);
-				Structs.Add(StartScript);
+				StructDefs.Remove(&StartDef->AsScriptStructChecked());
+				StructDefs.Add(&StartDef->AsScriptStructChecked());
 			}
 		}
 
-		for (FProperty* Prop : TFieldRange<FProperty>(Start, EFieldIteratorFlags::ExcludeSuper))
+		//ETSTODO - Fixme
+		for (FUnrealPropertyDefinitionInfo* PropertyDef : StartDef->GetProperties())
 		{
+			FProperty* Prop = PropertyDef->GetProperty();
 			if (FStructProperty* StructProp = CastField<FStructProperty>(Prop))
 			{
-				FindNoExportStructsRecursive(Structs, StructProp->Struct);
+				FindNoExportStructsRecursive(StructDefs, &GTypeDefinitionInfoMap.FindChecked<FUnrealStructDefinitionInfo>(StructProp->Struct));
 			}
 			else if (FArrayProperty* ArrayProp = CastField<FArrayProperty>(Prop))
 			{
 				if (FStructProperty* InnerStructProp = CastField<FStructProperty>(ArrayProp->Inner))
 				{
-					FindNoExportStructsRecursive(Structs, InnerStructProp->Struct);
+					FindNoExportStructsRecursive(StructDefs, &GTypeDefinitionInfoMap.FindChecked<FUnrealStructDefinitionInfo>(InnerStructProp->Struct));
 				}
 			}
 			else if (FMapProperty* MapProp = CastField<FMapProperty>(Prop))
 			{
 				if (FStructProperty* KeyStructProp = CastField<FStructProperty>(MapProp->KeyProp))
 				{
-					FindNoExportStructsRecursive(Structs, KeyStructProp->Struct);
+					FindNoExportStructsRecursive(StructDefs, &GTypeDefinitionInfoMap.FindChecked<FUnrealStructDefinitionInfo>(KeyStructProp->Struct));
 				}
 				if (FStructProperty* ValueStructProp = CastField<FStructProperty>(MapProp->ValueProp))
 				{
-					FindNoExportStructsRecursive(Structs, ValueStructProp->Struct);
+					FindNoExportStructsRecursive(StructDefs, &GTypeDefinitionInfoMap.FindChecked<FUnrealStructDefinitionInfo>(ValueStructProp->Struct));
 				}
 			}
 			else if (FSetProperty* SetProp = CastField<FSetProperty>(Prop))
 			{
 				if (FStructProperty* ElementStructProp = CastField<FStructProperty>(SetProp->ElementProp))
 				{
-					FindNoExportStructsRecursive(Structs, ElementStructProp->Struct);
+					FindNoExportStructsRecursive(StructDefs, &GTypeDefinitionInfoMap.FindChecked<FUnrealStructDefinitionInfo>(ElementStructProp->Struct));
 				}
 			}
 		}
-		Start = Start->GetSuperStruct();
 	}
 }
 
-static TArray<UScriptStruct*> FindNoExportStructs(UStruct* Start)
+static TArray<FUnrealScriptStructDefinitionInfo*> FindNoExportStructs(FUnrealStructDefinitionInfo* StartDef)
 {
-	TArray<UScriptStruct*> Result;
-	FindNoExportStructsRecursive(Result, Start);
+	TArray<FUnrealScriptStructDefinitionInfo*> Result;
+	FindNoExportStructsRecursive(Result, StartDef);
 
 	// These come out in reverse order of topology so reverse them
 	Algo::Reverse(Result);
@@ -2001,14 +1978,14 @@ void FNativeClassHeaderGenerator::ExportGeneratedPackageInitCode(FOutputDevice&
 		{
 			return !bADel;
 		}
-		const FString& AName = GTypeDefinitionInfoMap.FindChecked(A).AsFieldChecked().GetSingletonName(true);
-		const FString& BName = GTypeDefinitionInfoMap.FindChecked(B).AsFieldChecked().GetSingletonName(true);
+		const FString& AName = GTypeDefinitionInfoMap.FindChecked<FUnrealFieldDefinitionInfo>(A).GetSingletonName(true);
+		const FString& BName = GTypeDefinitionInfoMap.FindChecked<FUnrealFieldDefinitionInfo>(B).GetSingletonName(true);
 		return AName < BName;
 	});
 
 	for (UField* ScriptType : Singletons)
 	{
-		Out.Log(GTypeDefinitionInfoMap.FindChecked(ScriptType).AsFieldChecked().GetExternDecl(true));
+		Out.Log(GTypeDefinitionInfoMap.FindChecked<FUnrealFieldDefinitionInfo>(ScriptType).GetExternDecl(true));
 	}
 
 	FOutputDeviceNull OutputDeviceNull;
@@ -2027,7 +2004,7 @@ void FNativeClassHeaderGenerator::ExportGeneratedPackageInitCode(FOutputDevice&
 		Out.Logf(TEXT("\t\t\tstatic UObject* (*const SingletonFuncArray[])() = {\r\n"));
 		for (UField* ScriptType : Singletons)
 		{
-			const FString& Name = GTypeDefinitionInfoMap.FindChecked(ScriptType).AsFieldChecked().GetSingletonNameChopped(true);
+			const FString& Name = GTypeDefinitionInfoMap.FindChecked<FUnrealFieldDefinitionInfo>(ScriptType).GetSingletonNameChopped(true);
 
 			Out.Logf(TEXT("\t\t\t\t(UObject* (*)())%s,\r\n"), *Name);
 		}
@@ -2057,14 +2034,13 @@ void FNativeClassHeaderGenerator::ExportGeneratedPackageInitCode(FOutputDevice&
 	Out.Logf(TEXT("\t}\r\n"));
 }
 
-void FNativeClassHeaderGenerator::ExportNativeGeneratedInitCode(FOutputDevice& Out, FOutputDevice& OutDeclarations, FReferenceGatherers& OutReferenceGatherers, const FUnrealSourceFile& SourceFile, FClass* Class, FUHTStringBuilder& OutFriendText) const
+void FNativeClassHeaderGenerator::ExportNativeGeneratedInitCode(FOutputDevice& Out, FOutputDevice& OutDeclarations, FReferenceGatherers& OutReferenceGatherers, const FUnrealSourceFile& SourceFile, FUnrealClassDefinitionInfo& ClassDef, FUHTStringBuilder& OutFriendText) const
 {
+	UClass* Class = ClassDef.GetClass();
 	check(!OutFriendText.Len());
 
 	UE_CLOG(Class->ClassGeneratedBy, LogCompile, Fatal, TEXT("For intrinsic and compiled-in classes, ClassGeneratedBy should always be null"));
 
-	FUnrealClassDefinitionInfo& ClassDef = GTypeDefinitionInfoMap.FindChecked(Class).AsClassChecked();
-
 	const bool   bIsNoExport  = Class->HasAnyClassFlags(CLASS_NoExport);
 	const bool   bIsDynamic   = FClass::IsDynamic(Class);
 	const FString ClassNameCPP = FNameLookupCPP::GetNameCPP(Class);
@@ -2072,10 +2048,11 @@ void FNativeClassHeaderGenerator::ExportNativeGeneratedInitCode(FOutputDevice& O
 	const FString& ApiString = GetAPIString();
 
 	TSet<FName> AlreadyIncludedNames;
-	TArray<UFunction*> FunctionsToExport;
+	TArray<FUnrealFunctionDefinitionInfo*> FunctionsToExport;
 	bool bAllEditorOnlyFunctions = true;
-	for (UFunction* LocalFunc : TFieldRange<UFunction>(Class,EFieldIteratorFlags::ExcludeSuper))
+	for (FUnrealFunctionDefinitionInfo* LocalFuncDef : ClassDef.GetFunctions())
 	{
+		UFunction* LocalFunc = LocalFuncDef->GetFunction();
 		FName TrueName = FNativeClassHeaderGenerator::GetOverriddenFName(LocalFunc);
 		bool bAlreadyIncluded = false;
 		AlreadyIncludedNames.Add(TrueName, &bAlreadyIncluded);
@@ -2092,11 +2069,11 @@ void FNativeClassHeaderGenerator::ExportNativeGeneratedInitCode(FOutputDevice& O
 		{
 			bAllEditorOnlyFunctions &= LocalFunc->HasAnyFunctionFlags(FUNC_EditorOnly);
 		}
-		FunctionsToExport.Add(LocalFunc);
+		FunctionsToExport.Add(LocalFuncDef);
 	}
 
 	// Sort the list of functions
-	FunctionsToExport.Sort();
+	Algo::SortBy(FunctionsToExport, [](FUnrealFunctionDefinitionInfo* Obj) { return Obj->GetFunction()->GetName(); });
 
 	FUHTStringBuilder GeneratedClassRegisterFunctionText;
 
@@ -2122,7 +2099,7 @@ void FNativeClassHeaderGenerator::ExportNativeGeneratedInitCode(FOutputDevice& O
 		FString StaticsStructName = ClassDef.GetSingletonNameChopped(true) + TEXT("_Statics");
 
 		OutFriendText.Logf(TEXT("\tfriend struct %s;\r\n"), *StaticsStructName);
-		OutDeclarations.Log(GTypeDefinitionInfoMap.FindChecked(Class).AsFieldChecked().GetExternDecl(true));
+		OutDeclarations.Log(GTypeDefinitionInfoMap.FindChecked<FUnrealClassDefinitionInfo>(Class).GetExternDecl(true));
 
 		GeneratedClassRegisterFunctionText.Logf(TEXT("\tstruct %s\r\n"), *StaticsStructName);
 		GeneratedClassRegisterFunctionText.Logf(TEXT("\t{\r\n"));
@@ -2133,14 +2110,14 @@ void FNativeClassHeaderGenerator::ExportNativeGeneratedInitCode(FOutputDevice& O
 		UClass* SuperClass = Class->GetSuperClass();
 		if (SuperClass && SuperClass != Class)
 		{
-			FUnrealClassDefinitionInfo& SuperClassDef = GTypeDefinitionInfoMap.FindChecked(SuperClass).AsClassChecked();
+			FUnrealClassDefinitionInfo& SuperClassDef = GTypeDefinitionInfoMap.FindChecked<FUnrealClassDefinitionInfo>(SuperClass);
 			SuperClassDef.AddCrossModuleReference(OutReferenceGatherers.UniqueCrossModuleReferences, true);
 			OutDeclarations.Log(SuperClassDef.GetExternDecl(true));
 			Singletons.Logf(TEXT("\t\t(UObject* (*)())%s,\r\n"), *SuperClassDef.GetSingletonNameChopped(true));
 		}
 		if (!bIsDynamic)
 		{
-			FUnrealPackageDefinitionInfo& ClassPackageDef = GTypeDefinitionInfoMap.FindChecked(Class->GetOutermost()).AsPackageChecked();
+			FUnrealPackageDefinitionInfo& ClassPackageDef = GTypeDefinitionInfoMap.FindChecked<FUnrealPackageDefinitionInfo>(Class->GetOutermost());
 			PackageDef.AddCrossModuleReference(OutReferenceGatherers.UniqueCrossModuleReferences);
 			OutDeclarations.Logf(TEXT("\t%s_API UPackage* %s;\r\n"), *ApiString, *ClassPackageDef.GetSingletonName());
 			Singletons.Logf(TEXT("\t\t(UObject* (*)())%s,\r\n"), *ClassPackageDef.GetSingletonNameChopped());
@@ -2176,24 +2153,23 @@ void FNativeClassHeaderGenerator::ExportNativeGeneratedInitCode(FOutputDevice& O
 			StaticDefinitions.Log(BEGIN_WRAP_EDITOR_ONLY(bAllEditorOnlyFunctions));
 			StaticDefinitions.Logf(TEXT("\tconst FClassFunctionLinkInfo %s::FuncInfo[] = {\r\n"), *StaticsStructName);
 
-			for (UFunction* Function : FunctionsToExport)
+			for (FUnrealFunctionDefinitionInfo* FunctionDef : FunctionsToExport)
 			{
-				FUnrealFunctionDefinitionInfo& FuncDef = GTypeDefinitionInfoMap.FindChecked(Function).AsFunctionChecked();
-
+				UFunction* Function = FunctionDef->GetFunction();
 				const bool bIsEditorOnlyFunction = Function->HasAnyFunctionFlags(FUNC_EditorOnly);
 
 				if (!Function->IsA<UDelegateFunction>())
 				{
-					ExportFunction(Out, OutReferenceGatherers, SourceFile, Function, bIsNoExport);
+					ExportFunction(Out, OutReferenceGatherers, SourceFile, *FunctionDef, bIsNoExport);
 				}
 
 				FUHTStringBuilder FuncHashTag;
-				FuncDef.GetHashTag(FuncHashTag);
+				FunctionDef->GetHashTag(FuncHashTag);
 
 				StaticDefinitions.Logf(
 					TEXT("%s\t\t{ &%s, %s },%s\r\n%s"),
 					BEGIN_WRAP_EDITOR_ONLY(bIsEditorOnlyFunction),
-					*GetSingletonNameFuncAddr(FuncDef, OutReferenceGatherers.UniqueCrossModuleReferences),
+					*GetSingletonNameFuncAddr(*FunctionDef, OutReferenceGatherers.UniqueCrossModuleReferences),
 					*FNativeClassHeaderGenerator::GetUTF8OverriddenNameForLiteral(Function),
 					*FuncHashTag,
 					END_WRAP_EDITOR_ONLY(bIsEditorOnlyFunction)
@@ -2232,9 +2208,9 @@ void FNativeClassHeaderGenerator::ExportNativeGeneratedInitCode(FOutputDevice& O
 		FString MetaDataParams = OutputMetaDataCodeForObject(GeneratedClassRegisterFunctionText, StaticDefinitions, Class, *FString::Printf(TEXT("%s::Class_MetaDataParams"), *StaticsStructName), TEXT("\t\t"), TEXT("\t"));
 
 		TArray<FProperty*> Props;
-		for (FProperty* Prop : TFieldRange<FProperty>(Class, EFieldIteratorFlags::ExcludeSuper))
+		for (FUnrealPropertyDefinitionInfo* PropertyDef : ClassDef.GetProperties())
 		{
-			Props.Add(Prop);
+			Props.Add(PropertyDef->GetProperty());
 		}
 
 		TTuple<FString, FString> PropertyRange = OutputProperties(GeneratedClassRegisterFunctionText, StaticDefinitions, OutReferenceGatherers, *FString::Printf(TEXT("%s::"), *StaticsStructName), Props, TEXT("\t\t"), TEXT("\t"));
@@ -2249,7 +2225,7 @@ void FNativeClassHeaderGenerator::ExportNativeGeneratedInitCode(FOutputDevice& O
 			for (const FImplementedInterface& Inter : Class->Interfaces)
 			{
 				check(Inter.Class);
-				FUnrealFieldDefinitionInfo& InterClassDef = GTypeDefinitionInfoMap.FindChecked(Inter.Class).AsFieldChecked();
+				FUnrealFieldDefinitionInfo& InterClassDef = GTypeDefinitionInfoMap.FindChecked<FUnrealFieldDefinitionInfo>(Inter.Class);
 				InterClassDef.AddCrossModuleReference(OutReferenceGatherers.UniqueCrossModuleReferences, false);
 				FString OffsetString;
 				if (Inter.PointerOffset)
@@ -2379,7 +2355,7 @@ void FNativeClassHeaderGenerator::ExportNativeGeneratedInitCode(FOutputDevice& O
 	// Append base class' hash at the end of the generated code, this will force update derived classes
 	// when base class changes during hot-reload.
 	uint32 BaseClassHash = 0;
-	FClass* SuperClass = Class->GetSuperClass();
+	UClass* SuperClass = Class->GetSuperClass();
 	if (SuperClass && !SuperClass->HasAnyClassFlags(CLASS_Intrinsic))
 	{
 		BaseClassHash = GTypeDefinitionInfoMap.FindChecked(SuperClass).GetHash();
@@ -2392,14 +2368,16 @@ void FNativeClassHeaderGenerator::ExportNativeGeneratedInitCode(FOutputDevice& O
 
 	for (const FString& SparseClassDataString : SparseClassDataTypes)
 	{
-		UScriptStruct* SparseClassDataStruct = FindObjectSafe<UScriptStruct>(ANY_PACKAGE, *SparseClassDataString);
+		UScriptStruct* SparseClassDataStruct = FClasses::FindObject<UScriptStruct>(ANY_PACKAGE, *SparseClassDataString);
 		if (!SparseClassDataStruct)
 		{
 			continue;
 		}
+		FUnrealScriptStructDefinitionInfo& SparseScriptStructDef = GTypeDefinitionInfoMap.FindChecked<FUnrealScriptStructDefinitionInfo>(SparseClassDataStruct);
 		GeneratedClassRegisterFunctionText.Logf(TEXT("%s\r\n"), *SparseClassDataStruct->GetName());
-		for (FProperty* Child : TFieldRange<FProperty>(SparseClassDataStruct))
+		for (FUnrealPropertyDefinitionInfo* ChildDef : TUHTFieldRange<FUnrealPropertyDefinitionInfo>(SparseScriptStructDef))
 		{
+			FProperty* Child = ChildDef->GetProperty();
 			GeneratedClassRegisterFunctionText.Logf(TEXT("%s %s\r\n"), *Child->GetCPPType(), *Child->GetNameCPP());
 		}
 	}
@@ -2437,7 +2415,7 @@ void FNativeClassHeaderGenerator::ExportNativeGeneratedInitCode(FOutputDevice& O
 			*InitSearchableValuesFunctionParam);
 	}
 
-	if (ClassHasReplicatedProperties(Class))
+	if (ClassHasReplicatedProperties(ClassDef))
 	{
 		Out.Logf(TEXT(
 			"\r\n"
@@ -2478,10 +2456,10 @@ void FNativeClassHeaderGenerator::ExportNativeGeneratedInitCode(FOutputDevice& O
 	}
 }
 
-void FNativeClassHeaderGenerator::ExportFunction(FOutputDevice& Out, FReferenceGatherers& OutReferenceGatherers, const FUnrealSourceFile& SourceFile, UFunction* Function, bool bIsNoExport) const
+void FNativeClassHeaderGenerator::ExportFunction(FOutputDevice& Out, FReferenceGatherers& OutReferenceGatherers, const FUnrealSourceFile& SourceFile, FUnrealFunctionDefinitionInfo& FunctionDef, bool bIsNoExport) const
 {
-	FUnrealFunctionDefinitionInfo& FuncDef = GTypeDefinitionInfoMap.FindChecked(Function).AsFunctionChecked();
-	FuncDef.AddCrossModuleReference(OutReferenceGatherers.UniqueCrossModuleReferences, true);
+	UFunction* Function = FunctionDef.GetFunction();
+	FunctionDef.AddCrossModuleReference(OutReferenceGatherers.UniqueCrossModuleReferences, true);
 
 	UFunction* SuperFunction = Function->GetSuperFunction();
 
@@ -2489,8 +2467,8 @@ void FNativeClassHeaderGenerator::ExportFunction(FOutputDevice& Out, FReferenceG
 
 	bool bIsDelegate = Function->HasAnyFunctionFlags(FUNC_Delegate);
 
-	const FString& SingletonName = FuncDef.GetSingletonName(true);
-	FString StaticsStructName = FuncDef.GetSingletonNameChopped(true) + TEXT("_Statics");
+	const FString& SingletonName = FunctionDef.GetSingletonName(true);
+	FString StaticsStructName = FunctionDef.GetSingletonNameChopped(true) + TEXT("_Statics");
 
 	FUHTStringBuilder CurrentFunctionText;
 	FUHTStringBuilder StaticDefinitions;
@@ -2507,13 +2485,13 @@ void FNativeClassHeaderGenerator::ExportFunction(FOutputDevice& Out, FReferenceG
 	bool bParamsInStatic = bIsNoExport || !(Function->FunctionFlags & FUNC_Event); // non-events do not export a params struct, so lets do that locally for offset determination
 	if (bParamsInStatic)
 	{
-		TArray<UScriptStruct*> Structs = FindNoExportStructs(Function);
-		for (UScriptStruct* Struct : Structs)
+		TArray<FUnrealScriptStructDefinitionInfo*> StructDefs = FindNoExportStructs(&FunctionDef);
+		for (FUnrealScriptStructDefinitionInfo* StructDef : StructDefs)
 		{
-			ExportMirrorsForNoexportStruct(CurrentFunctionText, Struct, /*Indent=*/ 2);
+			ExportMirrorsForNoexportStruct(CurrentFunctionText, *StructDef, /*Indent=*/ 2);
 		}
 
-		ExportEventParm(CurrentFunctionText, OutReferenceGatherers.ForwardDeclarations, Function, /*Indent=*/ 2, /*bOutputConstructor=*/ false, EExportingState::TypeEraseDelegates);
+		ExportEventParm(CurrentFunctionText, OutReferenceGatherers.ForwardDeclarations, FunctionDef, /*Indent=*/ 2, /*bOutputConstructor=*/ false, EExportingState::TypeEraseDelegates);
 	}
 
 	UField* FieldOuter = Cast<UField>(Function->GetOuter());
@@ -2530,7 +2508,10 @@ void FNativeClassHeaderGenerator::ExportFunction(FOutputDevice& Out, FReferenceG
 	}
 	
 	TArray<FProperty*> Props;
-	Algo::Copy(TFieldRange<FProperty>(Function, EFieldIteratorFlags::ExcludeSuper), Props, Algo::NoRef);
+	for (FUnrealPropertyDefinitionInfo* PropertyDef : FunctionDef.GetProperties())
+	{
+		Props.Add(PropertyDef->GetProperty());
+	}
 
 	FString StructureSize;
 	if (Props.Num())
@@ -2565,7 +2546,7 @@ void FNativeClassHeaderGenerator::ExportFunction(FOutputDevice& Out, FReferenceG
 
 	TTuple<FString, FString> PropertyRange = OutputProperties(CurrentFunctionText, StaticDefinitions, OutReferenceGatherers, *FString::Printf(TEXT("%s::"), *StaticsStructName), Props, TEXT("\t\t"), TEXT("\t"));
 
-	const FFuncInfo&     FunctionData = FuncDef.GetFunctionData();
+	const FFuncInfo&     FunctionData = FunctionDef.GetFunctionData();
 	const bool           bIsNet       = !!(FunctionData.FunctionFlags & (FUNC_NetRequest | FUNC_NetResponse));
 
 	FString MetaDataParams = OutputMetaDataCodeForObject(CurrentFunctionText, StaticDefinitions, Function, *FString::Printf(TEXT("%s::Function_MetaDataParams"), *StaticsStructName), TEXT("\t\t"), TEXT("\t"));
@@ -2621,12 +2602,13 @@ void FNativeClassHeaderGenerator::ExportFunction(FOutputDevice& Out, FReferenceG
 	}
 
 	uint32 FunctionHash = GenerateTextHash(*CurrentFunctionText);
-	FuncDef.SetHash(FunctionHash);
+	FunctionDef.SetHash(FunctionHash);
 	Out.Log(CurrentFunctionText);
 }
 
-void FNativeClassHeaderGenerator::ExportNatives(FOutputDevice& Out, FClass* Class)
+void FNativeClassHeaderGenerator::ExportNatives(FOutputDevice& Out, FUnrealClassDefinitionInfo& ClassDef)
 {
+	UClass* Class = ClassDef.GetClass();
 	const FString ClassCPPName = FNameLookupCPP::GetNameCPP(Class);
 	FString TypeName = Class->HasAnyClassFlags(CLASS_Interface) ? FString::Printf(TEXT("I%s"), *Class->GetName()) : ClassCPPName;
 
@@ -2637,8 +2619,9 @@ void FNativeClassHeaderGenerator::ExportNatives(FOutputDevice& Out, FClass* Clas
 		bool bAllEditorOnly = true;
 
 		TArray<TTuple<UFunction*, FString>> NamedFunctionsToExport;
-		for (UFunction* Function : TFieldRange<UFunction>(Class, EFieldIteratorFlags::ExcludeSuper))
+		for (FUnrealFunctionDefinitionInfo* FunctionDef : ClassDef.GetFunctions())
 		{
+			UFunction* Function = FunctionDef->GetFunction();
 			if ((Function->FunctionFlags & (FUNC_Native | FUNC_NetRequest)) == FUNC_Native)
 			{
 				FString OverriddenName = FNativeClassHeaderGenerator::GetUTF8OverriddenNameForLiteral(Function);
@@ -2682,44 +2665,35 @@ void FNativeClassHeaderGenerator::ExportNatives(FOutputDevice& Out, FClass* Clas
 		}
 	}
 
-	for (UScriptStruct* Struct : TFieldRange<UScriptStruct>(Class, EFieldIteratorFlags::ExcludeSuper))
-	{
-		if (Struct->StructFlags & STRUCT_Native)
-		{
-			const FString StructName = Struct->GetName();
-			Out.Logf( TEXT("\t\tUScriptStruct::DeferCppStructOps<%s%s>(FName(TEXT(\"%s\")));\r\n"), Struct->GetPrefixCPP(), *StructName, *StructName);
-		}
-	}
-
 	Out.Logf(TEXT("\t}\r\n"));
 }
 
-void FNativeClassHeaderGenerator::ExportInterfaceCallFunctions(FOutputDevice& OutCpp, FUHTStringBuilder& Out, FReferenceGatherers& OutReferenceGatherers, const TArray<UFunction*>& CallbackFunctions, const TCHAR* ClassName) const
+void FNativeClassHeaderGenerator::ExportInterfaceCallFunctions(FOutputDevice& OutCpp, FUHTStringBuilder& Out, FReferenceGatherers& OutReferenceGatherers, const TArray<FUnrealFunctionDefinitionInfo*>& CallbackFunctions, const TCHAR* ClassName) const
 {
 	const FString& APIString = GetAPIString();
 
-	for (UFunction* Function : CallbackFunctions)
+	for (FUnrealFunctionDefinitionInfo* FunctionDef : CallbackFunctions)
 	{
+		UFunction* Function = FunctionDef->GetFunction();
 		FString FunctionName = Function->GetName();
 
-		FUnrealFunctionDefinitionInfo& FuncDef = GTypeDefinitionInfoMap.FindChecked(Function).AsFunctionChecked();
-		const FFuncInfo& FunctionData = FuncDef.GetFunctionData();
+		const FFuncInfo& FunctionData = FunctionDef->GetFunctionData();
 		const TCHAR* ConstQualifier = FunctionData.FunctionReference->HasAllFunctionFlags(FUNC_Const) ? TEXT("const ") : TEXT("");
 		FString ExtraParam = FString::Printf(TEXT("%sUObject* O"), ConstQualifier);
 
-		ExportNativeFunctionHeader(Out, OutReferenceGatherers.ForwardDeclarations, FunctionData, EExportFunctionType::Interface, EExportFunctionHeaderStyle::Declaration, *ExtraParam, *APIString);
+		ExportNativeFunctionHeader(Out, OutReferenceGatherers.ForwardDeclarations, *FunctionDef, FunctionData, EExportFunctionType::Interface, EExportFunctionHeaderStyle::Declaration, *ExtraParam, *APIString);
 		Out.Logf( TEXT(";") LINE_TERMINATOR );
 
 		FString FunctionNameName = FString::Printf(TEXT("NAME_%s_%s"), *FNameLookupCPP::GetNameCPP(CastChecked<UStruct>(Function->GetOuter())), *FunctionName);
 		OutCpp.Logf(TEXT("\tstatic FName %s = FName(TEXT(\"%s\"));") LINE_TERMINATOR, *FunctionNameName, *GetOverriddenFName(Function).ToString());
 
-		ExportNativeFunctionHeader(OutCpp, OutReferenceGatherers.ForwardDeclarations, FunctionData, EExportFunctionType::Interface, EExportFunctionHeaderStyle::Definition, *ExtraParam, *APIString);
+		ExportNativeFunctionHeader(OutCpp, OutReferenceGatherers.ForwardDeclarations, *FunctionDef, FunctionData, EExportFunctionType::Interface, EExportFunctionHeaderStyle::Definition, *ExtraParam, *APIString);
 		OutCpp.Logf( LINE_TERMINATOR TEXT("\t{") LINE_TERMINATOR );
 
 		OutCpp.Logf(TEXT("\t\tcheck(O != NULL);") LINE_TERMINATOR);
 		OutCpp.Logf(TEXT("\t\tcheck(O->GetClass()->ImplementsInterface(U%s::StaticClass()));") LINE_TERMINATOR, ClassName);
 
-		FParmsAndReturnProperties Parameters = GetFunctionParmsAndReturn(FunctionData.FunctionReference);
+		FParmsAndReturnProperties Parameters = GetFunctionParmsAndReturn(*FunctionDef);
 
 		// See if we need to create Parms struct
 		const bool bHasParms = Parameters.HasParms();
@@ -2734,17 +2708,18 @@ void FNativeClassHeaderGenerator::ExportInterfaceCallFunctions(FOutputDevice& Ou
 		OutCpp.Log(TEXT("\t\t{") LINE_TERMINATOR);
 
 		// code to populate Parms struct
-		for (FProperty* Param : Parameters.Parms)
+		for (FUnrealPropertyDefinitionInfo* ParamDef : Parameters.Parms)
 		{
-			const FString ParamName = Param->GetName();
+			const FString ParamName = ParamDef->GetProperty()->GetName();
 			OutCpp.Logf(TEXT("\t\t\tParms.%s=%s;") LINE_TERMINATOR, *ParamName, *ParamName);
 		}
 
 		const FString ObjectRef = FunctionData.FunctionReference->HasAllFunctionFlags(FUNC_Const) ? FString::Printf(TEXT("const_cast<UObject*>(O)")) : TEXT("O");
 		OutCpp.Logf(TEXT("\t\t\t%s->ProcessEvent(Func, %s);") LINE_TERMINATOR, *ObjectRef, bHasParms ? TEXT("&Parms") : TEXT("NULL"));
 
-		for (FProperty* Param : Parameters.Parms)
+		for (FUnrealPropertyDefinitionInfo* ParamDef : Parameters.Parms)
 		{
+			FProperty* Param = ParamDef->GetProperty();
 			if( Param->HasAllPropertyFlags(CPF_OutParm) && !Param->HasAnyPropertyFlags(CPF_ConstParm|CPF_ReturnParm))
 			{
 				const FString ParamName = Param->GetName();
@@ -2770,7 +2745,7 @@ void FNativeClassHeaderGenerator::ExportInterfaceCallFunctions(FOutputDevice& Ou
 			OutCpp.Logf(TEXT("I->%s_Implementation("), *FunctionName);
 
 			bool bFirst = true;
-			for (FProperty* Param : Parameters.Parms)
+			for (FUnrealPropertyDefinitionInfo* ParamDef : Parameters.Parms)
 			{
 				if (!bFirst)
 				{
@@ -2778,7 +2753,7 @@ void FNativeClassHeaderGenerator::ExportInterfaceCallFunctions(FOutputDevice& Ou
 				}
 				bFirst = false;
 
-				OutCpp.Logf(TEXT("%s"), *Param->GetName());
+				OutCpp.Logf(TEXT("%s"), *ParamDef->GetProperty()->GetName());
 			}
 
 			OutCpp.Logf(TEXT(");") LINE_TERMINATOR);
@@ -2816,11 +2791,10 @@ FString GetGeneratedMacroDeprecationWarning(const TCHAR* MacroName)
  *
  * @returns Access specifier string.
  */
-FString GetPreservedAccessSpecifierString(FClass* Class)
+FString GetPreservedAccessSpecifierString(FUnrealClassDefinitionInfo& ClassDef)
 {
 	FString PreservedAccessSpecifier;
-	FStructMetaData& Data = GTypeDefinitionInfoMap.FindChecked(Class).AsStructChecked().GetStructMetaData();
-	switch (Data.GeneratedBodyMacroAccessSpecifier)
+	switch (ClassDef.GetStructMetaData().GeneratedBodyMacroAccessSpecifier)
 	{
 	case EAccessSpecifier::ACCESS_Private:
 		PreservedAccessSpecifier = "private:";
@@ -2832,7 +2806,7 @@ FString GetPreservedAccessSpecifierString(FClass* Class)
 		PreservedAccessSpecifier = "public:";
 		break;
 	case EAccessSpecifier::ACCESS_NotAnAccessSpecifier :
-		PreservedAccessSpecifier = FString::Printf(TEXT("static_assert(false, \"Unknown access specifier for GENERATED_BODY() macro in class %s.\");"), *GetNameSafe(Class));
+		PreservedAccessSpecifier = FString::Printf(TEXT("static_assert(false, \"Unknown access specifier for GENERATED_BODY() macro in class %s.\");"), *GetNameSafe(ClassDef.GetClass()));
 		break;
 	}
 
@@ -2849,7 +2823,7 @@ void FNativeClassHeaderGenerator::ExportClassFromSourceFileInner(
 	FOutputDevice&           OutCpp,
 	FOutputDevice&           OutDeclarations,
 	FReferenceGatherers&     OutReferenceGatherers,
-	FClass*                  Class,
+	FUnrealClassDefinitionInfo&			 ClassDef,
 	const FUnrealSourceFile& SourceFile,
 	EExportClassOutFlags&    OutFlags
 ) const
@@ -2857,35 +2831,36 @@ void FNativeClassHeaderGenerator::ExportClassFromSourceFileInner(
 	FUHTStringBuilder StandardUObjectConstructorsMacroCall;
 	FUHTStringBuilder EnhancedUObjectConstructorsMacroCall;
 
-	FUnrealClassDefinitionInfo& ClassDef = GTypeDefinitionInfoMap.FindChecked(Class).AsClassChecked();
+	UClass* Class = ClassDef.GetClass();
 	FStructMetaData& StructData = ClassDef.GetStructMetaData();
 
 	// C++ -> VM stubs (native function execs)
 	FUHTStringBuilder ClassMacroCalls;
 	FUHTStringBuilder ClassNoPureDeclsMacroCalls;
-	ExportNativeFunctions(OutGeneratedHeaderText, OutCpp, ClassMacroCalls, ClassNoPureDeclsMacroCalls, OutReferenceGatherers, SourceFile, Class, StructData);
+	ExportNativeFunctions(OutGeneratedHeaderText, OutCpp, ClassMacroCalls, ClassNoPureDeclsMacroCalls, OutReferenceGatherers, SourceFile, ClassDef);
 
 	// Get Callback functions
-	TArray<UFunction*> CallbackFunctions;
-	for (UFunction* Function : TFieldRange<UFunction>(Class, EFieldIteratorFlags::ExcludeSuper))
+	TArray<FUnrealFunctionDefinitionInfo*> CallbackFunctions;
+	for (FUnrealFunctionDefinitionInfo* FunctionDef : ClassDef.GetFunctions())
 	{
+		UFunction* Function = FunctionDef->GetFunction();
 		if ((Function->FunctionFlags & FUNC_Event) && Function->GetSuperFunction() == nullptr)
 		{
-			CallbackFunctions.Add(Function);
+			CallbackFunctions.Add(FunctionDef);
 		}
 	}
 
 	FUHTStringBuilder PrologMacroCalls;
 	if (CallbackFunctions.Num() != 0)
 	{
-		Algo::SortBy(CallbackFunctions, [](UObject* Obj) { return Obj->GetName(); });
+		Algo::SortBy(CallbackFunctions, [](FUnrealFunctionDefinitionInfo* Obj) { return Obj->GetFunction()->GetName(); });
 
 		FUHTStringBuilder UClassMacroContent;
 
 		// export parameters structs for all events and delegates
-		for (UFunction* Function : CallbackFunctions)
+		for (FUnrealFunctionDefinitionInfo* FunctionDef : CallbackFunctions)
 		{
-			ExportEventParm(UClassMacroContent, OutReferenceGatherers.ForwardDeclarations, Function, /*Indent=*/ 1, /*bOutputConstructor=*/ true, EExportingState::Normal);
+			ExportEventParm(UClassMacroContent, OutReferenceGatherers.ForwardDeclarations, *FunctionDef, /*Indent=*/ 1, /*bOutputConstructor=*/ true, EExportingState::Normal);
 		}
 
 		FString MacroName = SourceFile.GetGeneratedMacroName(StructData, TEXT("_EVENT_PARMS"));
@@ -2913,13 +2888,13 @@ void FNativeClassHeaderGenerator::ExportClassFromSourceFileInner(
 	// Class definition.
 	if (!Class->HasAnyClassFlags(CLASS_NoExport))
 	{
-		ExportNatives(OutCpp, Class);
+		ExportNatives(OutCpp, ClassDef);
 	}
 
 	FUHTStringBuilder FriendText;
-	ExportNativeGeneratedInitCode(OutCpp, OutDeclarations, OutReferenceGatherers, SourceFile, Class, FriendText);
+	ExportNativeGeneratedInitCode(OutCpp, OutDeclarations, OutReferenceGatherers, SourceFile, ClassDef, FriendText);
 
-	FClass* SuperClass = Class->GetSuperClass();
+	UClass* SuperClass = Class->GetSuperClass();
 
 	// the name for the C++ version of the UClass
 	const FString ClassCPPName = FNameLookupCPP::GetNameCPP(Class);
@@ -3014,12 +2989,12 @@ void FNativeClassHeaderGenerator::ExportClassFromSourceFileInner(
 
 		if (SuperClass && Class->ClassWithin != SuperClass->ClassWithin)
 		{
-			Boilerplate.Logf(TEXT("\tDECLARE_WITHIN(%s)\r\n"), *FNameLookupCPP::GetNameCPP(Class->GetClassWithin()));
+			Boilerplate.Logf(TEXT("\tDECLARE_WITHIN(%s)\r\n"), *FNameLookupCPP::GetNameCPP(Class->ClassWithin));
 		}
 
 		if (Class->HasAnyClassFlags(CLASS_Interface))
 		{
-			ExportConstructorsMacros(OutGeneratedHeaderText, OutCpp, StandardUObjectConstructorsMacroCall, EnhancedUObjectConstructorsMacroCall, SourceFile.GetGeneratedMacroName(StructData), Class, *APIArg);
+			ExportConstructorsMacros(OutGeneratedHeaderText, OutCpp, StandardUObjectConstructorsMacroCall, EnhancedUObjectConstructorsMacroCall, SourceFile.GetGeneratedMacroName(StructData), ClassDef, *APIArg);
 
 			FString InterfaceMacroName = SourceFile.GetGeneratedMacroName(StructData, TEXT("_GENERATED_UINTERFACE_BODY"));
 			OutGeneratedHeaderText.Log(Macroize(*(InterfaceMacroName + TEXT("()")), *Boilerplate));
@@ -3050,7 +3025,7 @@ void FNativeClassHeaderGenerator::ExportClassFromSourceFileInner(
 						, DisableDeprecationWarnings
 						, *InterfaceMacroName
 						, *EnhancedUObjectConstructorsMacroCall
-						, *GetPreservedAccessSpecifierString(Class)
+						, *GetPreservedAccessSpecifierString(ClassDef)
 						, EnableDeprecationWarnings
 					)
 				)
@@ -3085,7 +3060,7 @@ void FNativeClassHeaderGenerator::ExportClassFromSourceFileInner(
 				InterfaceBoilerplate.Logf(TEXT("\tvirtual UObject* _getUObject() const { check(0 && \"Missing required implementation.\"); return nullptr; }\r\n"));
 			}
 
-			if (ClassHasReplicatedProperties(Class))
+			if (ClassHasReplicatedProperties(ClassDef))
 			{
 				WriteReplicatedMacroData(ClassRange, *ClassCPPName, *APIArg, Class, SuperClass, InterfaceBoilerplate, SourceFile, OutFlags);
 			}
@@ -3112,7 +3087,7 @@ void FNativeClassHeaderGenerator::ExportClassFromSourceFileInner(
 				Boilerplate.Logf(TEXT("\tvirtual UObject* _getUObject() const override { return const_cast<%s*>(this); }\r\n"), *ClassCPPName);
 			}
 
-			if (ClassHasReplicatedProperties(Class))
+			if (ClassHasReplicatedProperties(ClassDef))
 			{
 				WriteReplicatedMacroData(ClassRange, *ClassCPPName, *APIArg, Class, SuperClass, Boilerplate, SourceFile, OutFlags);
 			}
@@ -3126,7 +3101,7 @@ void FNativeClassHeaderGenerator::ExportClassFromSourceFileInner(
 				WriteMacro(OutGeneratedHeaderText, MacroName, Boilerplate);
 				ClassMacroCalls.Logf(TEXT("\t%s\r\n"), *MacroName);
 
-				ExportConstructorsMacros(OutGeneratedHeaderText, OutCpp, StandardUObjectConstructorsMacroCall, EnhancedUObjectConstructorsMacroCall, SourceFile.GetGeneratedMacroName(StructData), Class, *APIArg);
+				ExportConstructorsMacros(OutGeneratedHeaderText, OutCpp, StandardUObjectConstructorsMacroCall, EnhancedUObjectConstructorsMacroCall, SourceFile.GetGeneratedMacroName(StructData), ClassDef, *APIArg);
 			}
 		}
 	}
@@ -3165,7 +3140,7 @@ void FNativeClassHeaderGenerator::ExportClassFromSourceFileInner(
 		}
 
 		FString WrappedLegacyGeneratedBody = FString::Printf(TEXT("%s%s%s%s%s%s"), *DeprecationWarning, DisableDeprecationWarnings, Public, *LegacyGeneratedBody, Public, EnableDeprecationWarnings);
-		FString WrappedGeneratedBody = FString::Printf(TEXT("%s%s%s%s%s"), DisableDeprecationWarnings, Public, *GeneratedBody, *GetPreservedAccessSpecifierString(Class), EnableDeprecationWarnings);
+		FString WrappedGeneratedBody = FString::Printf(TEXT("%s%s%s%s%s"), DisableDeprecationWarnings, Public, *GeneratedBody, *GetPreservedAccessSpecifierString(ClassDef), EnableDeprecationWarnings);
 
 		OutGeneratedHeaderText.Log(Macroize(*SourceFile.GetGeneratedBodyMacroName(GeneratedBodyLine, true), MoveTemp(WrappedLegacyGeneratedBody)));
 		OutGeneratedHeaderText.Log(Macroize(*SourceFile.GetGeneratedBodyMacroName(GeneratedBodyLine, false), MoveTemp(WrappedGeneratedBody)));
@@ -3204,9 +3179,9 @@ void ExportCopyConstructorDefinition(FOutputDevice& Out, const TCHAR* API, const
  * @param Class Class to generate for.
  * @param API API string.
  */
-void ExportVTableHelperCtorAndCaller(FOutputDevice& Out, FStructMetaData& StructData, const TCHAR* API, const TCHAR* ClassCPPName)
+void ExportVTableHelperCtorAndCaller(FOutputDevice& Out, FUnrealClassDefinitionInfo& ClassDef, const TCHAR* API, const TCHAR* ClassCPPName)
 {
-	if (!StructData.bCustomVTableHelperConstructorDeclared)
+	if (!ClassDef.GetStructMetaData().bCustomVTableHelperConstructorDeclared)
 	{
 		Out.Logf(TEXT("\tDECLARE_VTABLE_PTR_HELPER_CTOR(%s_API, %s);" LINE_TERMINATOR), API, ClassCPPName);
 	}
@@ -3220,13 +3195,14 @@ void ExportVTableHelperCtorAndCaller(FOutputDevice& Out, FStructMetaData& Struct
  * @param Class Class to generate constructor for.
  * @param API API string for this constructor.
  */
-void ExportStandardConstructorsMacro(FOutputDevice& Out, FClass* Class, FStructMetaData& StructData, const TCHAR* API, const TCHAR* ClassCPPName)
+void ExportStandardConstructorsMacro(FOutputDevice& Out, FUnrealClassDefinitionInfo& ClassDef, const TCHAR* API, const TCHAR* ClassCPPName)
 {
+	UClass* Class = ClassDef.GetClass();
 	if (!Class->HasAnyClassFlags(CLASS_CustomConstructor))
 	{
 		Out.Logf(TEXT("\t/** Standard constructor, called after all reflected properties have been initialized */\r\n"));
 		Out.Logf(TEXT("\t%s_API %s(const FObjectInitializer& ObjectInitializer%s);\r\n"), API, ClassCPPName,
-			StructData.bDefaultConstructorDeclared ? TEXT("") : TEXT(" = FObjectInitializer::Get()"));
+			ClassDef.GetStructMetaData().bDefaultConstructorDeclared ? TEXT("") : TEXT(" = FObjectInitializer::Get()"));
 	}
 	if (Class->HasAnyClassFlags(CLASS_Abstract))
 	{
@@ -3238,7 +3214,7 @@ void ExportStandardConstructorsMacro(FOutputDevice& Out, FClass* Class, FStructM
 
 	}
 
-	ExportVTableHelperCtorAndCaller(Out, StructData, API, ClassCPPName);
+	ExportVTableHelperCtorAndCaller(Out, ClassDef, API, ClassCPPName);
 	ExportCopyConstructorDefinition(Out, API, ClassCPPName);
 }
 
@@ -3249,8 +3225,10 @@ void ExportStandardConstructorsMacro(FOutputDevice& Out, FClass* Class, FStructM
  * @param Class Class to generate constructor for.
  * @param API API string for this constructor.
  */
-void ExportConstructorDefinition(FOutputDevice& Out, FClass* Class, FStructMetaData& StructData, const TCHAR* API, const TCHAR* ClassCPPName)
+void ExportConstructorDefinition(FOutputDevice& Out, FUnrealClassDefinitionInfo& ClassDef, const TCHAR* API, const TCHAR* ClassCPPName)
 {
+	UClass* Class = ClassDef.GetClass();
+	FStructMetaData& StructData = ClassDef.GetStructMetaData();
 	if (!StructData.bConstructorDeclared)
 	{
 		Out.Logf(TEXT("\t/** Standard constructor, called after all reflected properties have been initialized */\r\n"));
@@ -3258,19 +3236,18 @@ void ExportConstructorDefinition(FOutputDevice& Out, FClass* Class, FStructMetaD
 		// Assume super class has OI constructor, this may not always be true but we should always be able to check this.
 		// In any case, it will default to old behaviour before we even checked this.
 		bool bSuperClassObjectInitializerConstructorDeclared = true;
-		FClass* SuperClass = Class->GetSuperClass();
-		if (SuperClass != nullptr)
+		FUnrealClassDefinitionInfo* SuperClassDef = ClassDef.GetSuperClass();
+		if (SuperClassDef != nullptr)
 		{
-			FUnrealStructDefinitionInfo& StructDef = GTypeDefinitionInfoMap.FindChecked(SuperClass).AsStructChecked();
-			if (StructDef.HasSource()) // Don't consider the internal types generated from the engine
+			if (SuperClassDef->HasSource()) // Don't consider the internal types generated from the engine
 			{
-				FStructMetaData& SuperClassData = GTypeDefinitionInfoMap.FindChecked(SuperClass).AsStructChecked().GetStructMetaData();
-
 				// Since we are dependent on our SuperClass having determined which constructors are defined, 
 				// if it is not yet determined we will need to wait on it becoming available. 
 				// Since the SourceFile array provided to the ParallelFor is in dependency order and does not allow cyclic dependencies, 
 				// we can be certain that another thread has started processing the file containing our SuperClass before this
 				// file would have been assigned out,  so we just have to wait
+				//ETSTODO
+				FStructMetaData& SuperClassData = SuperClassDef->GetStructMetaData();
 				while (!SuperClassData.bConstructorDeclared)
 				{
 					FPlatformProcess::Sleep(0.01);
@@ -3301,11 +3278,11 @@ void ExportConstructorDefinition(FOutputDevice& Out, FClass* Class, FStructMetaD
  * @param Out Output device to generate to.
  * @param Class Class to generate constructor call definition for.
  */
-void ExportDefaultConstructorCallDefinition(FOutputDevice& Out, FClass* Class, FStructMetaData& StructData, const TCHAR* ClassCPPName)
+void ExportDefaultConstructorCallDefinition(FOutputDevice& Out, FUnrealClassDefinitionInfo& ClassDef, const TCHAR* ClassCPPName)
 {
-	if (StructData.bObjectInitializerConstructorDeclared)
+	if (ClassDef.GetStructMetaData().bObjectInitializerConstructorDeclared)
 	{
-		if (Class->HasAnyClassFlags(CLASS_Abstract))
+		if (ClassDef.GetClass()->HasAnyClassFlags(CLASS_Abstract))
 		{
 			Out.Logf(TEXT("\tDEFINE_ABSTRACT_DEFAULT_OBJECT_INITIALIZER_CONSTRUCTOR_CALL(%s)\r\n"), ClassCPPName);
 		}
@@ -3314,9 +3291,9 @@ void ExportDefaultConstructorCallDefinition(FOutputDevice& Out, FClass* Class, F
 			Out.Logf(TEXT("\tDEFINE_DEFAULT_OBJECT_INITIALIZER_CONSTRUCTOR_CALL(%s)\r\n"), ClassCPPName);
 		}
 	}
-	else if (StructData.bDefaultConstructorDeclared)
+	else if (ClassDef.GetStructMetaData().bDefaultConstructorDeclared)
 	{
-		if (Class->HasAnyClassFlags(CLASS_Abstract))
+		if (ClassDef.GetClass()->HasAnyClassFlags(CLASS_Abstract))
 		{
 			Out.Logf(TEXT("\tDEFINE_ABSTRACT_DEFAULT_CONSTRUCTOR_CALL(%s)\r\n"), ClassCPPName);
 		}
@@ -3338,11 +3315,11 @@ void ExportDefaultConstructorCallDefinition(FOutputDevice& Out, FClass* Class, F
  * @param Class Class to generate constructor for.
  * @param API API string for this constructor.
  */
-void ExportEnhancedConstructorsMacro(FOutputDevice& Out, FClass* Class, FStructMetaData& StructData, const TCHAR* API, const TCHAR* ClassCPPName)
+void ExportEnhancedConstructorsMacro(FOutputDevice& Out, FUnrealClassDefinitionInfo& ClassDef, const TCHAR* API, const TCHAR* ClassCPPName)
 {
-	ExportConstructorDefinition(Out, Class, StructData, API, ClassCPPName);
-	ExportVTableHelperCtorAndCaller(Out, StructData, API, ClassCPPName);
-	ExportDefaultConstructorCallDefinition(Out, Class, StructData, ClassCPPName);
+	ExportConstructorDefinition(Out, ClassDef, API, ClassCPPName);
+	ExportVTableHelperCtorAndCaller(Out, ClassDef, API, ClassCPPName);
+	ExportDefaultConstructorCallDefinition(Out, ClassDef, ClassCPPName);
 }
 
 /**
@@ -3356,26 +3333,25 @@ FString GetBuildPath(FUnrealSourceFile& SourceFile)
 {
 	FString Out = SourceFile.GetFilename();
 
-	ConvertToBuildIncludePath(GTypeDefinitionInfoMap.FindChecked(SourceFile.GetPackage()).AsPackageChecked().GetModule(), Out);
+	ConvertToBuildIncludePath(GTypeDefinitionInfoMap.FindChecked<FUnrealPackageDefinitionInfo>(SourceFile.GetPackage()).GetModule(), Out);
 
 	return Out;
 }
 
-void FNativeClassHeaderGenerator::ExportConstructorsMacros(FOutputDevice& OutGeneratedHeaderText, FOutputDevice& Out, FOutputDevice& StandardUObjectConstructorsMacroCall, FOutputDevice& EnhancedUObjectConstructorsMacroCall, const FString& ConstructorsMacroPrefix, FClass* Class, const TCHAR* APIArg)
+void FNativeClassHeaderGenerator::ExportConstructorsMacros(FOutputDevice& OutGeneratedHeaderText, FOutputDevice& Out, FOutputDevice& StandardUObjectConstructorsMacroCall, FOutputDevice& EnhancedUObjectConstructorsMacroCall, const FString& ConstructorsMacroPrefix, FUnrealClassDefinitionInfo& ClassDef, const TCHAR* APIArg)
 {
+	UClass* Class = ClassDef.GetClass();
 	const FString ClassCPPName = FNameLookupCPP::GetNameCPP(Class);
 
-	FStructMetaData& StructData = GTypeDefinitionInfoMap.FindChecked(Class).AsStructChecked().GetStructMetaData();
-
 	FUHTStringBuilder StdMacro;
 	FUHTStringBuilder EnhMacro;
 	FString StdMacroName = ConstructorsMacroPrefix + TEXT("_STANDARD_CONSTRUCTORS");
 	FString EnhMacroName = ConstructorsMacroPrefix + TEXT("_ENHANCED_CONSTRUCTORS");
 
-	ExportStandardConstructorsMacro(StdMacro, Class, StructData, APIArg, *ClassCPPName);
-	ExportEnhancedConstructorsMacro(EnhMacro, Class, StructData, APIArg, *ClassCPPName);
+	ExportStandardConstructorsMacro(StdMacro, ClassDef, APIArg, *ClassCPPName);
+	ExportEnhancedConstructorsMacro(EnhMacro, ClassDef, APIArg, *ClassCPPName);
 
-	if (!StructData.bCustomVTableHelperConstructorDeclared)
+	if (!ClassDef.GetStructMetaData().bCustomVTableHelperConstructorDeclared)
 	{
 		Out.Logf(TEXT("\tDEFINE_VTABLE_PTR_HELPER_CTOR(%s);" LINE_TERMINATOR), *ClassCPPName);
 	}
@@ -3481,7 +3457,7 @@ void FNativeClassHeaderGenerator::ExportEnum(FOutputDevice& Out, UEnum* Enum) co
 	}
 	Out.Logf( TEXT("\r\n") );
 
-	FUnrealEnumDefinitionInfo& EnumDef = GTypeDefinitionInfoMap.FindChecked(Enum).AsEnumChecked();
+	FUnrealEnumDefinitionInfo& EnumDef = GTypeDefinitionInfoMap.FindChecked<FUnrealEnumDefinitionInfo>(Enum);
 
 	// Forward declare the StaticEnum<> specialisation for enum classes
 	if (Enum->GetCppForm() == UEnum::ECppForm::EnumClass)
@@ -3519,7 +3495,7 @@ void FNativeClassHeaderGenerator::ExportEnum(FOutputDevice& Out, UEnum* Enum) co
 // Exports the header text for the list of structs specified (GENERATED_BODY impls)
 void FNativeClassHeaderGenerator::ExportGeneratedStructBodyMacros(FOutputDevice& OutGeneratedHeaderText, FOutputDevice& Out, FReferenceGatherers& OutReferenceGatherers, const FUnrealSourceFile& SourceFile, UScriptStruct* Struct) const
 {
-	FUnrealScriptStructDefinitionInfo& ScriptStructDef = GTypeDefinitionInfoMap.FindChecked(Struct).AsScriptStructChecked();
+	FUnrealScriptStructDefinitionInfo& ScriptStructDef = GTypeDefinitionInfoMap.FindChecked<FUnrealScriptStructDefinitionInfo>(Struct);
 	ScriptStructDef.AddCrossModuleReference(OutReferenceGatherers.UniqueCrossModuleReferences, true);
 
 	const bool bIsDynamic = FClass::IsDynamic(Struct);
@@ -3851,16 +3827,16 @@ void FNativeClassHeaderGenerator::ExportGeneratedStructBodyMacros(FOutputDevice&
 	GeneratedStructRegisterFunctionText.Logf(TEXT("\t{\r\n"));
 
 	// if this is a no export struct, we will put a local struct here for offset determination
-	TArray<UScriptStruct*> NoExportStructs = FindNoExportStructs(Struct);
-	for (UScriptStruct* NoExportStruct : NoExportStructs)
+	TArray<FUnrealScriptStructDefinitionInfo*> NoExportStructs = FindNoExportStructs(&ScriptStructDef);
+	for (FUnrealScriptStructDefinitionInfo* NoExportStructDef : NoExportStructs)
 	{
-		ExportMirrorsForNoexportStruct(GeneratedStructRegisterFunctionText, NoExportStruct, /*Indent=*/ 2);
+		ExportMirrorsForNoexportStruct(GeneratedStructRegisterFunctionText, *NoExportStructDef, /*Indent=*/ 2);
 	}
 
 	if (BaseStruct)
 	{
 		CastChecked<UScriptStruct>(BaseStruct); // this better actually be a script struct
-		GTypeDefinitionInfoMap.FindChecked(BaseStruct).AsFieldChecked().AddCrossModuleReference(OutReferenceGatherers.UniqueCrossModuleReferences, true);
+		GTypeDefinitionInfoMap.FindChecked<FUnrealFieldDefinitionInfo>(BaseStruct).AddCrossModuleReference(OutReferenceGatherers.UniqueCrossModuleReferences, true);
 	}
 
 	EStructFlags UncomputedFlags = (EStructFlags)(Struct->StructFlags & ~STRUCT_ComputedFlags);
@@ -3885,7 +3861,10 @@ void FNativeClassHeaderGenerator::ExportGeneratedStructBodyMacros(FOutputDevice&
 	FString MetaDataParams = OutputMetaDataCodeForObject(GeneratedStructRegisterFunctionText, StaticDefinitions, Struct, *FString::Printf(TEXT("%s::Struct_MetaDataParams"), *StaticsStructName), TEXT("\t\t"), TEXT("\t"));
 
 	TArray<FProperty*> Props;
-	Algo::Copy(TFieldRange<FProperty>(Struct, EFieldIteratorFlags::ExcludeSuper), Props, Algo::NoRef);
+	for (FUnrealPropertyDefinitionInfo* PropertyDef : ScriptStructDef.GetProperties())
+	{
+		Props.Add(PropertyDef->GetProperty());
+	}
 
 	FString NewStructOps;
 	if (Struct->StructFlags & STRUCT_Native)
@@ -3930,7 +3909,7 @@ void FNativeClassHeaderGenerator::ExportGeneratedStructBodyMacros(FOutputDevice&
 	GeneratedStructRegisterFunctionText.Log (TEXT("\t{\r\n"));
 
 	FString NoExportStructNameCPP;
-	if (NoExportStructs.Contains(Struct))
+	if (NoExportStructs.Contains(&ScriptStructDef))
 	{
 		NoExportStructNameCPP = FString::Printf(TEXT("%s::%s"), *StaticsStructName, *StructNameCPP);
 	}
@@ -4049,7 +4028,7 @@ void FNativeClassHeaderGenerator::ExportGeneratedStructBodyMacros(FOutputDevice&
 
 void FNativeClassHeaderGenerator::ExportGeneratedEnumInitCode(FOutputDevice& Out, FReferenceGatherers& OutReferenceGatherers, const FUnrealSourceFile& SourceFile, UEnum* Enum) const
 {
-	FUnrealEnumDefinitionInfo& EnumDef  = GTypeDefinitionInfoMap.FindChecked(Enum).AsEnumChecked();
+	FUnrealEnumDefinitionInfo& EnumDef  = GTypeDefinitionInfoMap.FindChecked<FUnrealEnumDefinitionInfo>(Enum);
 	const bool bIsDynamic               = FClass::IsDynamic(static_cast<UField*>(Enum));
 	const FString& SingletonName        = EnumDef.GetSingletonNameChopped(true);
 	const FString EnumNameCpp           = Enum->GetName(); //UserDefinedEnum should already have a valid cpp name.
@@ -4233,27 +4212,29 @@ void FNativeClassHeaderGenerator::ExportGeneratedEnumInitCode(FOutputDevice& Out
 	Out.Log(GeneratedEnumRegisterFunctionText);
 }
 
-void FNativeClassHeaderGenerator::ExportMirrorsForNoexportStruct(FOutputDevice& Out, UScriptStruct* Struct, int32 TextIndent)
+void FNativeClassHeaderGenerator::ExportMirrorsForNoexportStruct(FOutputDevice& Out, FUnrealScriptStructDefinitionInfo& ScriptStructDef, int32 TextIndent)
 {
+	UScriptStruct* ScriptStruct = ScriptStructDef.GetScriptStruct();
+
 	// Export struct.
-	const FString StructName = FNameLookupCPP::GetNameCPP(Struct);
+	const FString StructName = FNameLookupCPP::GetNameCPP(ScriptStruct);
 	Out.Logf(TEXT("%sstruct %s"), FCString::Tab(TextIndent), *StructName);
-	if (Struct->GetSuperStruct() != NULL)
+	if (ScriptStruct->GetSuperStruct() != NULL)
 	{
-		Out.Logf(TEXT(" : public %s"), *FNameLookupCPP::GetNameCPP(Struct->GetSuperStruct()));
+		Out.Logf(TEXT(" : public %s"), *FNameLookupCPP::GetNameCPP(ScriptStruct->GetSuperStruct()));
 	}
 	Out.Logf(TEXT("\r\n%s{\r\n"), FCString::Tab(TextIndent));
 
 	// Export the struct's CPP properties.
-	ExportProperties(Out, Struct, TextIndent);
+	ExportProperties(Out, ScriptStructDef, TextIndent);
 
 	Out.Logf(TEXT("%s};\r\n\r\n"), FCString::Tab(TextIndent));
 }
 
-bool FNativeClassHeaderGenerator::WillExportEventParms( UFunction* Function )
+bool FNativeClassHeaderGenerator::WillExportEventParms(FUnrealFunctionDefinitionInfo& FunctionDef)
 {
-  TFieldIterator<FProperty> It(Function);
-  return It && (It->PropertyFlags&CPF_Parm);
+	const TArray<FUnrealPropertyDefinitionInfo*>& Properties = FunctionDef.GetProperties();
+	return Properties.Num() > 0 && (Properties[0]->GetProperty()->PropertyFlags & CPF_Parm) != 0;
 }
 
 void WriteEventFunctionPrologue(FOutputDevice& Output, int32 Indent, const FParmsAndReturnProperties& Parameters, UObject* FunctionOuter, const TCHAR* FunctionName)
@@ -4270,8 +4251,9 @@ void WriteEventFunctionPrologue(FOutputDevice& Output, int32 Indent, const FParm
 	Output.Logf(TEXT("%s%s Parms;\r\n"), FCString::Tab(Indent + 1), *EventStructName );
 
 	// Declare a parameter struct for this event/delegate and assign the struct members using the values passed into the event/delegate call.
-	for (FProperty* Prop : Parameters.Parms)
+	for (FUnrealPropertyDefinitionInfo* PropDef : Parameters.Parms)
 	{
+		FProperty* Prop = PropDef->GetProperty();
 		const FString PropertyName = Prop->GetName();
 		if (Prop->ArrayDim > 1)
 		{
@@ -4293,8 +4275,9 @@ void WriteEventFunctionPrologue(FOutputDevice& Output, int32 Indent, const FParm
 void WriteEventFunctionEpilogue(FOutputDevice& Output, int32 Indent, const FParmsAndReturnProperties& Parameters)
 {
 	// Out parm copying.
-	for (FProperty* Prop : Parameters.Parms)
+	for (FUnrealPropertyDefinitionInfo* PropDef : Parameters.Parms)
 	{
+		FProperty* Prop = PropDef->GetProperty();
 		if ((Prop->PropertyFlags & (CPF_OutParm | CPF_ConstParm)) == CPF_OutParm)
 		{
 			const FString PropertyName = Prop->GetName();
@@ -4313,16 +4296,20 @@ void WriteEventFunctionEpilogue(FOutputDevice& Output, int32 Indent, const FParm
 	if (Parameters.Return)
 	{
 		// Make sure uint32 -> bool is supported
-		bool bBoolProperty = Parameters.Return->IsA(FBoolProperty::StaticClass());
-		Output.Logf(TEXT("%sreturn %sParms.%s;\r\n"), FCString::Tab(Indent + 1), bBoolProperty ? TEXT("!!") : TEXT(""), *Parameters.Return->GetName());
+		FProperty* Return = Parameters.Return->GetProperty();
+		bool bBoolProperty = Return->IsA(FBoolProperty::StaticClass());
+		Output.Logf(TEXT("%sreturn %sParms.%s;\r\n"), FCString::Tab(Indent + 1), bBoolProperty ? TEXT("!!") : TEXT(""), *Return->GetName());
 	}
 	Output.Logf(TEXT("%s}\r\n"), FCString::Tab(Indent));
 }
 
-void FNativeClassHeaderGenerator::ExportDelegateDeclaration(FOutputDevice& Out, FReferenceGatherers& OutReferenceGatherers, const FUnrealSourceFile& SourceFile, UFunction* Function) const
+void FNativeClassHeaderGenerator::ExportDelegateDeclaration(FOutputDevice& Out, FReferenceGatherers& OutReferenceGatherers, const FUnrealSourceFile& SourceFile, FUnrealFunctionDefinitionInfo& FunctionDef) const
 {
 	static const TCHAR DelegateStr[] = TEXT("delegate");
 
+	UFunction* Function = FunctionDef.GetFunction();
+	FFuncInfo FunctionData = FunctionDef.GetFunctionData(); // THIS IS A COPY
+
 	check(Function->HasAnyFunctionFlags(FUNC_Delegate));
 
 	const bool bIsMulticastDelegate = Function->HasAnyFunctionFlags( FUNC_MulticastDelegate );
@@ -4330,9 +4317,6 @@ void FNativeClassHeaderGenerator::ExportDelegateDeclaration(FOutputDevice& Out,
 	// Unmangle the function name
 	const FString DelegateName = Function->GetName().LeftChop(HEADER_GENERATED_DELEGATE_SIGNATURE_SUFFIX_LENGTH);
 
-	FUnrealFunctionDefinitionInfo& FuncDef = GTypeDefinitionInfoMap.FindChecked(Function).AsFunctionChecked();
-	FFuncInfo FunctionData = FuncDef.GetFunctionData();
-
 	// Add class name to beginning of function, to avoid collisions with other classes with the same delegate name in this scope
 	check(FunctionData.MarshallAndCallName.StartsWith(DelegateStr));
 	FString ShortName = *FunctionData.MarshallAndCallName + UE_ARRAY_COUNT(DelegateStr) - 1;
@@ -4349,32 +4333,32 @@ void FNativeClassHeaderGenerator::ExportDelegateDeclaration(FOutputDevice& Out,
 	DelegateOutput.Log(TEXT("static "));
 
 	// export the line that looks like: int32 Main(const FString& Parms)
-	ExportNativeFunctionHeader(DelegateOutput, OutReferenceGatherers.ForwardDeclarations, FunctionData, EExportFunctionType::Event, EExportFunctionHeaderStyle::Declaration, *ExtraParam, *GetAPIString());
+	ExportNativeFunctionHeader(DelegateOutput, OutReferenceGatherers.ForwardDeclarations, FunctionDef, FunctionData, EExportFunctionType::Event, EExportFunctionHeaderStyle::Declaration, *ExtraParam, *GetAPIString());
 
 	// Only exporting function prototype
 	DelegateOutput.Logf(TEXT(";\r\n"));
 
-	ExportFunction(Out, OutReferenceGatherers, SourceFile, Function, false);
+	ExportFunction(Out, OutReferenceGatherers, SourceFile, FunctionDef, false);
 }
 
-void FNativeClassHeaderGenerator::ExportDelegateDefinition(FOutputDevice& Out, FReferenceGatherers& OutReferenceGatherers, const FUnrealSourceFile& SourceFile, UFunction* Function) const
+void FNativeClassHeaderGenerator::ExportDelegateDefinition(FOutputDevice& Out, FReferenceGatherers& OutReferenceGatherers, const FUnrealSourceFile& SourceFile, FUnrealFunctionDefinitionInfo& FunctionDef) const
 {
 	const TCHAR DelegateStr[] = TEXT("delegate");
 
+	UFunction* Function = FunctionDef.GetFunction();
+	FFuncInfo FunctionData = FunctionDef.GetFunctionData(); // THIS IS A COPY
+
 	check(Function->HasAnyFunctionFlags(FUNC_Delegate));
 
 	// Export parameters structs for all delegates.  We'll need these to declare our delegate execution function.
 	FUHTStringBuilder DelegateOutput;
-	ExportEventParm(DelegateOutput, OutReferenceGatherers.ForwardDeclarations, Function, /*Indent=*/ 0, /*bOutputConstructor=*/ true, EExportingState::Normal);
+	ExportEventParm(DelegateOutput, OutReferenceGatherers.ForwardDeclarations, FunctionDef, /*Indent=*/ 0, /*bOutputConstructor=*/ true, EExportingState::Normal);
 
 	const bool bIsMulticastDelegate = Function->HasAnyFunctionFlags( FUNC_MulticastDelegate );
 
 	// Unmangle the function name
 	const FString DelegateName = Function->GetName().LeftChop(HEADER_GENERATED_DELEGATE_SIGNATURE_SUFFIX_LENGTH);
 
-	FUnrealFunctionDefinitionInfo& FuncDef = GTypeDefinitionInfoMap.FindChecked(Function).AsFunctionChecked();
-	FFuncInfo FunctionData = FuncDef.GetFunctionData();
-
 	// Always export delegate wrapper functions as inline
 	FunctionData.FunctionExportFlags |= FUNCEXPORT_Inline;
 
@@ -4393,9 +4377,9 @@ void FNativeClassHeaderGenerator::ExportDelegateDefinition(FOutputDevice& Out, F
 	DelegateOutput.Log(TEXT("static "));
 
 	// export the line that looks like: int32 Main(const FString& Parms)
-	ExportNativeFunctionHeader(DelegateOutput, OutReferenceGatherers.ForwardDeclarations, FunctionData, EExportFunctionType::Event, EExportFunctionHeaderStyle::Declaration, *ExtraParam, *GetAPIString());
+	ExportNativeFunctionHeader(DelegateOutput, OutReferenceGatherers.ForwardDeclarations, FunctionDef, FunctionData, EExportFunctionType::Event, EExportFunctionHeaderStyle::Declaration, *ExtraParam, *GetAPIString());
 
-	FParmsAndReturnProperties Parameters = GetFunctionParmsAndReturn(FunctionData.FunctionReference);
+	FParmsAndReturnProperties Parameters = GetFunctionParmsAndReturn(FunctionDef);
 
 	WriteEventFunctionPrologue(DelegateOutput, 0, Parameters, Function->GetOuter(), *DelegateName);
 	{
@@ -4409,9 +4393,10 @@ void FNativeClassHeaderGenerator::ExportDelegateDefinition(FOutputDevice& Out, F
 	WriteMacro(Out, MacroName, DelegateOutput);
 }
 
-void FNativeClassHeaderGenerator::ExportEventParm(FUHTStringBuilder& Out, TSet<FString>& PropertyFwd, UFunction* Function, int32 Indent, bool bOutputConstructor, EExportingState ExportingState)
+void FNativeClassHeaderGenerator::ExportEventParm(FUHTStringBuilder& Out, TSet<FString>& PropertyFwd, FUnrealFunctionDefinitionInfo& FunctionDef, int32 Indent, bool bOutputConstructor, EExportingState ExportingState)
 {
-	if (!WillExportEventParms(Function))
+	UFunction* Function = FunctionDef.GetFunction();
+	if (!WillExportEventParms(FunctionDef))
 	{
 		return;
 	}
@@ -4426,8 +4411,9 @@ void FNativeClassHeaderGenerator::ExportEventParm(FUHTStringBuilder& Out, TSet<F
 	Out.Logf(TEXT("%sstruct %s\r\n"), FCString::Tab(Indent), *EventParmStructName);
 	Out.Logf(TEXT("%s{\r\n"), FCString::Tab(Indent));
 
-	for (FProperty* Prop : TFieldRange<FProperty>(Function))
+	for (FUnrealPropertyDefinitionInfo* PropDef : FunctionDef.GetProperties())
 	{
+		FProperty* Prop = PropDef->GetProperty();
 		if (!(Prop->PropertyFlags & CPF_Parm))
 		{
 			continue;
@@ -4457,8 +4443,7 @@ void FNativeClassHeaderGenerator::ExportEventParm(FUHTStringBuilder& Out, TSet<F
 			PropertyText.Logf(TEXT("const "));
 		}
 
-		FUnrealPropertyDefinitionInfo& PropDef = GTypeDefinitionInfoMap.FindChecked(Prop).AsPropertyChecked();
-		Prop->ExportCppDeclaration(PropertyText, EExportedDeclaration::Local, PropDef.GetArrayDimensions());
+		Prop->ExportCppDeclaration(PropertyText, EExportedDeclaration::Local, PropDef->GetArrayDimensions());
 		ApplyAlternatePropertyExportText(Prop, PropertyText, ExportingState);
 
 		PropertyText.Log(TEXT(";\r\n"));
@@ -4466,9 +4451,10 @@ void FNativeClassHeaderGenerator::ExportEventParm(FUHTStringBuilder& Out, TSet<F
 
 	}
 	// constructor must initialize the return property if it needs it
-	FProperty* Prop = Function->GetReturnProperty();
-	if (Prop && bOutputConstructor)
+	FUnrealPropertyDefinitionInfo* PropDef = FunctionDef.GetReturn();
+	if (PropDef && bOutputConstructor)
 	{
+		FProperty* Prop = PropDef->GetProperty();
 		FUHTStringBuilder InitializationAr;
 
 		FStructProperty* InnerStruct = CastField<FStructProperty>(Prop);
@@ -4497,7 +4483,7 @@ void FNativeClassHeaderGenerator::ExportEventParm(FUHTStringBuilder& Out, TSet<F
 			check(Prop->ArrayDim == 1); // can't return arrays
 			Out.Logf(TEXT("\r\n%s/** Constructor, initializes return property only **/\r\n"), FCString::Tab(Indent + 1));
 			Out.Logf(TEXT("%s%s()\r\n"), FCString::Tab(Indent + 1), *EventParmStructName);
-			Out.Logf(TEXT("%s%s %s(%s)\r\n"), FCString::Tab(Indent + 2), TEXT(":"), *Prop->GetName(), *GetNullParameterValue(Prop, true));
+			Out.Logf(TEXT("%s%s %s(%s)\r\n"), FCString::Tab(Indent + 2), TEXT(":"), *Prop->GetName(), *GetNullParameterValue(*PropDef, true));
 			Out.Logf(TEXT("%s{\r\n"), FCString::Tab(Indent + 1));
 			Out.Logf(TEXT("%s}\r\n"), FCString::Tab(Indent + 1));
 		}
@@ -4513,8 +4499,9 @@ void FNativeClassHeaderGenerator::ExportEventParm(FUHTStringBuilder& Out, TSet<F
  *
  * @return	the intrinsic null value for the property (0 for ints, TEXT("") for strings, etc.)
  */
-FString FNativeClassHeaderGenerator::GetNullParameterValue( FProperty* Prop, bool bInitializer/*=false*/ )
+FString FNativeClassHeaderGenerator::GetNullParameterValue(FUnrealPropertyDefinitionInfo& PropertyDef, bool bInitializer/*=false*/ )
 {
+	FProperty* Prop = PropertyDef.GetProperty();
 	FFieldClass* PropClass = Prop->GetClass();
 	FObjectPropertyBase* ObjectProperty = CastField<FObjectPropertyBase>(Prop);
 	if (PropClass == FByteProperty::StaticClass())
@@ -4595,12 +4582,13 @@ FString FNativeClassHeaderGenerator::GetNullParameterValue( FProperty* Prop, boo
 }
 
 
-FString FNativeClassHeaderGenerator::GetFunctionReturnString(UFunction* Function, FReferenceGatherers& OutReferenceGatherers)
+FString FNativeClassHeaderGenerator::GetFunctionReturnString(FUnrealFunctionDefinitionInfo& FunctionDef, FReferenceGatherers& OutReferenceGatherers)
 {
 	FString Result;
 
-	if (FProperty* Return = Function->GetReturnProperty())
+	if (FUnrealPropertyDefinitionInfo* ReturnDef = FunctionDef.GetReturn())
 	{
+		FProperty* Return = ReturnDef->GetProperty();
 		FString ExtendedReturnType;
 		OutReferenceGatherers.ForwardDeclarations.Add(Return->GetCPPTypeForwardDeclaration());
 		FString ReturnType = Return->GetCPPType(&ExtendedReturnType, CPPF_ArgumentOrReturnValue);
@@ -4662,13 +4650,14 @@ FString CreateClickableErrorMessage(const FString& Filename, int32 Line, int32 C
 	return FString::Printf(TEXT("%s(%d,%d): error: "), *Filename, Line, Column);
 }
 
-void FNativeClassHeaderGenerator::CheckRPCFunctions(FReferenceGatherers& OutReferenceGatherers, const FFuncInfo& FunctionData, const FString& ClassName, int32 ImplementationPosition, int32 ValidatePosition, const FUnrealSourceFile& SourceFile) const
+void FNativeClassHeaderGenerator::CheckRPCFunctions(FReferenceGatherers& OutReferenceGatherers, FUnrealFunctionDefinitionInfo& FunctionDef, const FString& ClassName, int32 ImplementationPosition, int32 ValidatePosition, const FUnrealSourceFile& SourceFile) const
 {
 	bool bHasImplementation = ImplementationPosition != INDEX_NONE;
 	bool bHasValidate = ValidatePosition != INDEX_NONE;
 
-	UFunction* Function = FunctionData.FunctionReference;
-	FString FunctionReturnType = GetFunctionReturnString(Function, OutReferenceGatherers);
+	UFunction* Function = FunctionDef.GetFunction();
+	const FFuncInfo& FunctionData = FunctionDef.GetFunctionData();
+	FString FunctionReturnType = GetFunctionReturnString(FunctionDef, OutReferenceGatherers);
 	const TCHAR* ConstModifier = (Function->HasAllFunctionFlags(FUNC_Const) ? TEXT("const ") : TEXT(" "));
 
 	const bool bIsNative = Function->HasAllFunctionFlags(FUNC_Native);
@@ -4682,7 +4671,7 @@ void FNativeClassHeaderGenerator::CheckRPCFunctions(FReferenceGatherers& OutRefe
 
 	check(bNeedsImplementation || bNeedsValidate);
 
-	FString ParameterString = GetFunctionParameterString(Function, OutReferenceGatherers);
+	FString ParameterString = GetFunctionParameterString(FunctionDef, OutReferenceGatherers);
 	const FString& Filename = SourceFile.GetFilename();
 	const FString& FileContent = SourceFile.GetContent();
 
@@ -4752,14 +4741,15 @@ void FNativeClassHeaderGenerator::CheckRPCFunctions(FReferenceGatherers& OutRefe
 void FNativeClassHeaderGenerator::ExportNativeFunctionHeader(
 	FOutputDevice&                   Out,
 	TSet<FString>&                   OutFwdDecls,
-	const FFuncInfo&                 FunctionData,
+	FUnrealFunctionDefinitionInfo&				 FunctionDef,
+	const FFuncInfo&				 FunctionData,
 	EExportFunctionType::Type        FunctionType,
 	EExportFunctionHeaderStyle::Type FunctionHeaderStyle,
 	const TCHAR*                     ExtraParam,
 	const TCHAR*                     APIString
 )
 {
-	UFunction* Function = FunctionData.FunctionReference;
+	UFunction* Function = FunctionDef.GetFunction();
 
 	const bool bIsDelegate   = Function->HasAnyFunctionFlags( FUNC_Delegate );
 	const bool bIsInterface  = !bIsDelegate && Function->GetOwnerClass()->HasAnyClassFlags(CLASS_Interface);
@@ -4807,7 +4797,8 @@ void FNativeClassHeaderGenerator::ExportNativeFunctionHeader(
 		}
 	}
 
-	FProperty* ReturnProperty = Function->GetReturnProperty();
+	FUnrealPropertyDefinitionInfo* ReturnPropertyDef = FunctionDef.GetReturn();
+	FProperty* ReturnProperty = ReturnPropertyDef ? ReturnPropertyDef->GetProperty() : nullptr;
 	if (ReturnProperty != nullptr)
 	{
 		if (ReturnProperty->HasAnyPropertyFlags(EPropertyFlags::CPF_ConstParm))
@@ -4858,8 +4849,9 @@ void FNativeClassHeaderGenerator::ExportNativeFunctionHeader(
 		++ParmCount;
 	}
 
-	for (FProperty* Property : TFieldRange<FProperty>(Function))
+	for (FUnrealPropertyDefinitionInfo* PropertyDef : FunctionDef.GetProperties())
 	{
+		FProperty* Property = PropertyDef->GetProperty();
 		if ((Property->PropertyFlags & (CPF_Parm | CPF_ReturnParm)) != CPF_Parm)
 		{
 			continue;
@@ -4874,8 +4866,7 @@ void FNativeClassHeaderGenerator::ExportNativeFunctionHeader(
 
 		FUHTStringBuilder PropertyText;
 
-		FUnrealPropertyDefinitionInfo& PropDef = GTypeDefinitionInfoMap.FindChecked(Property).AsPropertyChecked();
-		Property->ExportCppDeclaration( PropertyText, EExportedDeclaration::Parameter, PropDef.GetArrayDimensions() );
+		Property->ExportCppDeclaration( PropertyText, EExportedDeclaration::Parameter, PropertyDef->GetArrayDimensions() );
 		ApplyAlternatePropertyExportText(Property, PropertyText, EExportingState::Normal);
 
 		Out.Logf(TEXT("%s"), *PropertyText);
@@ -4901,11 +4892,11 @@ void FNativeClassHeaderGenerator::ExportNativeFunctionHeader(
 					FByteProperty* ByteProperty = CastField<FByteProperty>(ReturnProperty);
 					if (ByteProperty != nullptr && ByteProperty->Enum != nullptr && ByteProperty->Enum->GetCppForm() != UEnum::ECppForm::EnumClass)
 					{
-						ReturnValue = FString::Printf(TEXT(" return TEnumAsByte<%s>(%s); "), *ByteProperty->Enum->CppType, *GetNullParameterValue(ReturnProperty, false));
+						ReturnValue = FString::Printf(TEXT(" return TEnumAsByte<%s>(%s); "), *ByteProperty->Enum->CppType, *GetNullParameterValue(*ReturnPropertyDef, false));
 					}
 					else
 					{
-						ReturnValue = FString::Printf(TEXT(" return %s; "), *GetNullParameterValue(ReturnProperty, false));
+						ReturnValue = FString::Printf(TEXT(" return %s; "), *GetNullParameterValue(*ReturnPropertyDef, false));
 					}
 				}
 
@@ -4928,13 +4919,17 @@ void FNativeClassHeaderGenerator::ExportNativeFunctionHeader(
  * @param Return return parameter for the function
  * @param DeprecationWarningOutputDevice Device to output deprecation warnings for _Validate and _Implementation functions.
  */
-void FNativeClassHeaderGenerator::ExportFunctionThunk(FUHTStringBuilder& RPCWrappers, FReferenceGatherers& OutReferenceGatherers, UFunction* Function, const FFuncInfo& FunctionData, const TArray<FProperty*>& Parameters, FProperty* Return) const
+void FNativeClassHeaderGenerator::ExportFunctionThunk(FUHTStringBuilder& RPCWrappers, FReferenceGatherers& OutReferenceGatherers, FUnrealFunctionDefinitionInfo& FunctionDef, const TArray<FUnrealPropertyDefinitionInfo*>& ParameterDefs, FUnrealPropertyDefinitionInfo* ReturnDef) const
 {
+	UFunction* Function = FunctionDef.GetFunction();
+	const FFuncInfo& FunctionData = FunctionDef.GetFunctionData();
+
 	// export the GET macro for this parameter
 	FString ParameterList;
-	for (int32 ParameterIndex = 0; ParameterIndex < Parameters.Num(); ParameterIndex++)
+	for (int32 ParameterIndex = 0; ParameterIndex < ParameterDefs.Num(); ParameterIndex++)
 	{
-		FProperty* Param = Parameters[ParameterIndex];
+		FUnrealPropertyDefinitionInfo* ParamDef = ParameterDefs[ParameterIndex];
+		FProperty* Param = ParamDef->GetProperty();
 		OutReferenceGatherers.ForwardDeclarations.Add(Param->GetCPPTypeForwardDeclaration());
 
 		FString EvalBaseText = TEXT("P_GET_");	// e.g. P_GET_STR
@@ -5102,7 +5097,7 @@ void FNativeClassHeaderGenerator::ExportFunctionThunk(FUHTStringBuilder& RPCWrap
 
 	//Emit warning here if necessary
 	FUHTStringBuilder FunctionDeclaration;
-	ExportNativeFunctionHeader(FunctionDeclaration, OutReferenceGatherers.ForwardDeclarations, FunctionData, EExportFunctionType::Function, EExportFunctionHeaderStyle::Declaration, nullptr, *GetAPIString());
+	ExportNativeFunctionHeader(FunctionDeclaration, OutReferenceGatherers.ForwardDeclarations, FunctionDef, FunctionData, EExportFunctionType::Function, EExportFunctionHeaderStyle::Declaration, nullptr, *GetAPIString());
 
 	// Call the validate function if there is one
 	if (!(FunctionData.FunctionExportFlags & FUNCEXPORT_CppStatic) && (FunctionData.FunctionFlags & FUNC_NetValidate))
@@ -5116,8 +5111,9 @@ void FNativeClassHeaderGenerator::ExportFunctionThunk(FUHTStringBuilder& RPCWrap
 
 	// write out the return value
 	RPCWrappers.Log(TEXT("\t\t"));
-	if (Return)
+	if (ReturnDef)
 	{
+		FProperty* Return = ReturnDef->GetProperty();
 		OutReferenceGatherers.ForwardDeclarations.Add(Return->GetCPPTypeForwardDeclaration());
 
 		FUHTStringBuilder ReplacementText;
@@ -5145,13 +5141,14 @@ void FNativeClassHeaderGenerator::ExportFunctionThunk(FUHTStringBuilder& RPCWrap
 	RPCWrappers += TEXT("\t\tP_NATIVE_END;") LINE_TERMINATOR;
 }
 
-FString FNativeClassHeaderGenerator::GetFunctionParameterString(UFunction* Function, FReferenceGatherers& OutReferenceGatherers)
+FString FNativeClassHeaderGenerator::GetFunctionParameterString(FUnrealFunctionDefinitionInfo& FunctionDef, FReferenceGatherers& OutReferenceGatherers)
 {
 	FString ParameterList;
 	FUHTStringBuilder PropertyText;
 
-	for (FProperty* Property : TFieldRange<FProperty>(Function))
+	for (FUnrealPropertyDefinitionInfo* PropertyDef : FunctionDef.GetProperties())
 	{
+		FProperty* Property = PropertyDef->GetProperty();
 		OutReferenceGatherers.ForwardDeclarations.Add(Property->GetCPPTypeForwardDeclaration());
 
 		if ((Property->PropertyFlags & (CPF_Parm | CPF_ReturnParm)) != CPF_Parm)
@@ -5164,8 +5161,7 @@ FString FNativeClassHeaderGenerator::GetFunctionParameterString(UFunction* Funct
 			ParameterList += TEXT(", ");
 		}
 
-		FUnrealPropertyDefinitionInfo& PropDef = GTypeDefinitionInfoMap.FindChecked(Property).AsPropertyChecked();
-		Property->ExportCppDeclaration(PropertyText, EExportedDeclaration::Parameter, PropDef.GetArrayDimensions(), 0, true);
+		Property->ExportCppDeclaration(PropertyText, EExportedDeclaration::Parameter, PropertyDef->GetArrayDimensions(), 0, true);
 		ApplyAlternatePropertyExportText(Property, PropertyText, EExportingState::Normal);
 
 		ParameterList += PropertyText;
@@ -5185,11 +5181,13 @@ struct FNativeFunctionStringBuilder
 	FUHTStringBuilder AutogeneratedStaticDataFuncs;
 };
 
-void FNativeClassHeaderGenerator::ExportNativeFunctions(FOutputDevice& OutGeneratedHeaderText, FOutputDevice& OutGeneratedCPPText, FOutputDevice& OutMacroCalls, FOutputDevice& OutNoPureDeclsMacroCalls, FReferenceGatherers& OutReferenceGatherers, const FUnrealSourceFile& SourceFile, UClass* Class, FStructMetaData& StructData) const
+void FNativeClassHeaderGenerator::ExportNativeFunctions(FOutputDevice& OutGeneratedHeaderText, FOutputDevice& OutGeneratedCPPText, FOutputDevice& OutMacroCalls, FOutputDevice& OutNoPureDeclsMacroCalls, FReferenceGatherers& OutReferenceGatherers, const FUnrealSourceFile& SourceFile, FUnrealClassDefinitionInfo& ClassDef) const
 {
 	FNativeFunctionStringBuilder RuntimeStringBuilders;
 	FNativeFunctionStringBuilder EditorStringBuilders;
 
+	UClass* Class = ClassDef.GetClass();
+	FStructMetaData& StructData = ClassDef.GetStructMetaData();
 	const FString ClassCPPName = FNameLookupCPP::GetNameCPP(Class, Class->HasAnyClassFlags(CLASS_Interface));
 
 	ClassDefinitionRange ClassRange;
@@ -5215,12 +5213,13 @@ void FNativeClassHeaderGenerator::ExportNativeFunctions(FOutputDevice& OutGenera
 		RuntimeStringBuilders.AutogeneratedStaticData.Logf(TEXT("\treturn (F%s*)(GetClass()->GetOrCreateSparseClassData());\r\n"), *SparseClassDataString);
 		RuntimeStringBuilders.AutogeneratedStaticData += TEXT("}\r\n");
 
-		UScriptStruct* SparseClassDataStruct = FindObjectSafe<UScriptStruct>(ANY_PACKAGE, *SparseClassDataString);
+		UScriptStruct* SparseClassDataStruct = FClasses::FindObject<UScriptStruct>(ANY_PACKAGE, *SparseClassDataString);
 		while (SparseClassDataStruct != nullptr)
 		{
-			const FProperty* Child = CastField<FProperty>(SparseClassDataStruct->ChildProperties);
-			while (Child)
+			FUnrealScriptStructDefinitionInfo& SparseClassDataStructDef = GTypeDefinitionInfoMap.FindChecked<FUnrealScriptStructDefinitionInfo>(SparseClassDataStruct);
+			for (FUnrealPropertyDefinitionInfo* PropertyDef : SparseClassDataStructDef.GetProperties())
 			{
+				const FProperty* Child = PropertyDef->GetProperty();
 				FString ReturnExtendedType;
 				FString VarType = Child->GetCPPType(&ReturnExtendedType, EPropertyExportCPPFlags::CPPF_ArgumentOrReturnValue | EPropertyExportCPPFlags::CPPF_Implementation);
 				if (!ReturnExtendedType.IsEmpty())
@@ -5260,18 +5259,21 @@ void FNativeClassHeaderGenerator::ExportNativeFunctions(FOutputDevice& OutGenera
 					RuntimeStringBuilders.AutogeneratedStaticDataFuncs.Logf(TEXT("\treturn Get%s()->%s;\r\n"), *SparseClassDataString, *VarName);
 					RuntimeStringBuilders.AutogeneratedStaticDataFuncs.Logf(TEXT("}\r\n"));
 				}
-
-				Child = CastField<FProperty>(Child->Next);
 			}
 
 			SparseClassDataStruct = Cast<UScriptStruct>(SparseClassDataStruct->GetSuperStruct());
 		}
 	}
 
-	// export the C++ stubs
+	TArray<FUnrealFunctionDefinitionInfo*> Functions;
+	Algo::Copy(ClassDef.GetFunctions(), Functions);
+	Algo::Reverse(Functions);
 
-	for (UFunction* Function : TFieldRange<UFunction>(Class, EFieldIteratorFlags::ExcludeSuper))
+	// export the C++ stubs
+	for (FUnrealFunctionDefinitionInfo* FunctionDef : Functions)
 	{
+		UFunction* Function = FunctionDef->GetFunction();
+		const FFuncInfo& FunctionData = FunctionDef->GetFunctionData();
 		if (!(Function->FunctionFlags & FUNC_Native))
 		{
 			continue;
@@ -5280,8 +5282,7 @@ void FNativeClassHeaderGenerator::ExportNativeFunctions(FOutputDevice& OutGenera
 		const bool bEditorOnlyFunc = Function->HasAnyFunctionFlags(FUNC_EditorOnly);
 		FNativeFunctionStringBuilder& FuncStringBuilders = bEditorOnlyFunc ? EditorStringBuilders : RuntimeStringBuilders;
 
-		FUnrealFunctionDefinitionInfo& FuncDef = GTypeDefinitionInfoMap.FindChecked(Function).AsFunctionChecked();
-		const FFuncInfo& FunctionData = FuncDef.GetFunctionData();
+		FUnrealFunctionDefinitionInfo& FuncDef = GTypeDefinitionInfoMap.FindChecked<FUnrealFunctionDefinitionInfo>(Function);
 
 		// Custom thunks don't get any C++ stub function generated
 		if (FunctionData.FunctionExportFlags & FUNCEXPORT_CustomThunk)
@@ -5317,13 +5318,13 @@ void FNativeClassHeaderGenerator::ExportNativeFunctions(FOutputDevice& OutGenera
 
 			//Emit warning here if necessary
 			FUHTStringBuilder FunctionDeclaration;
-			ExportNativeFunctionHeader(FunctionDeclaration, OutReferenceGatherers.ForwardDeclarations, FunctionData, EExportFunctionType::Function, EExportFunctionHeaderStyle::Declaration, nullptr, *GetAPIString());
+			ExportNativeFunctionHeader(FunctionDeclaration, OutReferenceGatherers.ForwardDeclarations, FuncDef, FunctionData, EExportFunctionType::Function, EExportFunctionHeaderStyle::Declaration, nullptr, *GetAPIString());
 			FunctionDeclaration.Log(TEXT(";\r\n"));
 
 			// Declare validation function if needed
 			if (FunctionData.FunctionFlags & FUNC_NetValidate)
 			{
-				FString ParameterList = GetFunctionParameterString(Function, OutReferenceGatherers);
+				FString ParameterList = GetFunctionParameterString(FuncDef, OutReferenceGatherers);
 
 				const TCHAR* Virtual = (!FunctionData.FunctionReference->HasAnyFunctionFlags(FUNC_Static) && !(FunctionData.FunctionExportFlags & FUNCEXPORT_Final)) ? TEXT("virtual") : TEXT("");
 				FStringOutputDevice ValidDecl;
@@ -5344,7 +5345,7 @@ void FNativeClassHeaderGenerator::ExportNativeFunctions(FOutputDevice& OutGenera
 			// Versions that skip function autodeclaration throw an error when a function is missing.
 			if (ClassRange.bHasGeneratedBody && (SourceFile.GetGeneratedCodeVersionForStruct(Class) > EGeneratedCodeVersion::V1))
 			{
-				CheckRPCFunctions(OutReferenceGatherers, FunctionData, ClassCPPName, ImplementationPosition, ValidatePosition, SourceFile);
+				CheckRPCFunctions(OutReferenceGatherers, FuncDef, ClassCPPName, ImplementationPosition, ValidatePosition, SourceFile);
 			}
 		}
 
@@ -5362,8 +5363,8 @@ void FNativeClassHeaderGenerator::ExportNativeFunctions(FOutputDevice& OutGenera
 		FuncStringBuilders.RPCImplementations.Logf(TEXT("\tDEFINE_FUNCTION(%s::%s)"), *ClassCPPName, *FunctionData.UnMarshallAndCallName);
 		FuncStringBuilders.RPCImplementations += LINE_TERMINATOR TEXT("\t{") LINE_TERMINATOR;
 
-		FParmsAndReturnProperties Parameters = GetFunctionParmsAndReturn(FunctionData.FunctionReference);
-		ExportFunctionThunk(FuncStringBuilders.RPCImplementations, OutReferenceGatherers, Function, FunctionData, Parameters.Parms, Parameters.Return);
+		FParmsAndReturnProperties Parameters = GetFunctionParmsAndReturn(*FunctionDef);
+		ExportFunctionThunk(FuncStringBuilders.RPCImplementations, OutReferenceGatherers, FuncDef, Parameters.Parms, Parameters.Return);
 
 		FuncStringBuilders.RPCImplementations += TEXT("\t}") LINE_TERMINATOR;
 	}
@@ -5462,7 +5463,7 @@ void FNativeClassHeaderGenerator::ExportCallbackFunctions(
 	FOutputDevice&            OutGeneratedHeaderText,
 	FOutputDevice&            OutCpp,
 	TSet<FString>&            OutFwdDecls,
-	const TArray<UFunction*>& CallbackFunctions,
+	const TArray<FUnrealFunctionDefinitionInfo*>& CallbackFunctions,
 	const TCHAR*              CallbackWrappersMacroName,
 	EExportCallbackType       ExportCallbackType,
 	const TCHAR*              APIString
@@ -5471,13 +5472,14 @@ void FNativeClassHeaderGenerator::ExportCallbackFunctions(
 	FUHTStringBuilder RPCWrappers;
 
 	FMacroBlockEmitter OutCppEditorOnly(OutCpp, TEXT("WITH_EDITOR"));
-	for (UFunction* Function : CallbackFunctions)
+	for (FUnrealFunctionDefinitionInfo* FunctionDef : CallbackFunctions)
 	{
+		UFunction* Function = FunctionDef->GetFunction();
+
 		// Never expecting to export delegate functions this way
 		check(!Function->HasAnyFunctionFlags(FUNC_Delegate));
 
-		FUnrealFunctionDefinitionInfo& FuncDef = GTypeDefinitionInfoMap.FindChecked(Function).AsFunctionChecked();
-		const FFuncInfo& FunctionData = FuncDef.GetFunctionData();
+		const FFuncInfo& FunctionData = FunctionDef->GetFunctionData();
 		FString          FunctionName = Function->GetName();
 		UClass*          Class = CastChecked<UClass>(Function->GetOuter());
 		const FString    ClassName = FNameLookupCPP::GetNameCPP(Class);
@@ -5498,7 +5500,7 @@ void FNativeClassHeaderGenerator::ExportCallbackFunctions(
 		if (!bWillBeProgrammerTyped)
 		{
 			// export the line that looks like: int32 Main(const FString& Parms)
-			ExportNativeFunctionHeader(RPCWrappers, OutFwdDecls, FunctionData, EExportFunctionType::Event, EExportFunctionHeaderStyle::Declaration, nullptr, APIString);
+			ExportNativeFunctionHeader(RPCWrappers, OutFwdDecls, *FunctionDef, FunctionData, EExportFunctionType::Event, EExportFunctionHeaderStyle::Declaration, nullptr, APIString);
 
 			RPCWrappers.Log(TEXT(";\r\n"));
 			RPCWrappers.Log(TEXT("\r\n"));
@@ -5512,9 +5514,9 @@ void FNativeClassHeaderGenerator::ExportCallbackFunctions(
 		}
 
 		// Emit the thunk implementation
-		ExportNativeFunctionHeader(OutCpp, OutFwdDecls, FunctionData, EExportFunctionType::Event, EExportFunctionHeaderStyle::Definition, nullptr, APIString);
+		ExportNativeFunctionHeader(OutCpp, OutFwdDecls, *FunctionDef, FunctionData, EExportFunctionType::Event, EExportFunctionHeaderStyle::Definition, nullptr, APIString);
 
-		FParmsAndReturnProperties Parameters = GetFunctionParmsAndReturn(FunctionData.FunctionReference);
+		FParmsAndReturnProperties Parameters = GetFunctionParmsAndReturn(*FunctionDef);
 
 		if (ExportCallbackType != EExportCallbackType::Interface)
 		{
@@ -5805,7 +5807,7 @@ void FNativeClassHeaderGenerator::GenerateSourceFiles(
 						// Is this ever not the case?
 						if (Enum->GetOuter()->IsA(UPackage::StaticClass()))
 						{
-							GeneratedFunctionDeclarations.Log(GTypeDefinitionInfoMap.FindChecked(Enum).AsFieldChecked().GetExternDecl(true));
+							GeneratedFunctionDeclarations.Log(GTypeDefinitionInfoMap.FindChecked<FUnrealFieldDefinitionInfo>(Enum).GetExternDecl(true));
 							Generator.ExportGeneratedEnumInitCode(GeneratedCPPText, ReferenceGatherers, SourceFile, Enum);
 						}
 					}
@@ -5814,29 +5816,32 @@ void FNativeClassHeaderGenerator::GenerateSourceFiles(
 					// reverse the order.
 					for (UScriptStruct* Struct : Structs)
 					{
-						GeneratedFunctionDeclarations.Log(GTypeDefinitionInfoMap.FindChecked(Struct).AsFieldChecked().GetExternDecl(true));
+						GeneratedFunctionDeclarations.Log(GTypeDefinitionInfoMap.FindChecked<FUnrealFieldDefinitionInfo>(Struct).GetExternDecl(true));
 						Generator.ExportGeneratedStructBodyMacros(GeneratedHeaderText, GeneratedCPPText, ReferenceGatherers, SourceFile, Struct);
 					}
 
 					// export delegate definitions
 					for (UDelegateFunction* Func : DelegateFunctions)
 					{
-						GeneratedFunctionDeclarations.Log(GTypeDefinitionInfoMap.FindChecked(Func).AsFieldChecked().GetExternDecl(true));
-						Generator.ExportDelegateDeclaration(GeneratedCPPText, ReferenceGatherers, SourceFile, Func);
+						FUnrealFunctionDefinitionInfo& FunctionDef = GTypeDefinitionInfoMap.FindChecked<FUnrealFunctionDefinitionInfo>(Func);
+						GeneratedFunctionDeclarations.Log(FunctionDef.GetExternDecl(true));
+						Generator.ExportDelegateDeclaration(GeneratedCPPText, ReferenceGatherers, SourceFile, FunctionDef);
 					}
 
 					// export delegate wrapper function implementations
 					for (UDelegateFunction* Func : DelegateFunctions)
 					{
-						Generator.ExportDelegateDefinition(GeneratedHeaderText, ReferenceGatherers, SourceFile, Func);
+						FUnrealFunctionDefinitionInfo& FunctionDef = GTypeDefinitionInfoMap.FindChecked<FUnrealFunctionDefinitionInfo>(Func);
+						Generator.ExportDelegateDefinition(GeneratedHeaderText, ReferenceGatherers, SourceFile, FunctionDef);
 					}
 
-					for (const TSharedRef<FUnrealTypeDefinitionInfo>& ClassDef : SourceFile.GetDefinedClasses())
+					for (const TSharedRef<FUnrealTypeDefinitionInfo>& TypeDef : SourceFile.GetDefinedClasses())
 					{
-						UClass* Class = ClassDef->AsClassChecked().GetClass();
+						FUnrealClassDefinitionInfo& ClassDef = TypeDef->AsClassChecked();
+						UClass* Class = ClassDef.GetClass();
 						if (!(Class->ClassFlags & CLASS_Intrinsic))
 						{
-							Generator.ExportClassFromSourceFileInner(GeneratedHeaderText, GeneratedCPPText, GeneratedFunctionDeclarations, ReferenceGatherers, (FClass*)Class, SourceFile, GeneratedCPP.ExportFlags);
+							Generator.ExportClassFromSourceFileInner(GeneratedHeaderText, GeneratedCPPText, GeneratedFunctionDeclarations, ReferenceGatherers, ClassDef, SourceFile, GeneratedCPP.ExportFlags);
 						}
 					}
 
@@ -6168,7 +6173,7 @@ void FNativeClassHeaderGenerator::DeleteUnusedGeneratedHeaders(TSet<FString>&& P
 								FString BaseFilename = FPaths::GetBaseFilename(Filename);
 								const int32   GeneratedIndex = BaseFilename.Find(TEXT(".generated"), ESearchCase::IgnoreCase, ESearchDir::FromEnd);
 								const FString ClassName = MoveTemp(BaseFilename).Mid(0, GeneratedIndex);
-								UClass* IntrinsicClass = FindObject<UClass>(ANY_PACKAGE, *ClassName);
+								UClass* IntrinsicClass = FClasses::FindObject<UClass>(ANY_PACKAGE, *ClassName);
 								if (!IntrinsicClass || !IntrinsicClass->HasAnyClassFlags(CLASS_Intrinsic))
 								{
 									IFileManager::Get().Delete(*Fullpath);
@@ -6430,7 +6435,7 @@ void ResolveSuperClasses(UPackage* Package)
 			continue;
 		}
 
-		FUnrealClassDefinitionInfo& ClassDef = GTypeDefinitionInfoMap.FindChecked(DefinedClass).AsClassChecked();
+		FUnrealClassDefinitionInfo& ClassDef = GTypeDefinitionInfoMap.FindChecked<FUnrealClassDefinitionInfo>(DefinedClass);
 
 		// Resolve the base class
 		{
@@ -6440,11 +6445,11 @@ void ResolveSuperClasses(UPackage* Package)
 
 			if (!BaseClassNameStripped.IsEmpty() && !DefinedClass->GetSuperClass())
 			{
-				UClass* FoundBaseClass = FindObject<UClass>(Package, *BaseClassNameStripped);
+				UClass* FoundBaseClass = FClasses::FindObject<UClass>(Package, *BaseClassNameStripped);
 
 				if (FoundBaseClass == nullptr)
 				{
-					FoundBaseClass = FindObject<UClass>(ANY_PACKAGE, *BaseClassNameStripped);
+					FoundBaseClass = FClasses::FindObject<UClass>(ANY_PACKAGE, *BaseClassNameStripped);
 				}
 
 				if (FoundBaseClass == nullptr)
@@ -6455,7 +6460,7 @@ void ResolveSuperClasses(UPackage* Package)
 
 				DefinedClass->SetSuperStruct(FoundBaseClass);
 				DefinedClass->ClassCastFlags |= FoundBaseClass->ClassCastFlags;
-				SuperClassInfo.Struct = GTypeDefinitionInfoMap.FindChecked(FoundBaseClass).AsStruct();
+				SuperClassInfo.Struct = &GTypeDefinitionInfoMap.FindChecked<FUnrealStructDefinitionInfo>(FoundBaseClass);
 			}
 		}
 
@@ -6465,7 +6470,7 @@ void ResolveSuperClasses(UPackage* Package)
 			{
 				if (UClass* FoundClass = FClasses::FindScriptClass(BaseClassInfo.Name))
 				{
-					BaseClassInfo.Struct = GTypeDefinitionInfoMap.FindChecked(FoundClass).AsStruct();
+					BaseClassInfo.Struct = &GTypeDefinitionInfoMap.FindChecked<FUnrealStructDefinitionInfo>(FoundClass);
 				}
 			}
 		}
@@ -6474,7 +6479,7 @@ void ResolveSuperClasses(UPackage* Package)
 
 UPackage* GetModulePackage(FManifestModule& Module)
 {
-	UPackage* Package = Cast<UPackage>(StaticFindObjectFast(UPackage::StaticClass(), NULL, FName(*Module.LongPackageName), false, false));
+	UPackage* Package = FClasses::FindObjectFast<UPackage>(NULL, FName(*Module.LongPackageName), false, false);
 	if (Package == NULL)
 	{
 		Package = CreatePackage(*Module.LongPackageName);
@@ -6549,11 +6554,12 @@ void PrepareModules(TArray<FUnrealPackageDefinitionInfo*>& PackageDefs, const FS
 		UPackage* Package = GetModulePackage(Module);
 
 		// Create the package definition
-		FUnrealPackageDefinitionInfo* PackageDef = new FUnrealPackageDefinitionInfo(Module, Package);
-		GTypeDefinitionInfoMap.Add(Package, TSharedRef<FUnrealTypeDefinitionInfo>(PackageDef));
-		PackageDefs.Add(PackageDef);
+		TSharedRef<FUnrealPackageDefinitionInfo> PackageDefRef = MakeShared<FUnrealPackageDefinitionInfo>(Module, Package);
+		FUnrealPackageDefinitionInfo& PackageDef = *PackageDefRef;
+		GTypeDefinitionInfoMap.Add(Package, PackageDefRef);
+		PackageDefs.Add(&PackageDef);
 
-		TArray<TSharedRef<FUnrealSourceFile>>& AllSourceFiles = PackageDef->GetAllSourceFiles();
+		TArray<TSharedRef<FUnrealSourceFile>>& AllSourceFiles = PackageDef.GetAllSourceFiles();
 		AllSourceFiles.Reserve(Module.PublicUObjectClassesHeaders.Num() + Module.PublicUObjectHeaders.Num() + Module.PrivateUObjectHeaders.Num());
 
 		// Initialize the other header data structures and create the unreal source files
@@ -6577,7 +6583,7 @@ void PrepareModules(TArray<FUnrealPackageDefinitionInfo*>& PackageDefs, const FS
 				const FString& RawFilename = UObjectHeaders[Index];
 				const FString FullFilename = FPaths::ConvertRelativePathToFull(ModuleInfoPath, RawFilename);
 
-				FUnrealSourceFile* UnrealSourceFilePtr = new FUnrealSourceFile(Package, RawFilename);
+				FUnrealSourceFile* UnrealSourceFilePtr = new FUnrealSourceFile(PackageDef, RawFilename);
 				TSharedRef<FUnrealSourceFile> UnrealSourceFile(UnrealSourceFilePtr);
 				AllSourceFiles.Add(UnrealSourceFile);
 
@@ -6769,15 +6775,27 @@ void ResolveParents(TArray<FUnrealPackageDefinitionInfo*>& PackageDefs)
 void PrepareTypesForParsing(TArray<FUnrealPackageDefinitionInfo*>& PackageDefs)
 {
 	// Verify that all script declared superclasses exist.
+	TMap<const UClass*, TArray<const UClass*>> Uniques;
 	for (UClass* ScriptClass : TObjectRange<UClass>())
 	{
 		const UClass* ScriptSuperClass = ScriptClass->GetSuperClass();
+		if (ScriptSuperClass)
+		{
+			TSharedRef<FUnrealTypeDefinitionInfo>* ScriptClassTypeDef = GTypeDefinitionInfoMap.Find(ScriptClass);
+			TSharedRef<FUnrealTypeDefinitionInfo>* SuperTypeDef = GTypeDefinitionInfoMap.Find(ScriptSuperClass);
+			if (ScriptClassTypeDef && (*ScriptClassTypeDef)->HasSource() && SuperTypeDef && !(*SuperTypeDef)->HasSource())
+			{
+				UE_LOG(LogCompile, Log, TEXT("%s has been used"), *(*SuperTypeDef)->GetNameCPP());
+				Uniques.FindOrAdd(ScriptSuperClass).Add(ScriptClass);
+			}
+		}
 
 		if (ScriptSuperClass && !ScriptSuperClass->HasAnyClassFlags(CLASS_Intrinsic) && GTypeDefinitionInfoMap.Contains(ScriptClass) && !GTypeDefinitionInfoMap.Contains(ScriptSuperClass))
 		{
 			FResults::LogError(*ScriptClass, *FString::Printf(TEXT("Superclass %s of class %s not found"), *ScriptSuperClass->GetName(), *ScriptClass->GetName()));
 		}
 	}
+	UE_LOG(LogCompile, Log, TEXT("BOB"));
 }
 
 void TopologicalRecursion(FUnrealSourceFile& First, FUnrealSourceFile& Visit)
@@ -6864,7 +6882,7 @@ void ParseSourceFiles(TArray<FUnrealSourceFile*>& OrderedSourceFiles)
 	for (FUnrealSourceFile* SourceFile : OrderedSourceFiles)
 	{
 		UPackage* Package = SourceFile->GetPackage();
-		FUnrealPackageDefinitionInfo& PackageDef = GTypeDefinitionInfoMap.FindChecked(Package).AsPackageChecked();
+		FUnrealPackageDefinitionInfo& PackageDef = GTypeDefinitionInfoMap.FindChecked<FUnrealPackageDefinitionInfo>(Package);
 		FScopedDurationTimer SourceTimer(SourceFile->GetTime(ESourceFileTime::Parse));
 
 		ECompilationResult::Type Result = FHeaderParser::Parse(PackageDef, GWarn, *SourceFile);
@@ -6883,31 +6901,33 @@ void PrepareTypesForExport(TArray<FUnrealPackageDefinitionInfo*>& PackageDefs)
 		FResults::Try([PackageDef]()
 			{
 				// Collect all of the classes from the sources
-				TArray<FClass*> PackageClasses;
+				TArray<FUnrealClassDefinitionInfo*> PackageClasses;
 				PackageClasses.Reserve(128); // Arbitrary number
 				for (TSharedRef<FUnrealSourceFile>& SourceFile : PackageDef->GetAllSourceFiles())
 				{
 					for (TSharedRef<FUnrealTypeDefinitionInfo>& TypeDef : SourceFile->GetDefinedClasses())
 					{
-						PackageClasses.Add(static_cast<FClass*>(TypeDef->AsClassChecked().GetClass()));
+						PackageClasses.Add(&TypeDef->AsClassChecked());
 					}
 				}
 
 				// Validate the sparse class data for all classes in the current package
-				for (const FClass* Class : PackageClasses)
+				for (const FUnrealClassDefinitionInfo* ClassDef : PackageClasses)
 				{
-					FHeaderParser::CheckSparseClassData(Class);
+					FHeaderParser::CheckSparseClassData(ClassDef->GetClass());
 				}
 
 				// Check to see if we need write the classes 
-				const bool bPackageHasAnyExportClasses = PackageClasses.ContainsByPredicate([](FClass* Class)
+				const bool bPackageHasAnyExportClasses = PackageClasses.ContainsByPredicate([](FUnrealClassDefinitionInfo* ClassDef)
 					{
+						UClass* Class = ClassDef->GetClass();
 						return Class->HasAnyClassFlags(CLASS_Native) && !Class->HasAnyClassFlags(CLASS_NoExport | CLASS_Intrinsic);
 					});
 				if (bPackageHasAnyExportClasses)
 				{
-					for (FClass* Class : PackageClasses)
+					for (FUnrealClassDefinitionInfo* ClassDef : PackageClasses)
 					{
+						UClass* Class = ClassDef->GetClass();
 						if (!Class->HasAnyClassFlags(CLASS_Native))
 						{
 							Class->UnMark(EObjectMark(OBJECTMARK_TagImp | OBJECTMARK_TagExp));
@@ -6924,11 +6944,11 @@ void PrepareTypesForExport(TArray<FUnrealPackageDefinitionInfo*>& PackageDefs)
 				// This needs to be done outside of parallel blocks because it will modify UClass memory.
 				// Later calls to SetUpUhtReplicationData inside parallel blocks should be fine, because
 				// they will see the memory has already been set up, and just return the parent pointer.
-				for (FClass* Class : PackageClasses)
+				for (FUnrealClassDefinitionInfo* ClassDef : PackageClasses)
 				{
-					if (ClassHasReplicatedProperties(Class))
+					if (ClassHasReplicatedProperties(*ClassDef))
 					{
-						Class->SetUpUhtReplicationData();
+						ClassDef->GetClass()->SetUpUhtReplicationData();
 					}
 				}
 			}
@@ -6945,7 +6965,7 @@ void Export(TArray<FUnrealPackageDefinitionInfo*>& PackageDefs, TArray<FUnrealSo
 	GeneratedCPPs.Reserve(OrderedSourceFiles.Num());
 	for (FUnrealSourceFile* SourceFile : OrderedSourceFiles)
 	{
-		FUnrealPackageDefinitionInfo& PackageDef = GTypeDefinitionInfoMap.FindChecked(SourceFile->GetPackage()).AsPackageChecked();
+		FUnrealPackageDefinitionInfo& PackageDef = GTypeDefinitionInfoMap.FindChecked<FUnrealPackageDefinitionInfo>(SourceFile->GetPackage());
 		GeneratedCPPs.Emplace(PackageDef, *SourceFile);
 	}
 
@@ -7236,7 +7256,7 @@ UClass* ProcessParsedClass(bool bClassIsAnInterface, const FString& ClassName, c
 		FError::Throwf(TEXT("Class '%s' cannot inherit from itself"), *ClassName);
 	}
 
-	UClass* ResultClass = FindObject<UClass>(InParent, *ClassNameStripped);
+	UClass* ResultClass = FClasses::FindObject<UClass>(InParent, *ClassNameStripped);
 
 	// if we aren't generating headers, then we shouldn't set misaligned object, since it won't get cleared
 
@@ -7247,7 +7267,7 @@ UClass* ProcessParsedClass(bool bClassIsAnInterface, const FString& ClassName, c
 		// detect if the same class name is used in multiple packages
 		if (ResultClass == nullptr)
 		{
-			UClass* ConflictingClass = FindObject<UClass>(ANY_PACKAGE, *ClassNameStripped, true);
+			UClass* ConflictingClass = FClasses::FindObject<UClass>(ANY_PACKAGE, *ClassNameStripped, true);
 			if (ConflictingClass != nullptr)
 			{
 				UE_LOG_WARNING_UHT(TEXT("Duplicate class name: %s also exists in file %s"), *ClassName, *ConflictingClass->GetOutermost()->GetName());
@@ -7276,7 +7296,7 @@ UClass* ProcessParsedClass(bool bClassIsAnInterface, const FString& ClassName, c
 
 UEnum* ProcessParsedEnum(const FString& EnumName, UObject* InParent, EObjectFlags Flags)
 {
-	if (UEnum* Existing = FindObject<UEnum>(ANY_PACKAGE, *EnumName))
+	if (UEnum* Existing = FClasses::FindObject<UEnum>(ANY_PACKAGE, *EnumName))
 	{
 		FError::Throwf(TEXT("Duplicate enum name: %s also exists in file %s"), *EnumName, *Existing->GetOutermost()->GetName());
 	}
@@ -7296,7 +7316,7 @@ UScriptStruct* ProcessParsedStruct(const FString& StructName, UObject* InParent,
 {
 	FString StructNameStripped = GetClassNameWithPrefixRemoved(*StructName);
 
-	if (UScriptStruct* Existing = FindObject<UScriptStruct>(ANY_PACKAGE, *StructNameStripped))
+	if (UScriptStruct* Existing = FClasses::FindObject<UScriptStruct>(ANY_PACKAGE, *StructNameStripped))
 	{
 		FError::Throwf(TEXT("Duplicate struct name: %s also exists in file %s"), *StructNameStripped, *Existing->GetOutermost()->GetName());
 	}
