Added new PropertyTypes which in the long run will be where different property types are implemented.

diff --git a/Engine/Source/Programs/UnrealHeaderTool/Private/HeaderParser.cpp b/Engine/Source/Programs/UnrealHeaderTool/Private/HeaderParser.cpp
--- a/Engine/Source/Programs/UnrealHeaderTool/Private/HeaderParser.cpp
+++ b/Engine/Source/Programs/UnrealHeaderTool/Private/HeaderParser.cpp
@@ -3,6 +3,7 @@
 
 #include "HeaderParser.h"
 #include "UnrealHeaderTool.h"
+#include "PropertyTypes.h"
 #include "HAL/FileManager.h"
 #include "Misc/CommandLine.h"
 #include "Misc/ConfigCacheIni.h"
@@ -95,7 +96,6 @@ namespace
 	static const FName NAME_CannotImplementInterfaceInBlueprint(TEXT("CannotImplementInterfaceInBlueprint"));
 	static const FName NAME_UIMin(TEXT("UIMin"));
 	static const FName NAME_UIMax(TEXT("UIMax"));
-	static const FName NAME_BlueprintType(TEXT("BlueprintType"));
 }
 
 const FName FHeaderParserNames::NAME_HideCategories(TEXT("HideCategories"));
@@ -103,6 +103,7 @@ const FName FHeaderParserNames::NAME_ShowCategories(TEXT("ShowCategories"));
 const FName FHeaderParserNames::NAME_SparseClassDataTypes(TEXT("SparseClassDataTypes"));
 const FName FHeaderParserNames::NAME_IsConversionRoot(TEXT("IsConversionRoot"));
 const FName FHeaderParserNames::NAME_AdvancedClassDisplay(TEXT("AdvancedClassDisplay"));
+const FName FHeaderParserNames::NAME_BlueprintType(TEXT("BlueprintType"));
 
 FRigVMStructMap FHeaderParser::StructRigVMMap;
 TArray<FString> FHeaderParser::PropertyCPPTypesRequiringUIRanges = { TEXT("float"), TEXT("double") };
@@ -615,11 +616,6 @@ namespace
 		}
 	}
 
-	void AddEditInlineMetaData(TMap<FName, FString>& MetaData)
-	{
-		MetaData.Add(NAME_EditInline, TEXT("true"));
-	}
-
 	const TCHAR* GetHintText(EVariableCategory::Type VariableCategory)
 	{
 		switch (VariableCategory)
@@ -643,414 +639,12 @@ namespace
 		return nullptr;
 	}
 
-	// Check to see if anything in the class hierarchy passed in has CLASS_DefaultToInstanced
-	bool DoesAnythingInHierarchyHaveDefaultToInstanced(UClass* TestClass)
-	{
-		bool bDefaultToInstanced = false;
-
-		UClass* Search = TestClass;
-		while (!bDefaultToInstanced && (Search != NULL))
-		{
-			bDefaultToInstanced = Search->HasAnyClassFlags(CLASS_DefaultToInstanced);
-			if (!bDefaultToInstanced && !Search->HasAnyClassFlags(CLASS_Intrinsic | CLASS_Parsed))
-			{
-				// The class might not have been parsed yet, look for declaration data.
-				if (TSharedRef<FUnrealTypeDefinitionInfo>* Info = GTypeDefinitionInfoMap.Find(Search))
-				{
-					bDefaultToInstanced = !!((**Info).AsClassChecked().ClassFlags & CLASS_DefaultToInstanced);
-				}
-			}
-			Search = Search->GetSuperClass();
-		}
-
-		return bDefaultToInstanced;
-	}
-
-	FProperty* CreateVariableProperty(FPropertyBase& VarProperty, FFieldVariant Scope, const FName& Name, EObjectFlags ObjectFlags, EVariableCategory::Type VariableCategory, FUnrealSourceFile* UnrealSourceFile)
-	{
-		check(UnrealSourceFile);
-
-		// Check if it's an enum class property
-		// NOTE: VarProperty.Enum is a union and might not be an enum
-		if (TSharedRef<FUnrealTypeDefinitionInfo>* TypeDef = GTypeDefinitionInfoMap.Find(VarProperty.Enum))
-		{
-			if (FUnrealEnumDefinitionInfo* EnumDef = (*TypeDef)->AsEnum())
-			{
-				if (VarProperty.Enum->GetCppForm() == UEnum::ECppForm::EnumClass)
-				{
-					FPropertyBase UnderlyingProperty = VarProperty;
-					UnderlyingProperty.Enum = nullptr;
-					switch (EnumDef->GetUnderlyingType())
-					{
-					case EUnderlyingEnumType::int8:        UnderlyingProperty.Type = CPT_Int8;   break;
-					case EUnderlyingEnumType::int16:       UnderlyingProperty.Type = CPT_Int16;  break;
-					case EUnderlyingEnumType::int32:       UnderlyingProperty.Type = CPT_Int;    break;
-					case EUnderlyingEnumType::int64:       UnderlyingProperty.Type = CPT_Int64;  break;
-					case EUnderlyingEnumType::uint8:       UnderlyingProperty.Type = CPT_Byte;   break;
-					case EUnderlyingEnumType::uint16:      UnderlyingProperty.Type = CPT_UInt16; break;
-					case EUnderlyingEnumType::uint32:      UnderlyingProperty.Type = CPT_UInt32; break;
-					case EUnderlyingEnumType::uint64:      UnderlyingProperty.Type = CPT_UInt64; break;
-					case EUnderlyingEnumType::Unspecified:
-						UnderlyingProperty.Type = CPT_Int;
-						UnderlyingProperty.IntType = EIntType::Unsized;
-						break;
-
-					default:
-						check(false);
-					}
-
-					FEnumProperty* Result = new FEnumProperty(Scope, Name, ObjectFlags);
-					GTypeDefinitionInfoMap.Add(Result, MakeShared<FUnrealPropertyDefinitionInfo>(*UnrealSourceFile, 0, Result));
-					FNumericProperty* UnderlyingProp = CastFieldChecked<FNumericProperty>(CreateVariableProperty(UnderlyingProperty, Result, TEXT("UnderlyingType"), ObjectFlags, VariableCategory, UnrealSourceFile));
-					Result->UnderlyingProp = MoveTemp(UnderlyingProp);
-					Result->Enum = VarProperty.Enum;
-#if UHT_ENABLE_VALUE_PROPERTY_TAG
-					UnrealSourceFile->AddTypeDefIncludeIfNeeded(*EnumDef);
-#endif
-					return Result;
-				}
-			}
-		}
-
-		switch (VarProperty.Type)
-		{
-			case CPT_Byte:
-			{
-#if UHT_ENABLE_VALUE_PROPERTY_TAG
-				UnrealSourceFile->AddTypeDefIncludeIfNeeded(VarProperty.Enum);
-#endif
-				FByteProperty* Result = new FByteProperty(Scope, Name, ObjectFlags);
-				GTypeDefinitionInfoMap.Add(Result, MakeShared<FUnrealPropertyDefinitionInfo>(*UnrealSourceFile, 0, Result));
-				Result->Enum = VarProperty.Enum;
-				check(VarProperty.IntType == EIntType::Sized);
-				return Result;
-			}
-
-			case CPT_Int8:
-			{
-				FInt8Property* Result = new FInt8Property(Scope, Name, ObjectFlags);
-				GTypeDefinitionInfoMap.Add(Result, MakeShared<FUnrealPropertyDefinitionInfo>(*UnrealSourceFile, 0, Result));
-				check(VarProperty.IntType == EIntType::Sized);
-				return Result;
-			}
-
-			case CPT_Int16:
-			{
-				FInt16Property* Result = new FInt16Property(Scope, Name, ObjectFlags);
-				GTypeDefinitionInfoMap.Add(Result, MakeShared<FUnrealPropertyDefinitionInfo>(*UnrealSourceFile, 0, Result));
-				check(VarProperty.IntType == EIntType::Sized);
-				return Result;
-			}
-
-			case CPT_Int:
-			{
-				FIntProperty* Result = new FIntProperty(Scope, Name, ObjectFlags);
-				GTypeDefinitionInfoMap.Add(Result, MakeShared<FUnrealPropertyDefinitionInfo>(*UnrealSourceFile, 0, Result, VarProperty.IntType == EIntType::Unsized));
-				return Result;
-			}
-
-			case CPT_Int64:
-			{
-				FInt64Property* Result = new FInt64Property(Scope, Name, ObjectFlags);
-				GTypeDefinitionInfoMap.Add(Result, MakeShared<FUnrealPropertyDefinitionInfo>(*UnrealSourceFile, 0, Result));
-				check(VarProperty.IntType == EIntType::Sized);
-				return Result;
-			}
-
-			case CPT_UInt16:
-			{
-				FUInt16Property* Result = new FUInt16Property(Scope, Name, ObjectFlags);
-				GTypeDefinitionInfoMap.Add(Result, MakeShared<FUnrealPropertyDefinitionInfo>(*UnrealSourceFile, 0, Result));
-				check(VarProperty.IntType == EIntType::Sized);
-				return Result;
-			}
-
-			case CPT_UInt32:
-			{
-				FUInt32Property* Result = new FUInt32Property(Scope, Name, ObjectFlags);
-				GTypeDefinitionInfoMap.Add(Result, MakeShared<FUnrealPropertyDefinitionInfo>(*UnrealSourceFile, 0, Result, VarProperty.IntType == EIntType::Unsized));
-				return Result;
-			}
-
-			case CPT_UInt64:
-			{
-				FUInt64Property* Result = new FUInt64Property(Scope, Name, ObjectFlags);
-				GTypeDefinitionInfoMap.Add(Result, MakeShared<FUnrealPropertyDefinitionInfo>(*UnrealSourceFile, 0, Result));
-				check(VarProperty.IntType == EIntType::Sized);
-				return Result;
-			}
-
-			case CPT_Bool:
-			{
-				FBoolProperty* Result = new FBoolProperty(Scope, Name, ObjectFlags);
-				GTypeDefinitionInfoMap.Add(Result, MakeShared<FUnrealPropertyDefinitionInfo>(*UnrealSourceFile, 0, Result));
-				Result->SetBoolSize(sizeof(bool), true);
-				return Result;
-			}
-
-			case CPT_Bool8:
-			{
-				FBoolProperty* Result = new FBoolProperty(Scope, Name, ObjectFlags);
-				GTypeDefinitionInfoMap.Add(Result, MakeShared<FUnrealPropertyDefinitionInfo>(*UnrealSourceFile, 0, Result));
-				Result->SetBoolSize((VariableCategory == EVariableCategory::Return) ? sizeof(bool) : sizeof(uint8), VariableCategory == EVariableCategory::Return);
-				return Result;
-			}
-
-			case CPT_Bool16:
-			{
-				FBoolProperty* Result = new FBoolProperty(Scope, Name, ObjectFlags);
-				GTypeDefinitionInfoMap.Add(Result, MakeShared<FUnrealPropertyDefinitionInfo>(*UnrealSourceFile, 0, Result));
-				Result->SetBoolSize((VariableCategory == EVariableCategory::Return) ? sizeof(bool) : sizeof(uint16), VariableCategory == EVariableCategory::Return);
-				return Result;
-			}
-
-			case CPT_Bool32:
-			{
-				FBoolProperty* Result = new FBoolProperty(Scope, Name, ObjectFlags);
-				GTypeDefinitionInfoMap.Add(Result, MakeShared<FUnrealPropertyDefinitionInfo>(*UnrealSourceFile, 0, Result));
-				Result->SetBoolSize((VariableCategory == EVariableCategory::Return) ? sizeof(bool) : sizeof(uint32), VariableCategory == EVariableCategory::Return);
-				return Result;
-			}
-
-			case CPT_Bool64:
-			{
-				FBoolProperty* Result = new FBoolProperty(Scope, Name, ObjectFlags);
-				GTypeDefinitionInfoMap.Add(Result, MakeShared<FUnrealPropertyDefinitionInfo>(*UnrealSourceFile, 0, Result));
-				Result->SetBoolSize((VariableCategory == EVariableCategory::Return) ? sizeof(bool) : sizeof(uint64), VariableCategory == EVariableCategory::Return);
-				return Result;
-			}
-
-			case CPT_Float:
-			{
-				FFloatProperty* Result = new FFloatProperty(Scope, Name, ObjectFlags);
-				GTypeDefinitionInfoMap.Add(Result, MakeShared<FUnrealPropertyDefinitionInfo>(*UnrealSourceFile, 0, Result));
-				return Result;
-			}
-
-			case CPT_Double:
-			{
-				FDoubleProperty* Result = new FDoubleProperty(Scope, Name, ObjectFlags);
-				GTypeDefinitionInfoMap.Add(Result, MakeShared<FUnrealPropertyDefinitionInfo>(*UnrealSourceFile, 0, Result));
-				return Result;
-			}
-
-			case CPT_FLargeWorldCoordinatesReal:
-			{
-				FLargeWorldCoordinatesRealProperty* Result = new FLargeWorldCoordinatesRealProperty(Scope, Name, ObjectFlags);
-				GTypeDefinitionInfoMap.Add(Result, MakeShared<FUnrealPropertyDefinitionInfo>(*UnrealSourceFile, 0, Result));
-				return Result;
-			}
-
-			case CPT_ObjectReference:
-				check(VarProperty.PropertyClass);
-
-#if UHT_ENABLE_PTR_PROPERTY_TAG
-				UnrealSourceFile->AddTypeDefIncludeIfNeeded(VarProperty.PropertyClass);
-#endif
-
-				if (VarProperty.PropertyClass->IsChildOf(UClass::StaticClass()))
-				{
-					FClassProperty* Result = new FClassProperty(Scope, Name, ObjectFlags);
-					GTypeDefinitionInfoMap.Add(Result, MakeShared<FUnrealPropertyDefinitionInfo>(*UnrealSourceFile, 0, Result));
-					Result->MetaClass     = VarProperty.MetaClass;
-					Result->PropertyClass = VarProperty.PropertyClass;
-					return Result;
-				}
-				else
-				{
-					if (DoesAnythingInHierarchyHaveDefaultToInstanced(VarProperty.PropertyClass))
-					{
-						VarProperty.PropertyFlags |= CPF_InstancedReference;
-						AddEditInlineMetaData(VarProperty.MetaData);
-					}
-
-					FObjectProperty* Result = new FObjectProperty(Scope, Name, ObjectFlags);
-					GTypeDefinitionInfoMap.Add(Result, MakeShared<FUnrealPropertyDefinitionInfo>(*UnrealSourceFile, 0, Result));
-					Result->PropertyClass = VarProperty.PropertyClass;
-					return Result;
-				}
-
-			case CPT_WeakObjectReference:
-			{
-				check(VarProperty.PropertyClass);
-
-#if UHT_ENABLE_PTR_PROPERTY_TAG
-				UnrealSourceFile->AddTypeDefIncludeIfNeeded(VarProperty.PropertyClass);
-#endif
-
-				FWeakObjectProperty* Result = new FWeakObjectProperty(Scope, Name, ObjectFlags);
-				GTypeDefinitionInfoMap.Add(Result, MakeShared<FUnrealPropertyDefinitionInfo>(*UnrealSourceFile, 0, Result));
-				Result->PropertyClass = VarProperty.PropertyClass;
-				return Result;
-			}
-
-			case CPT_LazyObjectReference:
-			{
-				check(VarProperty.PropertyClass);
-
-#if UHT_ENABLE_PTR_PROPERTY_TAG
-				UnrealSourceFile->AddTypeDefIncludeIfNeeded(VarProperty.PropertyClass);
-#endif
-
-				FLazyObjectProperty* Result = new FLazyObjectProperty(Scope, Name, ObjectFlags);
-				GTypeDefinitionInfoMap.Add(Result, MakeShared<FUnrealPropertyDefinitionInfo>(*UnrealSourceFile, 0, Result));
-				Result->PropertyClass = VarProperty.PropertyClass;
-				return Result;
-			}
-
-			case CPT_ObjectPtrReference:
-			{
-				check(VarProperty.PropertyClass);
-
-#if UHT_ENABLE_PTR_PROPERTY_TAG
-				UnrealSourceFile->AddTypeDefIncludeIfNeeded(VarProperty.PropertyClass);
-#endif
-
-				if (VarProperty.PropertyClass->IsChildOf(UClass::StaticClass()))
-				{
-					FClassPtrProperty* Result = new FClassPtrProperty(Scope, Name, ObjectFlags);
-					GTypeDefinitionInfoMap.Add(Result, MakeShared<FUnrealPropertyDefinitionInfo>(*UnrealSourceFile, 0, Result));
-					Result->MetaClass     = VarProperty.MetaClass;
-					Result->PropertyClass = VarProperty.PropertyClass;
-					return Result;
-				}
-				else
-				{
-					if (DoesAnythingInHierarchyHaveDefaultToInstanced(VarProperty.PropertyClass))
-					{
-						VarProperty.PropertyFlags |= CPF_InstancedReference;
-						AddEditInlineMetaData(VarProperty.MetaData);
-					}
-
-					FObjectPtrProperty* Result = new FObjectPtrProperty(Scope, Name, ObjectFlags);
-					GTypeDefinitionInfoMap.Add(Result, MakeShared<FUnrealPropertyDefinitionInfo>(*UnrealSourceFile, 0, Result));
-					Result->PropertyClass = VarProperty.PropertyClass;
-					return Result;
-				}
-			}
-
-			case CPT_SoftObjectReference:
-				check(VarProperty.PropertyClass);
-
-#if UHT_ENABLE_PTR_PROPERTY_TAG
-				UnrealSourceFile->AddTypeDefIncludeIfNeeded(VarProperty.PropertyClass);
-#endif
-
-				if (VarProperty.PropertyClass->IsChildOf(UClass::StaticClass()))
-				{
-					FSoftClassProperty* Result = new FSoftClassProperty(Scope, Name, ObjectFlags);
-					GTypeDefinitionInfoMap.Add(Result, MakeShared<FUnrealPropertyDefinitionInfo>(*UnrealSourceFile, 0, Result));
-					Result->MetaClass     = VarProperty.MetaClass;
-					Result->PropertyClass = VarProperty.PropertyClass;
-					return Result;
-				}
-				else
-				{
-					FSoftObjectProperty* Result = new FSoftObjectProperty(Scope, Name, ObjectFlags);
-					GTypeDefinitionInfoMap.Add(Result, MakeShared<FUnrealPropertyDefinitionInfo>(*UnrealSourceFile, 0, Result));
-					Result->PropertyClass = VarProperty.PropertyClass;
-					return Result;
-				}
-
-			case CPT_Interface:
-			{
-				check(VarProperty.PropertyClass);
-				check(VarProperty.PropertyClass->HasAnyClassFlags(CLASS_Interface));
-
-#if UHT_ENABLE_PTR_PROPERTY_TAG
-				UnrealSourceFile->AddTypeDefIncludeIfNeeded(VarProperty.PropertyClass);
-#endif
-
-				FInterfaceProperty* Result = new  FInterfaceProperty(Scope, Name, ObjectFlags);
-				GTypeDefinitionInfoMap.Add(Result, MakeShared<FUnrealPropertyDefinitionInfo>(*UnrealSourceFile, 0, Result));
-				Result->InterfaceClass = VarProperty.PropertyClass;
-				return Result;
-			}
-
-			case CPT_Name:
-			{
-				FNameProperty* Result = new FNameProperty(Scope, Name, ObjectFlags);
-				GTypeDefinitionInfoMap.Add(Result, MakeShared<FUnrealPropertyDefinitionInfo>(*UnrealSourceFile, 0, Result));
-				return Result;
-			}
-
-			case CPT_String:
-			{
-				FStrProperty* Result = new FStrProperty(Scope, Name, ObjectFlags);
-				GTypeDefinitionInfoMap.Add(Result, MakeShared<FUnrealPropertyDefinitionInfo>(*UnrealSourceFile, 0, Result));
-				return Result;
-			}
-
-			case CPT_Text:
-			{
-				FTextProperty* Result = new FTextProperty(Scope, Name, ObjectFlags);
-				GTypeDefinitionInfoMap.Add(Result, MakeShared<FUnrealPropertyDefinitionInfo>(*UnrealSourceFile, 0, Result));
-				return Result;
-			}
-
-			case CPT_Struct:
-			{
-#if UHT_ENABLE_VALUE_PROPERTY_TAG
-				UnrealSourceFile->AddTypeDefIncludeIfNeeded(VarProperty.Struct);
-#endif
-
-				if (VarProperty.Struct->StructFlags & STRUCT_HasInstancedReference)
-				{
-					VarProperty.PropertyFlags |= CPF_ContainsInstancedReference;
-				}
-
-				FStructProperty* Result = new FStructProperty(Scope, Name, ObjectFlags);
-				GTypeDefinitionInfoMap.Add(Result, MakeShared<FUnrealPropertyDefinitionInfo>(*UnrealSourceFile, 0, Result));
-				Result->Struct = VarProperty.Struct;
-				return Result;
-			}
-
-			case CPT_Delegate:
-			{
-				FDelegateProperty* Result = new FDelegateProperty(Scope, Name, ObjectFlags);
-				GTypeDefinitionInfoMap.Add(Result, MakeShared<FUnrealPropertyDefinitionInfo>(*UnrealSourceFile, 0, Result));
-				return Result;
-			}
-
-			case CPT_MulticastDelegate:
-			{
-				FMulticastDelegateProperty* Result;
-				if (VarProperty.Function->IsA<USparseDelegateFunction>())
-				{
-					Result = new FMulticastSparseDelegateProperty(Scope, Name, ObjectFlags);
-				}
-				else
-				{
-					Result = new FMulticastInlineDelegateProperty(Scope, Name, ObjectFlags);
-				}
-				GTypeDefinitionInfoMap.Add(Result, MakeShared<FUnrealPropertyDefinitionInfo>(*UnrealSourceFile, 0, Result));
-				return Result;
-			}
-
-			case CPT_FieldPath:
-			{
-				FFieldPathProperty* Result = new FFieldPathProperty(Scope, Name, ObjectFlags);
-				GTypeDefinitionInfoMap.Add(Result, MakeShared<FUnrealPropertyDefinitionInfo>(*UnrealSourceFile, 0, Result));
-				Result->PropertyClass = VarProperty.PropertyPathClass;
-
-				return Result;
-			}
-
-			default:
-				FError::Throwf(TEXT("Unknown property type %i (%s)"), (uint8)VarProperty.Type, *Name.ToString());
-		}
-
-		// Unreachable
-		check(false); //-V779
-		return nullptr;
-	}
-
 	/**
 	 * Ensures at script compile time that the metadata formatting is correct
 	 * @param	InKey			the metadata key being added
 	 * @param	InValue			the value string that will be associated with the InKey
 	 */
-	void ValidateMetaDataFormat(FFieldVariant Field, const FName InKey, const FString& InValue)
+	void ValidateMetaDataFormat(FUnrealTypeDefinitionInfo& TypeDef, const FName InKey, const FString& InValue)
 	{
 		switch (GetCheckedMetadataSpecifier(InKey))
 		{
@@ -1074,9 +668,9 @@ namespace
 
 			case ECheckedMetadataSpecifier::BlueprintProtected:
 			{
-				if (Field.IsUObject())
+				if (FUnrealFunctionDefinitionInfo* FuncDef = TypeDef.AsFunction())
 				{
-					UFunction* Function = Field.Get<UFunction>();
+					UFunction* Function = FuncDef->GetFunction();
 					if (Function->HasAnyFunctionFlags(FUNC_Static))
 					{
 						// Determine if it's a function library
@@ -1097,44 +691,34 @@ namespace
 
 			case ECheckedMetadataSpecifier::CommutativeAssociativeBinaryOperator:
 			{
-				if (UFunction* Function = Field.Get<UFunction>())
+				if (FUnrealFunctionDefinitionInfo* FuncDef = TypeDef.AsFunction())
 				{
+					UFunction* Function = FuncDef->GetFunction();
 					bool bGoodParams = (Function->NumParms == 3);
 					if (bGoodParams)
 					{
 						FProperty* FirstParam = nullptr;
 						FProperty* SecondParam = nullptr;
 						FProperty* ReturnValue = nullptr;
-
-						TFieldIterator<FProperty> It(Function);
-
-						auto GetNextParam = [&]()
+						for (FUnrealPropertyDefinitionInfo* PropertyDef : FuncDef->GetProperties())
 						{
-							if (It)
+							FProperty* Property = PropertyDef->GetProperty();
+							if (Property->HasAnyPropertyFlags(CPF_ReturnParm))
+							{
+								ReturnValue = Property;
+							}
+							else
 							{
-								if (It->HasAnyPropertyFlags(CPF_ReturnParm))
+								if (FirstParam == nullptr)
 								{
-									ReturnValue = *It;
+									FirstParam = Property;
 								}
-								else
+								else if (SecondParam == nullptr)
 								{
-									if (FirstParam == nullptr)
-									{
-										FirstParam = *It;
-									}
-									else if (SecondParam == nullptr)
-									{
-										SecondParam = *It;
-									}
+									SecondParam = Property;
 								}
-								++It;
 							}
-						};
-
-						GetNextParam();
-						GetNextParam();
-						GetNextParam();
-						ensure(!It);
+						}
 
 						if (ReturnValue == nullptr || SecondParam == nullptr || !SecondParam->SameType(FirstParam))
 						{
@@ -1144,7 +728,7 @@ namespace
 
 					if (!bGoodParams)
 					{
-						UE_LOG_ERROR_UHT(TEXT("Commutative asssociative binary operators must have exactly 2 parameters of the same type and a return value."));
+						UE_LOG_ERROR_UHT(TEXT("Commutative associative binary operators must have exactly 2 parameters of the same type and a return value."));
 					}
 				}
 			}
@@ -1153,8 +737,9 @@ namespace
 			case ECheckedMetadataSpecifier::ExpandBoolAsExecs:
 			case ECheckedMetadataSpecifier::ExpandEnumAsExecs:
 			{
-				if (UFunction* Function = Field.Get<UFunction>())
+				if (FUnrealFunctionDefinitionInfo* FuncDef = TypeDef.AsFunction())
 				{
+					UFunction* Function = FuncDef->GetFunction();
 					// multiple entry parsing in the same format as eg SetParam.
 					TArray<FString> RawGroupings;
 					InValue.ParseIntoArray(RawGroupings, TEXT(","), false);
@@ -1172,7 +757,7 @@ namespace
 								continue;
 							}
 							
-							FField* FoundField = FHeaderParser::FindProperty(Function, *Entry, false);
+							FField* FoundField = FHeaderParser::FindProperty(*FuncDef, *Entry, false);
 							if (!FoundField)
 							{
 								UE_LOG_ERROR_UHT(TEXT("Function does not have a parameter named '%s'"), *Entry);
@@ -1214,15 +799,18 @@ namespace
 			case ECheckedMetadataSpecifier::Units:
 			{
 				// Check for numeric property
-				FField* MaybeProperty = Field.ToField();
-				if (!MaybeProperty->IsA<FNumericProperty>() && !MaybeProperty->IsA<FStructProperty>())
+				if (FUnrealPropertyDefinitionInfo* PropDef = TypeDef.AsProperty())
 				{
-					FError::Throwf(TEXT("'Units' meta data can only be applied to numeric and struct properties"));
+					FProperty* Property = PropDef->GetProperty();
+					if (!Property->IsA<FNumericProperty>() && !Property->IsA<FStructProperty>())
+					{
+						FError::Throwf(TEXT("'Units' meta data can only be applied to numeric and struct properties"));
+					}
 				}
 
 				if (!FUnitConversion::UnitFromString(*InValue))
 				{
-					FError::Throwf(TEXT("Unrecognized units (%s) specified for property '%s'"), *InValue, *Field.GetFullName());
+					FError::Throwf(TEXT("Unrecognized units (%s) specified for property '%s'"), *InValue, *TypeDef.AsFieldChecked().GetField()->GetFullName());
 				}
 			}
 			break;
@@ -1240,90 +828,21 @@ namespace
 	}
 
 	// Ensures at script compile time that the metadata formatting is correct
-	void ValidateMetaDataFormat(FFieldVariant Field, const TMap<FName, FString>& MetaData)
+	void ValidateMetaDataFormat(FUnrealTypeDefinitionInfo& TypeDef, const TMap<FName, FString>& MetaData)
 	{
 		for (const TPair<FName, FString>& Pair : MetaData)
 		{
-			ValidateMetaDataFormat(Field, Pair.Key, Pair.Value);
-		}
-	}
-
-	// Validates the metadata, then adds it to the class data
-	void AddMetaDataToClassData(FFieldVariant Field, TMap<FName, FString>&& InMetaData)
-	{
-		// Evaluate any key redirects on the passed in pairs
-		for (TPair<FName, FString>& Pair : InMetaData)
-		{
-			FName& CurrentKey = Pair.Key;
-			FName NewKey = UMetaData::GetRemappedKeyName(CurrentKey);
-
-			if (NewKey != NAME_None)
-			{
-				UE_LOG_WARNING_UHT(TEXT("Remapping old metadata key '%s' to new key '%s', please update the declaration."), *CurrentKey.ToString(), *NewKey.ToString());
-				CurrentKey = NewKey;
-			}
-		}
-
-		// Finish validating and associate the metadata with the field
-		ValidateMetaDataFormat(Field, InMetaData);
-		if (Field.IsUObject())
-		{
-			FStructMetaData::AddMetaData(CastChecked<UField>(Field.ToUObject()), MoveTemp(InMetaData));
-		}
-		else
-		{
-			FStructMetaData::AddMetaData(Field.ToField(), MoveTemp(InMetaData));
+			ValidateMetaDataFormat(TypeDef, Pair.Key, Pair.Value);
 		}
 	}
 
 	bool IsPropertySupportedByBlueprint(const FProperty* Property, bool bMemberVariable)
 	{
-		if (Property == nullptr)
+		if (FUnrealPropertyDefinitionInfo* PropDef = GTypeDefinitionInfoMap.Find<FUnrealPropertyDefinitionInfo>(Property))
 		{
-			return false;
+			return FPropertyTraits::IsSupportedByBlueprint(*PropDef, bMemberVariable);
 		}
-		if (const FArrayProperty* ArrayProperty = CastField<const FArrayProperty>(Property))
-		{
-			// Script VM doesn't support array of weak ptrs.
-			return IsPropertySupportedByBlueprint(ArrayProperty->Inner, false);
-		}
-		else if (const FSetProperty* SetProperty = CastField<const FSetProperty>(Property))
-		{
-			return IsPropertySupportedByBlueprint(SetProperty->ElementProp, false);
-		}
-		else if (const FMapProperty* MapProperty = CastField<const FMapProperty>(Property))
-		{
-			return IsPropertySupportedByBlueprint(MapProperty->KeyProp, false) &&
-				IsPropertySupportedByBlueprint(MapProperty->ValueProp, false);
-		}
-		else if (const FStructProperty* StructProperty = CastField<const FStructProperty>(Property))
-		{
-			return (StructProperty->Struct->GetBoolMetaDataHierarchical(NAME_BlueprintType));
-		}
-
-		const bool bSupportedType = Property->IsA<FInterfaceProperty>()
-			|| Property->IsA<FClassProperty>()
-			|| Property->IsA<FSoftObjectProperty>()
-			|| Property->IsA<FObjectProperty>()
-			|| Property->IsA<FObjectPtrProperty>()
-			|| Property->IsA<FClassPtrProperty>()
-			|| Property->IsA<FFloatProperty>()
-			|| Property->IsA<FDoubleProperty>()
-			|| Property->IsA<FLargeWorldCoordinatesRealProperty>()
-			|| Property->IsA<FIntProperty>()
-			|| Property->IsA<FInt64Property>()
-			|| Property->IsA<FByteProperty>()
-			|| Property->IsA<FNameProperty>()
-			|| Property->IsA<FBoolProperty>()
-			|| Property->IsA<FStrProperty>()
-			|| Property->IsA<FTextProperty>()
-			|| Property->IsA<FDelegateProperty>()
-			|| Property->IsA<FEnumProperty>()
-			|| Property->IsA<FFieldPathProperty>();
-
-		const bool bIsSupportedMemberVariable = Property->IsA<FWeakObjectProperty>() || Property->IsA<FMulticastDelegateProperty>();
-
-		return bSupportedType || (bIsSupportedMemberVariable && bMemberVariable);
+		return false;
 	}
 
 	void SkipAlignasIfNecessary(FBaseParser& Parser)
@@ -1355,13 +874,13 @@ namespace
 		Parser.RequireSymbol(TEXT('('), ErrorMessageGetter);
 
 		FToken Token;
-		if (Parser.GetToken(Token) && (Token.Type != CPT_Float || Token.TokenType != TOKEN_Const))
+		if (Parser.GetToken(Token) && (Token.ConstantType != CPT_Float || Token.TokenType != TOKEN_Const))
 		{
 			FError::Throwf(TEXT("Expected engine version in %s macro"), MacroToken.Identifier);
 		}
 
 		Parser.RequireSymbol(TEXT(','), ErrorMessageGetter);
-		if (Parser.GetToken(Token) && (Token.Type != CPT_String || Token.TokenType != TOKEN_Const))
+		if (Parser.GetToken(Token) && (Token.ConstantType != CPT_String || Token.TokenType != TOKEN_Const))
 		{
 			FError::Throwf(TEXT("Expected deprecation message in %s macro"), MacroToken.Identifier);
 		}
@@ -1410,6 +929,32 @@ namespace
 	};
 }
 
+void AddEditInlineMetaData(TMap<FName, FString>& MetaData)
+{
+	MetaData.Add(NAME_EditInline, TEXT("true"));
+}
+
+// Validates the metadata, then adds it to the class data
+void AddMetaDataToClassData(FUnrealTypeDefinitionInfo& TypeDef, TMap<FName, FString>&& InMetaData)
+{
+	// Evaluate any key redirects on the passed in pairs
+	for (TPair<FName, FString>& Pair : InMetaData)
+	{
+		FName& CurrentKey = Pair.Key;
+		FName NewKey = UMetaData::GetRemappedKeyName(CurrentKey);
+
+		if (NewKey != NAME_None)
+		{
+			UE_LOG_WARNING_UHT(TEXT("Remapping old metadata key '%s' to new key '%s', please update the declaration."), *CurrentKey.ToString(), *NewKey.ToString());
+			CurrentKey = NewKey;
+		}
+	}
+
+	// Finish validating and associate the metadata with the field
+	ValidateMetaDataFormat(TypeDef, InMetaData);
+	TypeDef.AddMetaData(MoveTemp(InMetaData));
+}
+
 FUHTConfig::FUHTConfig()
 {
 	// Read Ini options, GConfig must exist by this point
@@ -1532,7 +1077,7 @@ FClass* FHeaderParser::GetQualifiedClass(const TCHAR* Thing)
  */
 UField* FHeaderParser::FindField
 (
-	UStruct*		Scope,
+	FUnrealStructDefinitionInfo&  InScope,
 	const TCHAR*	InIdentifier,
 	bool			bIncludeParents,
 	UClass*			FieldClass,
@@ -1543,32 +1088,35 @@ UField* FHeaderParser::FindField
 	FName InName(InIdentifier, FNAME_Find);
 	if (InName != NAME_None)
 	{
-		for( ; Scope; Scope = Cast<UStruct>(Scope->GetOuter()) )
+		for (FUnrealStructDefinitionInfo* Scope = &InScope; Scope; Scope = UHTCast<FUnrealStructDefinitionInfo>(Scope->GetOuter()))
 		{
-			for (TFieldIterator<FField> It(Scope); It; ++It)
+			for (FUnrealPropertyDefinitionInfo* PropertyDef : TUHTFieldRange<FUnrealPropertyDefinitionInfo>(*Scope))
 			{
-				if (It->GetFName() == InName)
+				FProperty* Property = PropertyDef->GetProperty();
+				if (Property->GetFName() == InName)
 				{
 					if (Thing)
 					{
-						FError::Throwf(TEXT("%s: expecting @todo: FProp, got %s"), Thing, /**FieldClass->GetName(),*/ *It->GetClass()->GetName());
+						FError::Throwf(TEXT("%s: expecting @todo: FProp, got %s"), Thing, /**FieldClass->GetName(),*/ *Property->GetClass()->GetName());
 					}
 					return nullptr;
 				}
+
 			}
-			for( TFieldIterator<UField> It(Scope); It; ++It )
+			for (FUnrealFunctionDefinitionInfo* FunctionDef : TUHTFieldRange<FUnrealFunctionDefinitionInfo>(*Scope))
 			{
-				if (It->GetFName() == InName)
+				UFunction* Function = FunctionDef->GetFunction();
+				if (Function->GetFName() == InName)
 				{
-					if (!It->IsA(FieldClass))
+					if (!Function->IsA(FieldClass))
 					{
 						if (Thing)
 						{
-							FError::Throwf(TEXT("%s: expecting %s, got %s"), Thing, *FieldClass->GetName(), *It->GetClass()->GetName() );
+							FError::Throwf(TEXT("%s: expecting %s, got %s"), Thing, *FieldClass->GetName(), *Function->GetClass()->GetName());
 						}
 						return nullptr;
 					}
-					return *It;
+					return Function;
 				}
 			}
 
@@ -1579,41 +1127,43 @@ UField* FHeaderParser::FindField
 		}
 	}
 
-						return NULL;
+	return nullptr;
 }
 
-FField* FHeaderParser::FindProperty(UStruct* Scope, const TCHAR* InIdentifier, bool bIncludeParents, FFieldClass* FieldClass, const TCHAR* Thing)
+FField* FHeaderParser::FindProperty(FUnrealStructDefinitionInfo& InScope, const TCHAR* InIdentifier, bool bIncludeParents, FFieldClass* FieldClass, const TCHAR* Thing)
 {
 	check(InIdentifier);
 	FName InName(InIdentifier, FNAME_Find);
 	if (InName != NAME_None)
 	{
-		for (; Scope; Scope = Cast<UStruct>(Scope->GetOuter()))
+		for (FUnrealStructDefinitionInfo* Scope = &InScope; Scope; Scope = UHTCast<FUnrealStructDefinitionInfo>(Scope->GetOuter()))
 		{
-			for (TFieldIterator<UField> It(Scope); It; ++It)
+			for (FUnrealFunctionDefinitionInfo* FunctionDef : TUHTFieldRange<FUnrealFunctionDefinitionInfo>(*Scope))
 			{
-				if (It->GetFName() == InName)
+				UFunction* Function = FunctionDef->GetFunction();
+				if (Function->GetFName() == InName)
 				{
 					if (Thing)
 					{
-						FError::Throwf(TEXT("%s: expecting a property, got %s"), Thing, *It->GetClass()->GetName());
+						FError::Throwf(TEXT("%s: expecting a property, got %s"), Thing, *Function->GetClass()->GetName());
 					}
 					return nullptr;
 				}
 			}
-			for (TFieldIterator<FField> It(Scope); It; ++It)
+			for (FUnrealPropertyDefinitionInfo* PropertyDef : TUHTFieldRange<FUnrealPropertyDefinitionInfo>(*Scope))
 			{
-				if (It->GetFName() == InName)
+				FProperty* Property = PropertyDef->GetProperty();
+				if (Property->GetFName() == InName)
 				{
-					if (!It->IsA(FieldClass))
+					if (!Property->IsA(FieldClass))
 					{
 						if (Thing)
 						{
-							FError::Throwf(TEXT("%s: expecting %s: FProp, got %s"), Thing, *FieldClass->GetName(), *It->GetClass()->GetName());
+							FError::Throwf(TEXT("%s: expecting %s: FProp, got %s"), Thing, *FieldClass->GetName(), *Property->GetClass()->GetName());
 						}
 						return nullptr;
 					}
-					return *It;
+					return Property;
 				}
 			}
 
@@ -1623,54 +1173,7 @@ FField* FHeaderParser::FindProperty(UStruct* Scope, const TCHAR* InIdentifier, b
 			}
 		}
 	}
-
-	return NULL;
-}
-/**
- * @return	true if Scope has FProperty objects in its list of fields
- */
-bool FHeaderParser::HasMemberProperties( const UStruct* Scope )
-{
-	// it's safe to pass a NULL Scope to TFieldIterator, but this function shouldn't be called with a NULL Scope
-	checkSlow(Scope);
-	TFieldIterator<FProperty> It(Scope,EFieldIteratorFlags::ExcludeSuper);
-	return It ? true : false;
-}
-
-/**
- * Get the parent struct specified.
- *
- * @param	CurrentScope	scope to start in
- * @param	SearchName		parent scope to search for
- *
- * @return	a pointer to the parent struct with the specified name, or NULL if the parent couldn't be found
- */
-UStruct* FHeaderParser::GetSuperScope( UStruct* CurrentScope, const FName& SearchName )
-{
-	UStruct* SuperScope = CurrentScope;
-	while (SuperScope && !SuperScope->GetInheritanceSuper())
-	{
-		SuperScope = CastChecked<UStruct>(SuperScope->GetOuter());
-	}
-	if (SuperScope != NULL)
-	{
-		// iterate up the inheritance chain looking for one that has the desired name
-		do
-		{
-			UStruct* NextScope = SuperScope->GetInheritanceSuper();
-			if (NextScope)
-			{
-				SuperScope = NextScope;
-			}
-			else
-			{
-				// otherwise we've failed
-				SuperScope = NULL;
-			}
-		} while (SuperScope != NULL && SuperScope->GetFName() != SearchName);
-	}
-
-	return SuperScope;
+	return nullptr;
 }
 
 /**
@@ -1730,8 +1233,9 @@ UEnum* FHeaderParser::CompileEnum()
 
 	// Get the enum specifier list
 	FToken                     EnumToken;
+	FMetaData				   EnumMetaData;
 	TArray<FPropertySpecifier> SpecifiersFound;
-	ReadSpecifierSetInsideMacro(SpecifiersFound, TEXT("Enum"), EnumToken.MetaData);
+	ReadSpecifierSetInsideMacro(SpecifiersFound, TEXT("Enum"), EnumMetaData);
 
 	// Check enum type. This can be global 'enum', 'namespace' or 'enum class' enums.
 	bool            bReadEnumName = false;
@@ -1784,10 +1288,10 @@ UEnum* FHeaderParser::CompileEnum()
 		FError::Throwf(TEXT("Missing enumeration name") );
 	}
 
-	ParseFieldMetaData(EnumToken.MetaData, EnumToken.Identifier);
+	ParseFieldMetaData(EnumMetaData, EnumToken.Identifier);
 
 	// Create enum definition.
-	FUnrealEnumDefinitionInfo& EnumDef = GTypeDefinitionInfoMap.FindByNameChecked(EnumToken.Identifier).AsEnumChecked();
+	FUnrealEnumDefinitionInfo& EnumDef = GTypeDefinitionInfoMap.FindByNameChecked<FUnrealEnumDefinitionInfo>(EnumToken.Identifier);
 	UEnum* Enum = EnumDef.GetEnum();
 
 	Scope->AddType(EnumDef);
@@ -1811,7 +1315,7 @@ UEnum* FHeaderParser::CompileEnum()
 	}
 
 	// Validate the metadata for the enum
-	ValidateMetaDataFormat(Enum, EnumToken.MetaData);
+	ValidateMetaDataFormat(EnumDef, EnumMetaData);
 
 	// Read base for enum class
 	EUnderlyingEnumType UnderlyingType = EUnderlyingEnumType::uint8;
@@ -1877,7 +1381,7 @@ UEnum* FHeaderParser::CompileEnum()
 		}
 	}
 
-	if (UnderlyingType != EUnderlyingEnumType::uint8 && EnumToken.MetaData.Contains(NAME_BlueprintType))
+	if (UnderlyingType != EUnderlyingEnumType::uint8 && EnumMetaData.Contains(FHeaderParserNames::NAME_BlueprintType))
 	{
 		FError::Throwf(TEXT("Invalid BlueprintType enum base - currently only uint8 supported"));
 	}
@@ -1915,20 +1419,21 @@ UEnum* FHeaderParser::CompileEnum()
 	}
 
 	// List of all metadata generated for this enum
-	TMap<FName,FString> EnumValueMetaData = EnumToken.MetaData;
+	TMap<FName,FString> EnumValueMetaData = EnumMetaData;
 
 	AddModuleRelativePathToMetadata(EnumDef, EnumValueMetaData);
 	AddFormattedPrevCommentAsTooltipMetaData(EnumValueMetaData);
 
 	// Parse all enums tags.
 	FToken TagToken;
+	FMetaData TagMetaData;
 	TArray<TMap<FName, FString>> EntryMetaData;
 
 	TArray<TPair<FName, int64>> EnumNames;
 	int64 CurrentEnumValue = 0;
 	while (GetIdentifier(TagToken))
 	{
-		AddFormattedPrevCommentAsTooltipMetaData(TagToken.MetaData);
+		AddFormattedPrevCommentAsTooltipMetaData(TagMetaData);
 
 		// Try to read an optional explicit enum value specification
 		if (MatchSymbol(TEXT('=')))
@@ -2005,23 +1510,23 @@ UEnum* FHeaderParser::CompileEnum()
 			++CurrentEnumValue;
 		}
 
-		TagToken.MetaData.Add(NAME_Name, NewTag.ToString());
-		EntryMetaData.Add(TagToken.MetaData);
+		TagMetaData.Add(NAME_Name, NewTag.ToString());
+		EntryMetaData.Add(TagMetaData);
 
 		// check for metadata on this enum value
-		ParseFieldMetaData(TagToken.MetaData, TagToken.Identifier);
-		if (TagToken.MetaData.Num() > 0)
+		ParseFieldMetaData(TagMetaData, TagToken.Identifier);
+		if (TagMetaData.Num() > 0)
 		{
 			// special case for enum value metadata - we need to prepend the key name with the enum value name
 			const FString TokenString = TagToken.Identifier;
-			for (const auto& MetaData : TagToken.MetaData)
+			for (const auto& MetaData : TagMetaData)
 			{
 				FString KeyString = TokenString + TEXT(".") + MetaData.Key.ToString();
 				EnumValueMetaData.Emplace(*KeyString, MetaData.Value);
 			}
 
 			// now clear the metadata because we're going to reuse this token for parsing the next enum value
-			TagToken.MetaData.Empty();
+			TagMetaData.Empty();
 		}
 
 		if (!MatchSymbol(TEXT(',')))
@@ -2074,7 +1579,7 @@ UEnum* FHeaderParser::CompileEnum()
 		PackageMetaData->SetObjectValues(Enum, MoveTemp(EnumValueMetaData));
 	}
 
-	if (!Enum->IsValidEnumValue(0) && EnumToken.MetaData.Contains(NAME_BlueprintType))
+	if (!Enum->IsValidEnumValue(0) && EnumMetaData.Contains(FHeaderParserNames::NAME_BlueprintType))
 	{
 		UE_LOG_WARNING_UHT(TEXT("'%s' does not have a 0 entry! (This is a problem when the enum is initalized by default)"), *Enum->GetName());
 	}
@@ -2543,7 +2048,7 @@ UScriptStruct* FHeaderParser::CompileStructDeclaration()
 	const FString EffectiveStructName = *StructNameStripped;
 
 	// Create.
-	FUnrealScriptStructDefinitionInfo& StructDef = GTypeDefinitionInfoMap.FindByNameChecked(*StructNameStripped).AsScriptStructChecked();
+	FUnrealScriptStructDefinitionInfo& StructDef = GTypeDefinitionInfoMap.FindByNameChecked<FUnrealScriptStructDefinitionInfo>(*StructNameStripped);
 	UScriptStruct* Struct = StructDef.GetScriptStruct();
 
 	// Process the list of specifiers
@@ -2620,17 +2125,17 @@ UScriptStruct* FHeaderParser::CompileStructDeclaration()
 			UScriptStruct* Type = nullptr;
 			if (bOverrideParentStructName)
 			{
-				Type = FindObject<UScriptStruct>(ANY_PACKAGE, *ParentStructNameStripped);
+				Type = FClasses::FindObject<UScriptStruct>(ANY_PACKAGE, *ParentStructNameStripped);
 			}
 
 			if (Type == nullptr)
 			{
-				Type = FindObject<UScriptStruct>(ANY_PACKAGE, *ParentStructNameInScript);
+				Type = FClasses::FindObject<UScriptStruct>(ANY_PACKAGE, *ParentStructNameInScript);
 			}
 
 			if (Type != nullptr)
 			{
-				BaseStructDef = &GTypeDefinitionInfoMap.FindChecked(Type).AsFieldChecked();
+				BaseStructDef = &GTypeDefinitionInfoMap.FindChecked<FUnrealFieldDefinitionInfo>(Type);
 			}
 		}
 
@@ -2692,7 +2197,7 @@ UScriptStruct* FHeaderParser::CompileStructDeclaration()
 	AddFormattedPrevCommentAsTooltipMetaData(MetaData);
 
 	// Register the metadata
-	AddMetaDataToClassData(Struct, MoveTemp(MetaData));
+	AddMetaDataToClassData(StructDef, MoveTemp(MetaData));
 
 	// Get opening brace.
 	RequireSymbol( TEXT('{'), TEXT("'struct'") );
@@ -3078,175 +2583,169 @@ void FHeaderParser::PopNest(ENestType NestType, const TCHAR* Descr)
 	}
 }
 
-void FHeaderParser::FixupDelegateProperties( UStruct* Struct, FScope& Scope, TMap<FName, UFunction*>& DelegateCache )
+void FHeaderParser::FixupDelegateProperties(FUnrealStructDefinitionInfo& StructDef, FScope& Scope, TMap<FName, UFunction*>& DelegateCache )
 {
-	check(Struct);
+	UStruct* Struct = StructDef.GetStruct();
 
-	for ( FField* Field = Struct->ChildProperties; Field; Field = Field->Next )
+	for (FUnrealPropertyDefinitionInfo* PropertyDef : StructDef.GetProperties())
 	{
-		FProperty* Property = CastField<FProperty>(Field);
-		if ( Property != NULL )
+		FProperty* Property = PropertyDef->GetProperty();
+		FDelegateProperty* DelegateProperty = CastField<FDelegateProperty>(Property);
+		FMulticastDelegateProperty* MulticastDelegateProperty = CastField<FMulticastDelegateProperty>(Property);
+		if ( DelegateProperty == NULL && MulticastDelegateProperty == NULL )
 		{
-			FDelegateProperty* DelegateProperty = CastField<FDelegateProperty>(Property);
-			FMulticastDelegateProperty* MulticastDelegateProperty = CastField<FMulticastDelegateProperty>(Property);
-			if ( DelegateProperty == NULL && MulticastDelegateProperty == NULL )
+			// if this is an array property, see if the array's type is a delegate
+			FArrayProperty* ArrayProp = CastField<FArrayProperty>(Property);
+			if ( ArrayProp != NULL )
 			{
-				// if this is an array property, see if the array's type is a delegate
-				FArrayProperty* ArrayProp = CastField<FArrayProperty>(Property);
-				if ( ArrayProp != NULL )
-				{
-					DelegateProperty = CastField<FDelegateProperty>(ArrayProp->Inner);
-					MulticastDelegateProperty = CastField<FMulticastDelegateProperty>(ArrayProp->Inner);
-				}
+				DelegateProperty = CastField<FDelegateProperty>(ArrayProp->Inner);
+				MulticastDelegateProperty = CastField<FMulticastDelegateProperty>(ArrayProp->Inner);
 			}
-			if (DelegateProperty != nullptr || MulticastDelegateProperty != nullptr)
-			{
-				// this FDelegateProperty corresponds to an actual delegate variable (i.e. delegate<SomeDelegate> Foo); we need to lookup the token data for
-				// this property and verify that the delegate property's "type" is an actual delegate function
-				FUnrealPropertyDefinitionInfo& PropDef = GTypeDefinitionInfoMap.FindChecked(Property).AsPropertyChecked();
-				FToken& DelegatePropertyToken = PropDef.GetToken();
+		}
+		if (DelegateProperty != nullptr || MulticastDelegateProperty != nullptr)
+		{
+			// this FDelegateProperty corresponds to an actual delegate variable (i.e. delegate<SomeDelegate> Foo); we need to lookup the token data for
+			// this property and verify that the delegate property's "type" is an actual delegate function
+			FUnrealPropertyDefinitionInfo& PropDef = GTypeDefinitionInfoMap.FindChecked<FUnrealPropertyDefinitionInfo>(Property);
+			FPropertyBase& DelegatePropertyToken = PropDef.GetPropertyBase();
 
-				// attempt to find the delegate function in the map of functions we've already found
-				UFunction* SourceDelegateFunction = DelegateCache.FindRef(DelegatePropertyToken.DelegateName);
-				if (SourceDelegateFunction == nullptr)
+			// attempt to find the delegate function in the map of functions we've already found
+			UFunction* SourceDelegateFunction = DelegateCache.FindRef(DelegatePropertyToken.DelegateName);
+			if (SourceDelegateFunction == nullptr)
+			{
+				FString NameOfDelegateFunction = DelegatePropertyToken.DelegateName.ToString() + FString( HEADER_GENERATED_DELEGATE_SIGNATURE_SUFFIX );
+				if ( !NameOfDelegateFunction.Contains(TEXT(".")) )
 				{
-					FString NameOfDelegateFunction = DelegatePropertyToken.DelegateName.ToString() + FString( HEADER_GENERATED_DELEGATE_SIGNATURE_SUFFIX );
-					if ( !NameOfDelegateFunction.Contains(TEXT(".")) )
+					// an unqualified delegate function name - search for a delegate function by this name within the current scope
+					if (FUnrealFieldDefinitionInfo* FoundDef = Scope.FindTypeByName(*NameOfDelegateFunction))
 					{
-						// an unqualified delegate function name - search for a delegate function by this name within the current scope
-						if (FUnrealFieldDefinitionInfo* FoundDef = Scope.FindTypeByName(*NameOfDelegateFunction))
+						if (FUnrealFunctionDefinitionInfo* SourceDelegateFunctionDef = FoundDef->AsFunction())
 						{
-							if (FUnrealFunctionDefinitionInfo* SourceDelegateFunctionDef = FoundDef->AsFunction())
-							{
-								SourceDelegateFunction = SourceDelegateFunctionDef->GetFunction();
-							}
-						}
-						if (SourceDelegateFunction == nullptr)
-						{
-							// Try to find in other packages.
-							UObject* DelegateSignatureOuter = DelegatePropertyToken.DelegateSignatureOwnerClass 
-								? ((UObject*)DelegatePropertyToken.DelegateSignatureOwnerClass) 
-								: ((UObject*)ANY_PACKAGE);
-							SourceDelegateFunction = Cast<UFunction>(StaticFindObject(UFunction::StaticClass(), DelegateSignatureOuter, *NameOfDelegateFunction));
-
-							if (SourceDelegateFunction == nullptr)
-							{
-								// convert this into a fully qualified path name for the error message.
-								NameOfDelegateFunction = Scope.GetName().ToString() + TEXT(".") + NameOfDelegateFunction;
-							}
+							SourceDelegateFunction = SourceDelegateFunctionDef->GetFunction();
 						}
 					}
-					else
+					if (SourceDelegateFunction == nullptr)
 					{
-						FString DelegateClassName, DelegateName;
-						NameOfDelegateFunction.Split(TEXT("."), &DelegateClassName, &DelegateName);
-
-						// verify that we got a valid string for the class name
-						if ( DelegateClassName.Len() == 0 )
-						{
-							UngetToken(DelegatePropertyToken);
-							FError::Throwf(TEXT("Invalid scope specified in delegate property function reference: '%s'"), *NameOfDelegateFunction);
-						}
+						// Try to find in other packages.
+						UObject* DelegateSignatureOuter = DelegatePropertyToken.DelegateSignatureOwnerClass 
+							? ((UObject*)DelegatePropertyToken.DelegateSignatureOwnerClass) 
+							: ((UObject*)ANY_PACKAGE);
+						SourceDelegateFunction = FClasses::FindObject<UFunction>(DelegateSignatureOuter, *NameOfDelegateFunction);
 
-						// verify that we got a valid string for the name of the function
-						if ( DelegateName.Len() == 0 )
+						if (SourceDelegateFunction == nullptr)
 						{
-							UngetToken(DelegatePropertyToken);
-							FError::Throwf(TEXT("Invalid delegate name specified in delegate property function reference '%s'"), *NameOfDelegateFunction);
+							// convert this into a fully qualified path name for the error message.
+							NameOfDelegateFunction = Scope.GetName().ToString() + TEXT(".") + NameOfDelegateFunction;
 						}
+					}
+				}
+				else
+				{
+					FString DelegateClassName, DelegateName;
+					NameOfDelegateFunction.Split(TEXT("."), &DelegateClassName, &DelegateName);
 
-						// make sure that the class that contains the delegate can be referenced here
-						UClass* DelegateOwnerClass = FClasses::FindScriptClassOrThrow(DelegateClassName);
-						if (FScope::GetTypeScope(DelegateOwnerClass)->FindTypeByName(*DelegateName) != nullptr)
-						{
-							FError::Throwf(TEXT("Inaccessible type: '%s'"), *DelegateOwnerClass->GetPathName());
-						}
-						SourceDelegateFunction = Cast<UFunction>(FindField(DelegateOwnerClass, *DelegateName, false, UFunction::StaticClass(), NULL));	
+					// verify that we got a valid string for the class name
+					if ( DelegateClassName.Len() == 0 )
+					{
+						UngetToken(PropDef.GetLineNumber(), PropDef.GetParsePosition());
+						FError::Throwf(TEXT("Invalid scope specified in delegate property function reference: '%s'"), *NameOfDelegateFunction);
 					}
 
-					if ( SourceDelegateFunction == NULL )
+					// verify that we got a valid string for the name of the function
+					if ( DelegateName.Len() == 0 )
 					{
-						UngetToken(DelegatePropertyToken);
-						FError::Throwf(TEXT("Failed to find delegate function '%s'"), *NameOfDelegateFunction);
+						UngetToken(PropDef.GetLineNumber(), PropDef.GetParsePosition());
+						FError::Throwf(TEXT("Invalid delegate name specified in delegate property function reference '%s'"), *NameOfDelegateFunction);
 					}
-					else if ( (SourceDelegateFunction->FunctionFlags&FUNC_Delegate) == 0 )
+
+					// make sure that the class that contains the delegate can be referenced here
+					UClass* DelegateOwnerClass = FClasses::FindScriptClassOrThrow(DelegateClassName);
+					if (FScope::GetTypeScope(DelegateOwnerClass)->FindTypeByName(*DelegateName) != nullptr)
 					{
-						UngetToken(DelegatePropertyToken);
-						FError::Throwf(TEXT("Only delegate functions can be used as the type for a delegate property; '%s' is not a delegate."), *NameOfDelegateFunction);
+						FError::Throwf(TEXT("Inaccessible type: '%s'"), *DelegateOwnerClass->GetPathName());
 					}
+					SourceDelegateFunction = Cast<UFunction>(FindField(GTypeDefinitionInfoMap.FindChecked<FUnrealStructDefinitionInfo>(DelegateOwnerClass), *DelegateName, false, UFunction::StaticClass(), NULL));	
 				}
 
-				// successfully found the delegate function that this delegate property corresponds to
+				if ( SourceDelegateFunction == NULL )
+				{
+					UngetToken(PropDef.GetLineNumber(), PropDef.GetParsePosition());
+					FError::Throwf(TEXT("Failed to find delegate function '%s'"), *NameOfDelegateFunction);
+				}
+				else if ( (SourceDelegateFunction->FunctionFlags&FUNC_Delegate) == 0 )
+				{
+					UngetToken(PropDef.GetLineNumber(), PropDef.GetParsePosition());
+					FError::Throwf(TEXT("Only delegate functions can be used as the type for a delegate property; '%s' is not a delegate."), *NameOfDelegateFunction);
+				}
+			}
+
+			// successfully found the delegate function that this delegate property corresponds to
 
-				// save this into the delegate cache for faster lookup later
-				DelegateCache.Add(DelegatePropertyToken.DelegateName, SourceDelegateFunction);
+			// save this into the delegate cache for faster lookup later
+			DelegateCache.Add(DelegatePropertyToken.DelegateName, SourceDelegateFunction);
 
-				// bind it to the delegate property
-				if( DelegateProperty != NULL )
+			// bind it to the delegate property
+			if( DelegateProperty != NULL )
+			{
+				if( !SourceDelegateFunction->HasAnyFunctionFlags( FUNC_MulticastDelegate ) )
 				{
-					if( !SourceDelegateFunction->HasAnyFunctionFlags( FUNC_MulticastDelegate ) )
-					{
-						DelegateProperty->SignatureFunction = DelegatePropertyToken.Function = SourceDelegateFunction;
-					}
-					else
-					{
-						FError::Throwf(TEXT("Unable to declare a single-cast delegate property for a multi-cast delegate type '%s'.  Either add a 'multicast' qualifier to the property or change the delegate type to be single-cast as well."), *SourceDelegateFunction->GetName());
-					}
+					DelegateProperty->SignatureFunction = DelegatePropertyToken.Function = SourceDelegateFunction;
 				}
-				else if( MulticastDelegateProperty != NULL )
+				else
 				{
-					if( SourceDelegateFunction->HasAnyFunctionFlags( FUNC_MulticastDelegate ) )
-					{
-						MulticastDelegateProperty->SignatureFunction = DelegatePropertyToken.Function = SourceDelegateFunction;
+					FError::Throwf(TEXT("Unable to declare a single-cast delegate property for a multi-cast delegate type '%s'.  Either add a 'multicast' qualifier to the property or change the delegate type to be single-cast as well."), *SourceDelegateFunction->GetName());
+				}
+			}
+			else if( MulticastDelegateProperty != NULL )
+			{
+				if( SourceDelegateFunction->HasAnyFunctionFlags( FUNC_MulticastDelegate ) )
+				{
+					MulticastDelegateProperty->SignatureFunction = DelegatePropertyToken.Function = SourceDelegateFunction;
 
-						if(MulticastDelegateProperty->HasAnyPropertyFlags(CPF_BlueprintAssignable | CPF_BlueprintCallable))
+					if(MulticastDelegateProperty->HasAnyPropertyFlags(CPF_BlueprintAssignable | CPF_BlueprintCallable))
+					{
+						FUnrealFunctionDefinitionInfo& SourceDelegateFunctionDef = GTypeDefinitionInfoMap.FindChecked<FUnrealFunctionDefinitionInfo>(SourceDelegateFunction);
+						for (FUnrealPropertyDefinitionInfo* FuncParamDef : SourceDelegateFunctionDef.GetProperties())
 						{
-							for (TFieldIterator<FProperty> PropIt(SourceDelegateFunction); PropIt && (PropIt->PropertyFlags & CPF_Parm); ++PropIt)
-							{
-								FProperty* FuncParam = *PropIt;
+							FProperty* FuncParam = FuncParamDef->GetProperty();
 
-								if (!IsPropertySupportedByBlueprint(FuncParam, false))
-								{
-									FString ExtendedCPPType;
-									FString CPPType = FuncParam->GetCPPType(&ExtendedCPPType);
-									UE_LOG_ERROR_UHT(TEXT("Type '%s%s' is not supported by blueprint. %s.%s"), *CPPType, *ExtendedCPPType, *SourceDelegateFunction->GetName(), *FuncParam->GetName());
-								}
+							if (!IsPropertySupportedByBlueprint(FuncParam, false))
+							{
+								FString ExtendedCPPType;
+								FString CPPType = FuncParam->GetCPPType(&ExtendedCPPType);
+								UE_LOG_ERROR_UHT(TEXT("Type '%s%s' is not supported by blueprint. %s.%s"), *CPPType, *ExtendedCPPType, *SourceDelegateFunction->GetName(), *FuncParam->GetName());
+							}
 
-								if(FuncParam->HasAllPropertyFlags(CPF_OutParm) && !FuncParam->HasAllPropertyFlags(CPF_ConstParm)  )
+							if(FuncParam->HasAllPropertyFlags(CPF_OutParm) && !FuncParam->HasAllPropertyFlags(CPF_ConstParm)  )
+							{
+								const bool bClassGeneratedFromBP = FClass::IsDynamic(Struct);
+								const bool bAllowedArrayRefFromBP = bClassGeneratedFromBP && FuncParam->IsA<FArrayProperty>();
+								if (!bAllowedArrayRefFromBP)
 								{
-									const bool bClassGeneratedFromBP = FClass::IsDynamic(Struct);
-									const bool bAllowedArrayRefFromBP = bClassGeneratedFromBP && FuncParam->IsA<FArrayProperty>();
-									if (!bAllowedArrayRefFromBP)
-									{
-										UE_LOG_ERROR_UHT(TEXT("BlueprintAssignable delegates do not support non-const references at the moment. Function: %s Parameter: '%s'"), *SourceDelegateFunction->GetName(), *FuncParam->GetName());
-									}
+									UE_LOG_ERROR_UHT(TEXT("BlueprintAssignable delegates do not support non-const references at the moment. Function: %s Parameter: '%s'"), *SourceDelegateFunction->GetName(), *FuncParam->GetName());
 								}
 							}
 						}
-
-					}
-					else
-					{
-						FError::Throwf(TEXT("Unable to declare a multi-cast delegate property for a single-cast delegate type '%s'.  Either remove the 'multicast' qualifier from the property or change the delegate type to be 'multicast' as well."), *SourceDelegateFunction->GetName());
 					}
+
+				}
+				else
+				{
+					FError::Throwf(TEXT("Unable to declare a multi-cast delegate property for a single-cast delegate type '%s'.  Either remove the 'multicast' qualifier from the property or change the delegate type to be 'multicast' as well."), *SourceDelegateFunction->GetName());
 				}
 			}
 		}
 	}
 
-	for (UField* Field = Struct->Children; Field; Field = Field->Next)
+	// Functions might have their own delegate properties which need to be validated
+	for (FUnrealFunctionDefinitionInfo* Function : StructDef.GetFunctions())
 	{
-		// if this is a state, function, or script struct, it might have its own delegate properties which need to be validated
-		UStruct* InternalStruct = Cast<UStruct>(Field);
-		if ( InternalStruct != NULL )
-		{
-			FixupDelegateProperties(InternalStruct, Scope, DelegateCache);
-		}
+		FixupDelegateProperties(*Function, Scope, DelegateCache);
 	}
 
-	CheckDocumentationPolicyForStruct(Struct);
+	CheckDocumentationPolicyForStruct(StructDef);
 
-	ParseRigVMMethodParameters(Struct);
+	ParseRigVMMethodParameters(StructDef);
 }
 
 void FHeaderParser::CheckSparseClassData(const UStruct* StructToCheck)
@@ -3285,7 +2784,7 @@ void FHeaderParser::CheckSparseClassData(const UStruct* StructToCheck)
 
 	for (const FString& SparseClassDataTypeName : SparseClassDataTypes)
 	{
-		UScriptStruct* SparseClassDataStruct = FindObjectSafe<UScriptStruct>(ANY_PACKAGE, *SparseClassDataTypeName);
+		UScriptStruct* SparseClassDataStruct = FClasses::FindObject<UScriptStruct>(ANY_PACKAGE, *SparseClassDataTypeName);
 
 		// make sure the sparse class data struct actually exists
 		if (!SparseClassDataStruct)
@@ -3295,8 +2794,10 @@ void FHeaderParser::CheckSparseClassData(const UStruct* StructToCheck)
 		}
 
 		// check the data struct for invalid properties
-		for (TFieldIterator<FProperty> Property(SparseClassDataStruct); Property; ++Property)
+		FUnrealScriptStructDefinitionInfo& SparseClassDataStructDef = GTypeDefinitionInfoMap.FindChecked<FUnrealScriptStructDefinitionInfo>(SparseClassDataStruct);
+		for (FUnrealPropertyDefinitionInfo* PropertyDef : TUHTFieldRange<FUnrealPropertyDefinitionInfo>(SparseClassDataStructDef))
 		{
+			FProperty* Property = PropertyDef->GetProperty();
 			if (Property->HasAnyPropertyFlags(CPF_BlueprintAssignable))
 			{
 				FError::Throwf(TEXT("Sparse class data types can not contain blueprint assignable delegates. Type '%s' Delegate '%s'"), *SparseClassDataStruct->GetName(), *Property->GetName());
@@ -3321,7 +2822,7 @@ void FHeaderParser::CheckSparseClassData(const UStruct* StructToCheck)
 		((FClass*)ParentClass)->GetSparseClassDataTypes(ParentSparseClassDataTypeNames);
 		for (FString& ParentSparseClassDataTypeName : ParentSparseClassDataTypeNames)
 		{
-			UScriptStruct* ParentSparseClassDataStruct = FindObjectSafe<UScriptStruct>(ANY_PACKAGE, *ParentSparseClassDataTypeName);
+			UScriptStruct* ParentSparseClassDataStruct = FClasses::FindObject<UScriptStruct>(ANY_PACKAGE, *ParentSparseClassDataTypeName);
 			if (ParentSparseClassDataStruct && !SparseClassDataStruct->IsChildOf(ParentSparseClassDataStruct))
 			{
 				FError::Throwf(TEXT("Class %s is a child of %s but its sparse class data struct, %s, does not inherit from %s."), *ClassToCheck->GetName(), *ParentClass->GetName(), *SparseClassDataStruct->GetName(), *ParentSparseClassDataStruct->GetName());
@@ -3339,17 +2840,19 @@ void FHeaderParser::ValidateClassFlags(const UClass* ToValidate)
 	}
 }
 
-void FHeaderParser::VerifyBlueprintPropertyGetter(FProperty* Prop, UFunction* TargetFunc)
+void FHeaderParser::VerifyBlueprintPropertyGetter(FProperty* Prop, FUnrealFunctionDefinitionInfo* TargetFuncDef)
 {
-	check(TargetFunc);
+	check(TargetFuncDef);
+	UFunction* TargetFunc = TargetFuncDef->GetFunction();
 
-	FProperty* ReturnProp = TargetFunc->GetReturnProperty();
-	if (TargetFunc->NumParms > 1 || (TargetFunc->NumParms == 1 && ReturnProp == nullptr))
+	const TArray<FUnrealPropertyDefinitionInfo*>& Properties = TargetFuncDef->GetProperties();
+	FUnrealPropertyDefinitionInfo* ReturnPropDef = TargetFuncDef->GetReturn();
+	if (Properties.Num() > 1 || (Properties.Num() == 1 && ReturnPropDef == nullptr))
 	{
 		UE_LOG_ERROR_UHT(TEXT("Blueprint Property getter function %s must not have parameters."), *TargetFunc->GetName());
 	}
 
-	if (ReturnProp == nullptr || !Prop->SameType(ReturnProp))
+	if (ReturnPropDef == nullptr || !Prop->SameType(ReturnPropDef->GetProperty()))
 	{
 		FString ExtendedCPPType;
 		FString CPPType = Prop->GetCPPType(&ExtendedCPPType);
@@ -3366,19 +2869,20 @@ void FHeaderParser::VerifyBlueprintPropertyGetter(FProperty* Prop, UFunction* Ta
 	}
 }
 
-void FHeaderParser::VerifyBlueprintPropertySetter(FProperty* Prop, UFunction* TargetFunc)
+void FHeaderParser::VerifyBlueprintPropertySetter(FProperty* Prop, FUnrealFunctionDefinitionInfo* TargetFuncDef)
 {
-	check(TargetFunc);
-	FProperty* ReturnProp = TargetFunc->GetReturnProperty();
+	check(TargetFuncDef);
+	UFunction* TargetFunc = TargetFuncDef->GetFunction();
+	FUnrealPropertyDefinitionInfo* ReturnPropDef = TargetFuncDef->GetReturn();
 
-	if (ReturnProp)
+	if (ReturnPropDef)
 	{
 		UE_LOG_ERROR_UHT(TEXT("Blueprint Property setter function %s must not have a return value."), *TargetFunc->GetName());
 	}
 	else
 	{
-		TFieldIterator<FProperty> Parm(TargetFunc);
-		if (TargetFunc->NumParms != 1 || !Prop->SameType(*Parm))
+		const TArray<FUnrealPropertyDefinitionInfo*>& Properties = TargetFuncDef->GetProperties();
+		if (TargetFunc->NumParms != 1 || !Prop->SameType(Properties[0]->GetProperty()))
 		{
 			FString ExtendedCPPType;
 			FString CPPType = Prop->GetCPPType(&ExtendedCPPType);
@@ -3400,10 +2904,11 @@ void FHeaderParser::VerifyBlueprintPropertySetter(FProperty* Prop, UFunction* Ta
 	}
 }
 
-void FHeaderParser::VerifyRepNotifyCallback(FProperty* Prop, UFunction* TargetFunc)
+void FHeaderParser::VerifyRepNotifyCallback(FProperty* Prop, FUnrealFunctionDefinitionInfo* TargetFuncDef)
 {
-	if( TargetFunc )
+	if (TargetFuncDef)
 	{
+		UFunction* TargetFunc = TargetFuncDef->GetFunction();
 		if (TargetFunc->GetReturnProperty())
 		{
 			UE_LOG_ERROR_UHT(TEXT("Replication notification function %s must not have return value."), *TargetFunc->GetName());
@@ -3412,29 +2917,30 @@ void FHeaderParser::VerifyRepNotifyCallback(FProperty* Prop, UFunction* TargetFu
 		const bool bIsArrayProperty = ( Prop->ArrayDim > 1 || CastField<FArrayProperty>(Prop) );
 		const int32 MaxParms = bIsArrayProperty ? 2 : 1;
 
-		if ( TargetFunc->NumParms > MaxParms)
+		const TArray<FUnrealPropertyDefinitionInfo*>& Properties = TargetFuncDef->GetProperties();
+		if (Properties.Num() > MaxParms)
 		{
 			UE_LOG_ERROR_UHT(TEXT("Replication notification function %s has too many parameters."), *TargetFunc->GetName());
 		}
 
-		TFieldIterator<FProperty> Parm(TargetFunc);
-		if ( TargetFunc->NumParms >= 1 && Parm)
+
+		if (Properties.Num() >= 1)
 		{
+			const FProperty* Parm = Properties[0]->GetProperty();
 			// First parameter is always the old value:
-			if ( !Prop->SameType(*Parm) )
+			if ( !Prop->SameType(Parm) )
 			{
 				FString ExtendedCPPType;
 				FString CPPType = Prop->GetCPPType(&ExtendedCPPType);
 				UE_LOG_ERROR_UHT(TEXT("Replication notification function %s has invalid parameter for property %s. First (optional) parameter must be of type %s%s."), *TargetFunc->GetName(), *Prop->GetName(), *CPPType, *ExtendedCPPType);
 			}
-
-			++Parm;
 		}
 
-		if ( TargetFunc->NumParms >= 2 && Parm)
+		if ( TargetFunc->NumParms >= 2)
 		{
+			FProperty* Parm = Properties[1]->GetProperty();
 			// A 2nd parameter for arrays can be specified as a const TArray<uint8>&. This is a list of element indices that have changed
-			FArrayProperty *ArrayProp = CastField<FArrayProperty>(*Parm);
+			FArrayProperty *ArrayProp = CastField<FArrayProperty>(Parm);
 			if (!(ArrayProp && CastField<FByteProperty>(ArrayProp->Inner)) || !(Parm->GetPropertyFlags() & CPF_ConstParm) || !(Parm->GetPropertyFlags() & CPF_ReferenceParm))
 			{
 				UE_LOG_ERROR_UHT(TEXT("Replication notification function %s (optional) second parameter must be of type 'const TArray<uint8>&'"), *TargetFunc->GetName());
@@ -3447,77 +2953,73 @@ void FHeaderParser::VerifyRepNotifyCallback(FProperty* Prop, UFunction* TargetFu
 		UE_LOG_ERROR_UHT(TEXT("Replication notification function %s not found"), *Prop->RepNotifyFunc.ToString() );
 	}
 }
-void FHeaderParser::VerifyPropertyMarkups( UClass* TargetClass )
+void FHeaderParser::VerifyPropertyMarkups(FUnrealClassDefinitionInfo& TargetClassDef)
 {
+	UClass* TargetClass = TargetClassDef.GetClass();
+
 	// Iterate over all properties, looking for those flagged as CPF_RepNotify
-	for ( FField* Field = TargetClass->ChildProperties; Field; Field = Field->Next )
+	for (FUnrealPropertyDefinitionInfo* PropertyDef : TargetClassDef.GetProperties())
 	{
-		if (FProperty* Prop = CastField<FProperty>(Field))
+		FProperty* Prop = PropertyDef->GetProperty();
+
+		auto FindTargetFunction = [&](const FName FuncName)
 		{
-			auto FindTargetFunction = [&](const FName FuncName)
+			// Search through this class and its super classes looking for the specified callback
+			for (FUnrealClassDefinitionInfo* SearchClassDef = &TargetClassDef; SearchClassDef; SearchClassDef = SearchClassDef->GetSuperClass())
 			{
-				// Search through this class and its superclasses looking for the specified callback
-				UFunction* TargetFunc = nullptr;
-				UClass* SearchClass = TargetClass;
-				while( SearchClass && !TargetFunc )
+				// Since the function map is not valid yet, we have to iterate over the fields to look for the function
+				for (FUnrealFunctionDefinitionInfo* FunctionDef : SearchClassDef->GetFunctions())
 				{
-					// Since the function map is not valid yet, we have to iterate over the fields to look for the function
-					for( UField* TestField = SearchClass->Children; TestField; TestField = TestField->Next )
+					UFunction* TestFunc = FunctionDef->GetFunction();
+					if (TestFunc && FNativeClassHeaderGenerator::GetOverriddenFName(TestFunc) == FuncName)
 					{
-						UFunction* TestFunc = Cast<UFunction>(TestField);
-						if (TestFunc && FNativeClassHeaderGenerator::GetOverriddenFName(TestFunc) == FuncName)
-						{
-							TargetFunc = TestFunc;
-							break;
-						}
+						return FunctionDef;
 					}
-					SearchClass = SearchClass->GetSuperClass();
 				}
+			}
+			return static_cast<FUnrealFunctionDefinitionInfo*>(nullptr);
+		};
 
-				return TargetFunc;
-			};
+		FUnrealPropertyDefinitionInfo& PropDef = GTypeDefinitionInfoMap.FindChecked<FUnrealPropertyDefinitionInfo>(Prop);
+		FPropertyBase& PropertyToken = PropDef.GetPropertyBase();
 
-			FUnrealPropertyDefinitionInfo& PropDef = GTypeDefinitionInfoMap.FindChecked(Prop).AsPropertyChecked();
-			FToken& PropertyToken = PropDef.GetToken();
+		TGuardValue<int32> GuardedInputPos(InputPos, PropDef.GetParsePosition());
+		TGuardValue<int32> GuardedInputLine(InputLine, PropDef.GetLineNumber());
 
-			TGuardValue<int32> GuardedInputPos(InputPos, PropertyToken.StartPos);
-			TGuardValue<int32> GuardedInputLine(InputLine, PropertyToken.StartLine);
+		if (Prop->HasAnyPropertyFlags(CPF_RepNotify))
+		{
+			VerifyRepNotifyCallback(Prop, FindTargetFunction(Prop->RepNotifyFunc));
+		}
 
-			if (Prop->HasAnyPropertyFlags(CPF_RepNotify))
+		if (Prop->HasAnyPropertyFlags(CPF_BlueprintVisible))
+		{
+			const FString& GetterFuncName = Prop->GetMetaData(NAME_BlueprintGetter);
+			if (!GetterFuncName.IsEmpty())
 			{
-				VerifyRepNotifyCallback(Prop, FindTargetFunction(Prop->RepNotifyFunc));
+				if (FUnrealFunctionDefinitionInfo* TargetFuncDef = FindTargetFunction(*GetterFuncName))
+				{
+					VerifyBlueprintPropertyGetter(Prop, TargetFuncDef);
+				}
+				else
+				{
+					// Couldn't find a valid function...
+					UE_LOG_ERROR_UHT(TEXT("Blueprint Property getter function %s not found"), *GetterFuncName);
+				}
 			}
 
-			if (Prop->HasAnyPropertyFlags(CPF_BlueprintVisible))
+			if (!Prop->HasAnyPropertyFlags(CPF_BlueprintReadOnly))
 			{
-				const FString& GetterFuncName = Prop->GetMetaData(NAME_BlueprintGetter);
-				if (!GetterFuncName.IsEmpty())
+				const FString& SetterFuncName = Prop->GetMetaData(NAME_BlueprintSetter);
+				if (!SetterFuncName.IsEmpty())
 				{
-					if (UFunction* TargetFunc = FindTargetFunction(*GetterFuncName))
+					if (FUnrealFunctionDefinitionInfo* TargetFuncDef = FindTargetFunction(*SetterFuncName))
 					{
-						VerifyBlueprintPropertyGetter(Prop, TargetFunc);
+						VerifyBlueprintPropertySetter(Prop, TargetFuncDef);
 					}
 					else
 					{
 						// Couldn't find a valid function...
-						UE_LOG_ERROR_UHT(TEXT("Blueprint Property getter function %s not found"), *GetterFuncName);
-					}
-				}
-
-				if (!Prop->HasAnyPropertyFlags(CPF_BlueprintReadOnly))
-				{
-					const FString& SetterFuncName = Prop->GetMetaData(NAME_BlueprintSetter);
-					if (!SetterFuncName.IsEmpty())
-					{
-						if (UFunction* TargetFunc = FindTargetFunction(*SetterFuncName))
-						{
-							VerifyBlueprintPropertySetter(Prop, TargetFunc);
-						}
-						else
-						{
-							// Couldn't find a valid function...
-							UE_LOG_ERROR_UHT(TEXT("Blueprint Property setter function %s not found"), *SetterFuncName);
-						}
+						UE_LOG_ERROR_UHT(TEXT("Blueprint Property setter function %s not found"), *SetterFuncName);
 					}
 				}
 			}
@@ -3556,7 +3058,7 @@ void FHeaderParser::CompileDirective()
 	else if (Directive.Matches(TEXT("linenumber"), ESearchCase::CaseSensitive))
 	{
 		FToken Number;
-		if (!GetToken(Number) || (Number.TokenType != TOKEN_Const) || (Number.Type != CPT_Int && Number.Type != CPT_Int64))
+		if (!GetToken(Number) || (Number.TokenType != TOKEN_Const) || (Number.ConstantType != CPT_Int && Number.ConstantType != CPT_Int64))
 		{
 			FError::Throwf(TEXT("Missing line number in line number directive"));
 		}
@@ -3570,7 +3072,7 @@ void FHeaderParser::CompileDirective()
 	else if (Directive.Matches(TEXT("include"), ESearchCase::CaseSensitive))
 	{
 		FToken IncludeName;
-		if (GetToken(IncludeName) && (IncludeName.TokenType == TOKEN_Const) && (IncludeName.Type == CPT_String))
+		if (GetToken(IncludeName) && (IncludeName.TokenType == TOKEN_Const) && (IncludeName.ConstantType == CPT_String))
 		{
 			const FString& ExpectedHeaderName = SourceFile.GetGeneratedHeaderFilename();
 			if (FCString::Stricmp(IncludeName.String, *ExpectedHeaderName) == 0)
@@ -3677,6 +3179,7 @@ void FHeaderParser::GetVarType(
 	FPropertyBase&                  VarProperty,
 	EPropertyFlags                  Disallow,
 	const FToken*                   OuterPropertyType,
+	const EPropertyFlags*			OuterPropertyFlags,
 	EPropertyDeclarationStyle::Type PropertyDeclarationStyle,
 	EVariableCategory::Type         VariableCategory,
 	FIndexRange*                    ParsedVarIndexRange,
@@ -4400,9 +3903,9 @@ void FHeaderParser::GetVarType(
 	{
 		RequireSymbol( TEXT('<'), TEXT("'tarray'") );
 
-		VarType.PropertyFlags = Flags;
+		EPropertyFlags VarTypePropertyFlags = Flags;
 
-		GetVarType(Scope, VarProperty, Disallow, &VarType, EPropertyDeclarationStyle::None, VariableCategory);
+		GetVarType(Scope, VarProperty, Disallow, &VarType, &VarTypePropertyFlags, EPropertyDeclarationStyle::None, VariableCategory);
 		if (VarProperty.IsContainer())
 		{
 			FError::Throwf(TEXT("Nested containers are not supported.") );
@@ -4414,7 +3917,7 @@ void FHeaderParser::GetVarType(
 			bNativeConstTemplateArg = true;
 		}
 
-		VarType.PropertyFlags = VarProperty.PropertyFlags & (CPF_ContainsInstancedReference | CPF_InstancedReference); // propagate these to the array, we will fix them later
+		VarTypePropertyFlags = VarProperty.PropertyFlags & (CPF_ContainsInstancedReference | CPF_InstancedReference); // propagate these to the array, we will fix them later
 		VarProperty.ArrayType = EArrayType::Dynamic;
 
 		FToken CloseTemplateToken;
@@ -4489,10 +3992,10 @@ void FHeaderParser::GetVarType(
 	{
 		RequireSymbol( TEXT('<'), TEXT("'tmap'") );
 
-		VarType.PropertyFlags = Flags;
+		EPropertyFlags VarTypePropertyFlags = Flags;
 
-		FToken MapKeyType;
-		GetVarType(Scope, MapKeyType, Disallow, &VarType, EPropertyDeclarationStyle::None, VariableCategory);
+		FPropertyBase MapKeyType(CPT_None);
+		GetVarType(Scope, MapKeyType, Disallow, &VarType, &VarTypePropertyFlags, EPropertyDeclarationStyle::None, VariableCategory);
 		if (MapKeyType.IsContainer())
 		{
 			FError::Throwf(TEXT("Nested containers are not supported.") );
@@ -4520,7 +4023,7 @@ void FHeaderParser::GetVarType(
 			FError::Throwf(TEXT("Missing value type while parsing TMap."));
 		}
 
-		GetVarType(Scope, VarProperty, Disallow, &VarType, EPropertyDeclarationStyle::None, VariableCategory);
+		GetVarType(Scope, VarProperty, Disallow, &VarType, &VarTypePropertyFlags, EPropertyDeclarationStyle::None, VariableCategory);
 		if (VarProperty.IsContainer())
 		{
 			FError::Throwf(TEXT("Nested containers are not supported.") );
@@ -4528,8 +4031,8 @@ void FHeaderParser::GetVarType(
 		// TODO: Prevent sparse delegate types from being used in a container
 
 		EPropertyFlags InnerFlags = (MapKeyType.PropertyFlags | VarProperty.PropertyFlags) & (CPF_ContainsInstancedReference | CPF_InstancedReference); // propagate these to the map value, we will fix them later
-		VarType.PropertyFlags = InnerFlags;
-		VarProperty.MapKeyProp = MakeShared<FToken>(MoveTemp(MapKeyType));
+		VarTypePropertyFlags = InnerFlags;
+		VarProperty.MapKeyProp = MakeShared<FPropertyBase>(MoveTemp(MapKeyType));
 		VarProperty.MapKeyProp->PropertyFlags = InnerFlags | (VarProperty.MapKeyProp->PropertyFlags & CPF_UObjectWrapper); // Make sure the 'UObjectWrapper' flag is maintained so that 'TMap<TSubclassOf<...>, ...>' works
 
 		FToken CloseTemplateToken;
@@ -4574,9 +4077,9 @@ void FHeaderParser::GetVarType(
 	{
 		RequireSymbol( TEXT('<'), TEXT("'tset'") );
 
-		VarType.PropertyFlags = Flags;
+		EPropertyFlags VarTypePropertyFlags = Flags;
 
-		GetVarType(Scope, VarProperty, Disallow, &VarType, EPropertyDeclarationStyle::None, VariableCategory);
+		GetVarType(Scope, VarProperty, Disallow, &VarType, &VarTypePropertyFlags, EPropertyDeclarationStyle::None, VariableCategory);
 		if (VarProperty.IsContainer())
 		{
 			FError::Throwf(TEXT("Nested containers are not supported.") );
@@ -4598,7 +4101,7 @@ void FHeaderParser::GetVarType(
 			UE_LOG_ERROR_UHT(TEXT("Replicated sets are not supported."));
 		}
 
-		VarType.PropertyFlags = VarProperty.PropertyFlags & (CPF_ContainsInstancedReference | CPF_InstancedReference); // propagate these to the set, we will fix them later
+		VarTypePropertyFlags = VarProperty.PropertyFlags & (CPF_ContainsInstancedReference | CPF_InstancedReference); // propagate these to the set, we will fix them later
 		VarProperty.ArrayType = EArrayType::Set;
 
 		FToken CloseTemplateToken;
@@ -4672,7 +4175,7 @@ void FHeaderParser::GetVarType(
 		FToken InnerEnumType;
 		if (GetIdentifier(InnerEnumType, true))
 		{
-			if (UEnum* Enum = FindObject<UEnum>(ANY_PACKAGE, InnerEnumType.Identifier))
+			if (UEnum* Enum = FClasses::FindObject<UEnum>(ANY_PACKAGE, InnerEnumType.Identifier))
 			{
 				// In-scope enumeration.
 				VarProperty = FPropertyBase(Enum, CPT_Byte);
@@ -4725,7 +4228,7 @@ void FHeaderParser::GetVarType(
 
 		VarProperty = FPropertyBase(PropertyClass, CPT_FieldPath);
 	}
-	else if (UEnum* Enum = FindObject<UEnum>(ANY_PACKAGE, VarType.Identifier))
+	else if (UEnum* Enum = FClasses::FindObject<UEnum>(ANY_PACKAGE, VarType.Identifier))
 	{
 		EPropertyType UnderlyingType = CPT_Byte;
 
@@ -4758,10 +4261,10 @@ void FHeaderParser::GetVarType(
 		bool bHandledType = false;
 		FString IdentifierStripped = GetClassNameWithPrefixRemoved(VarType.Identifier);
 		bool bStripped = false;
-		UScriptStruct* Struct = FindObject<UScriptStruct>(ANY_PACKAGE, VarType.Identifier);
+		UScriptStruct* Struct = FClasses::FindObject<UScriptStruct>(ANY_PACKAGE, VarType.Identifier);
 		if (!Struct)
 		{
-			Struct = FindObject<UScriptStruct>(ANY_PACKAGE, *IdentifierStripped);
+			Struct = FClasses::FindObject<UScriptStruct>(ANY_PACKAGE, *IdentifierStripped);
 			bStripped = true;
 		}
 
@@ -4835,7 +4338,7 @@ void FHeaderParser::GetVarType(
 			// Struct keyword in front of a struct is legal, we 'consume' it
 			bUnconsumedStructKeyword = false;
 		}
-		else if (FindObject<UScriptStruct>(ANY_PACKAGE, *IdentifierStripped) != nullptr)
+		else if (FClasses::FindObject<UScriptStruct>(ANY_PACKAGE, *IdentifierStripped) != nullptr)
 		{
 			bHandledType = true;
 
@@ -5010,7 +4513,7 @@ void FHeaderParser::GetVarType(
 					}
 
 					// Inherit instancing flags
-					if (DoesAnythingInHierarchyHaveDefaultToInstanced(TempClass))
+					if (FUnrealClassDefinitionInfo::HierarchyHasAnyClassFlags(TempClass, CLASS_DefaultToInstanced))
 					{
 						Flags |= ((CPF_InstancedReference | CPF_ExportObject) & (~Disallow));
 					}
@@ -5053,7 +4556,7 @@ void FHeaderParser::GetVarType(
 			// Resolve delegates declared in another class  //@TODO: UCREMOVAL: This seems extreme
 			if (!bHandledType)
 			{
-				if (UFunction* ExternDelegateFunc = (UFunction*)StaticFindObject(UFunction::StaticClass(), ANY_PACKAGE, *(IdentifierStripped + HEADER_GENERATED_DELEGATE_SIGNATURE_SUFFIX)))
+				if (UFunction* ExternDelegateFunc = FClasses::FindObject<UFunction>(ANY_PACKAGE, *(IdentifierStripped + HEADER_GENERATED_DELEGATE_SIGNATURE_SUFFIX)))
 				{
 					SetDelegateType(ExternDelegateFunc, IdentifierStripped);
 				}
@@ -5218,7 +4721,7 @@ void FHeaderParser::GetVarType(
 	if (VariableCategory != EVariableCategory::Member)
 	{
 		// These conditions are checked externally for struct/member variables where the flag can be inferred later on from the variable name itself
-		ValidatePropertyIsDeprecatedIfNecessary(VarProperty, OuterPropertyType);
+		ValidatePropertyIsDeprecatedIfNecessary(VarProperty, OuterPropertyFlags);
 	}
 
 	// Check for invalid transients
@@ -5258,10 +4761,10 @@ void FHeaderParser::GetVarType(
 	}
 }
 
-FProperty* FHeaderParser::GetVarNameAndDim
+FUnrealPropertyDefinitionInfo& FHeaderParser::GetVarNameAndDim
 (
 	FUnrealStructDefinitionInfo& ParentStruct, 
-	FToken&                 VarProperty,
+	FPropertyBase& VarProperty,
 	EVariableCategory::Type VariableCategory,
 	ELayoutMacroType        LayoutMacroType
 )
@@ -5279,13 +4782,13 @@ FProperty* FHeaderParser::GetVarNameAndDim
 	const TCHAR* HintText = GetHintText(VariableCategory);
 
 	AddModuleRelativePathToMetadata(ParentStruct, VarProperty.MetaData);
+	TCHAR Identifier[NAME_SIZE];
 
 	// Get variable name.
 	if (VariableCategory == EVariableCategory::Return)
 	{
 		// Hard-coded variable name, such as with return value.
-		VarProperty.TokenType = TOKEN_Identifier;
-		FCString::Strcpy( VarProperty.Identifier, TEXT("ReturnValue") );
+		FCString::Strcpy( Identifier, TEXT("ReturnValue") );
 	}
 	else
 	{
@@ -5309,13 +4812,12 @@ FProperty* FHeaderParser::GetVarNameAndDim
 				break;
 		}
 
-		VarProperty.TokenType = TOKEN_Identifier;
-		FCString::Strcpy(VarProperty.Identifier, VarToken.Identifier);
+		FCString::Strcpy(Identifier, VarToken.Identifier);
 	}
 
 	// Check to see if the variable is deprecated, and if so set the flag
 	{
-		FString VarName(VarProperty.Identifier);
+		FString VarName(Identifier);
 
 		const int32 DeprecatedIndex = VarName.Find(TEXT("_DEPRECATED"));
 		const int32 NativizedPropertyPostfixIndex = VarName.Find(TEXT("__pf")); //TODO: check OverrideNativeName in Meta Data, to be sure it's not a random occurrence of the "__pf" string.
@@ -5353,14 +4855,14 @@ FProperty* FHeaderParser::GetVarNameAndDim
 			VarProperty.PropertyFlags |= CPF_Deprecated;
 			VarName.MidInline(0, DeprecatedIndex, false);
 
-			FCString::Strcpy(VarProperty.Identifier, *VarName);
+			FCString::Strcpy(Identifier, *VarName);
 		}
 	}
 
 	// Make sure it doesn't conflict.
 	int32 OuterContextCount = 0;
-	UField* ExistingField = FindField(Scope, VarProperty.Identifier, true, UField::StaticClass(), NULL);
-	FField* ExistingProperty = FindProperty(Scope, VarProperty.Identifier, true, FField::StaticClass(), NULL);
+	UField* ExistingField = FindField(ParentStruct, Identifier, true, UField::StaticClass(), NULL);
+	FField* ExistingProperty = FindProperty(ParentStruct, Identifier, true, FField::StaticClass(), NULL);
 
 	if (ExistingField != nullptr || ExistingProperty != nullptr)
 	{
@@ -5390,7 +4892,7 @@ FProperty* FHeaderParser::GetVarNameAndDim
 		{
 			FError::Throwf(TEXT("%s: '%s' cannot be defined in '%s' as it is already defined in scope '%s' (shadowing is not allowed)"),
 				HintText,
-				VarProperty.Identifier,
+				Identifier,
 				*Scope->GetName(),
 				ExistingField ? *ExistingField->GetOuter()->GetName() : *ExistingProperty->GetOwnerVariant().GetFullName());
 		}
@@ -5429,7 +4931,7 @@ FProperty* FHeaderParser::GetVarNameAndDim
 			// Ignore how the actual array dimensions are actually defined - we'll calculate those with the compiler anyway.
 			if (!GetRawToken(Dimensions, TEXT(']')))
 			{
-				FError::Throwf(TEXT("%s %s: Missing ']'"), HintText, VarProperty.Identifier);
+				FError::Throwf(TEXT("%s %s: Missing ']'"), HintText, Identifier);
 			}
 		}
 		else
@@ -5437,7 +4939,7 @@ FProperty* FHeaderParser::GetVarNameAndDim
 			// Ignore how the actual array dimensions are actually defined - we'll calculate those with the compiler anyway.
 			if (!GetRawToken(Dimensions, TEXT(')')))
 			{
-				FError::Throwf(TEXT("%s %s: Missing ']'"), HintText, VarProperty.Identifier);
+				FError::Throwf(TEXT("%s %s: Missing ']'"), HintText, Identifier);
 			}
 		}
 
@@ -5499,7 +5001,7 @@ FProperty* FHeaderParser::GetVarNameAndDim
 		if (!Enum)
 		{
 			// If the enum wasn't declared in this scope, then try to find it anywhere we can
-			Enum = FindObject<UEnum>(ANY_PACKAGE, Dimensions.String);
+			Enum = FClasses::FindObject<UEnum>(ANY_PACKAGE, Dimensions.String);
 		}
 
 		if (Enum)
@@ -5517,7 +5019,7 @@ FProperty* FHeaderParser::GetVarNameAndDim
 	// Try gathering metadata for member fields
 	if (VariableCategory == EVariableCategory::Member)
 	{
-		ParseFieldMetaData(VarProperty.MetaData, VarProperty.Identifier);
+		ParseFieldMetaData(VarProperty.MetaData, Identifier);
 		AddFormattedPrevCommentAsTooltipMetaData(VarProperty.MetaData);
 	}
 	// validate UFunction parameters
@@ -5537,124 +5039,30 @@ FProperty* FHeaderParser::GetVarNameAndDim
 		}
 	}
 
-	FName PropertyName(VarProperty.Identifier);
-
-	FProperty* Prev = nullptr;
-	for (TFieldIterator<FProperty> It(Scope, EFieldIteratorFlags::ExcludeSuper); It; ++It)
-	{
-		Prev = *It;
-	}
-
-	auto PropagateFlagsFromInnerAndHandlePersistentInstanceMetadata = [](EPropertyFlags& DestFlags, const TMap<FName, FString>& InMetaData, FProperty* Inner) {
-		// Copy some of the property flags to the container property.
-		if (Inner->PropertyFlags & (CPF_ContainsInstancedReference | CPF_InstancedReference))
-		{
-			DestFlags |= CPF_ContainsInstancedReference;
-			DestFlags &= ~(CPF_InstancedReference | CPF_PersistentInstance); //this was propagated to the inner
-
-			if (Inner->PropertyFlags & CPF_PersistentInstance)
-			{
-				AddMetaDataToClassData(Inner, TMap<FName,FString>(InMetaData));
-			}
-		}
-	};
-
-	FProperty* Result = nullptr;
-	if (VarProperty.ArrayType == EArrayType::Dynamic)
-	{
-		FArrayProperty* Array     = new FArrayProperty(Scope, PropertyName, ObjectFlags);
-		FProperty*      InnerProp = CreateVariableProperty(VarProperty, Array, PropertyName, RF_Public, VariableCategory, &SourceFile);
-
-		Array->Inner         = InnerProp;
-		Array->PropertyFlags = VarProperty.PropertyFlags;
-
-		// Propagate flags
-		InnerProp->PropertyFlags |= Array->PropertyFlags & CPF_PropagateToArrayInner;
-
-		PropagateFlagsFromInnerAndHandlePersistentInstanceMetadata(Array->PropertyFlags, VarProperty.MetaData, InnerProp);
-
-		Result = Array;
-
-		FUnrealPropertyDefinitionInfo* PropDef = new FUnrealPropertyDefinitionInfo(SourceFile, 0, Array, VarProperty.AllocatorType);
-		GTypeDefinitionInfoMap.Add(Result, TSharedRef<FUnrealTypeDefinitionInfo>(PropDef));
-	}
-	else if (VarProperty.ArrayType == EArrayType::Set)
-	{
-		FSetProperty* Set       = new FSetProperty(Scope, PropertyName, ObjectFlags);
-		FProperty*    InnerProp = CreateVariableProperty(VarProperty, Set, PropertyName, RF_Public, VariableCategory, &SourceFile);
-
-		Set->ElementProp   = InnerProp;
-		Set->PropertyFlags = VarProperty.PropertyFlags;
-
-		// Propagate flags
-		InnerProp->PropertyFlags |= Set->PropertyFlags & CPF_PropagateToSetElement;
+	// Create the property
+	// NOTE: Only FProperyBase part of VarProperty is used.  
+	FName PropertyName(Identifier);
+	//VarProperty.StartLine = InputLine;
+	//VarProperty.StartPos = InputPos;
+	FUnrealPropertyDefinitionInfo& PropDef = FPropertyTraits::CreateProperty(VarProperty, Scope, PropertyName, ObjectFlags, VariableCategory, Dimensions.String, SourceFile, InputLine, InputPos);
+	FProperty* Property = PropDef.GetProperty();
 
-		PropagateFlagsFromInnerAndHandlePersistentInstanceMetadata(Set->PropertyFlags, VarProperty.MetaData, InnerProp);
+	// Add the property to the parent
+	ParentStruct.AddProperty(PropDef);
 
-		Result = Set;
-
-		FUnrealPropertyDefinitionInfo* PropDef = new FUnrealPropertyDefinitionInfo(SourceFile, 0, Set);
-		GTypeDefinitionInfoMap.Add(Result, TSharedRef<FUnrealTypeDefinitionInfo>(PropDef));
-	}
-	else if (VarProperty.MapKeyProp.IsValid())
+	// Add to the end of the properties slist
+	FField** Prev = &Scope->ChildProperties;
+	for (; *Prev != nullptr; Prev = &(*Prev)->Next)
 	{
-		FMapProperty* Map = new FMapProperty(Scope, PropertyName, ObjectFlags);
-		FProperty*    KeyProp   = CreateVariableProperty(*VarProperty.MapKeyProp, Map, *(PropertyName.ToString() + TEXT("_Key")), RF_Public, VariableCategory, &SourceFile);
-		FProperty*    ValueProp = CreateVariableProperty(VarProperty,             Map, PropertyName,                              RF_Public, VariableCategory, &SourceFile);
-
-		Map->KeyProp       = KeyProp;
-		Map->ValueProp     = ValueProp;
-		Map->PropertyFlags = VarProperty.PropertyFlags;
-
-		// Propagate flags
-		KeyProp  ->PropertyFlags |= VarProperty.MapKeyProp->PropertyFlags & CPF_PropagateToMapKey;
-		ValueProp->PropertyFlags |= Map->PropertyFlags                    & CPF_PropagateToMapValue;
-
-		PropagateFlagsFromInnerAndHandlePersistentInstanceMetadata(Map->PropertyFlags, VarProperty.MapKeyProp->MetaData, KeyProp);
-		PropagateFlagsFromInnerAndHandlePersistentInstanceMetadata(Map->PropertyFlags, VarProperty.MetaData,             ValueProp);
-
-		Result = Map;
-
-		FUnrealPropertyDefinitionInfo* PropDef = new FUnrealPropertyDefinitionInfo(SourceFile, 0, Map, VarProperty.AllocatorType);
-		GTypeDefinitionInfoMap.Add(Result, TSharedRef<FUnrealTypeDefinitionInfo>(PropDef));
+		// No body
 	}
-	else
-	{
-		Result = CreateVariableProperty(VarProperty, Scope, PropertyName, ObjectFlags, VariableCategory, &SourceFile);
-
-		if (VarProperty.ArrayType == EArrayType::Static)
-		{
-			Result->ArrayDim = 2; // 2 = static array
-			GTypeDefinitionInfoMap.FindChecked(Result).AsPropertyChecked().SetArrayDimensions(Dimensions.String);
-		}
-
-		Result->PropertyFlags = VarProperty.PropertyFlags;
-	}
-
-	if (Prev != nullptr)
-	{
-		Result->Next = Prev->Next;
-		Prev->Next = Result;
-	}
-	else
-	{
-		Result->Next = Scope->ChildProperties;
-		Scope->ChildProperties = Result;
-	}
-
-	VarProperty.TokenProperty = Result;
-	VarProperty.StartLine = InputLine;
-	VarProperty.StartPos = InputPos;
-
-	FUnrealPropertyDefinitionInfo& ResultDef = GTypeDefinitionInfoMap.FindChecked(Result).AsPropertyChecked();
-	ResultDef.GetToken() = VarProperty;
-	ParentStruct.AddProperty(ResultDef);
-	ParentStruct.GetStructMetaData().AddProperty(ResultDef);
+	check(*Prev == nullptr);
+	Property->Next = nullptr;
+	*Prev = Property;
 
 	// if we had any metadata, add it to the class
-	AddMetaDataToClassData(VarProperty.TokenProperty, TMap<FName,FString>(VarProperty.MetaData));
-
-	return Result;
+	AddMetaDataToClassData(PropDef, TMap<FName,FString>(PropDef.GetPropertyBase().MetaData));
+	return PropDef;
 }
 
 /*-----------------------------------------------------------------------------
@@ -5664,7 +5072,7 @@ FProperty* FHeaderParser::GetVarNameAndDim
 //
 // Compile a declaration in Token. Returns 1 if compiled, 0 if not.
 //
-bool FHeaderParser::CompileDeclaration(TArray<UDelegateFunction*>& DelegatesToFixup, FToken& Token)
+bool FHeaderParser::CompileDeclaration(TArray<FUnrealFunctionDefinitionInfo*>& DelegatesToFixup, FToken& Token)
 {
 	EAccessSpecifier AccessSpecifier = ParseAccessProtectionSpecifier(Token);
 	if (AccessSpecifier)
@@ -5825,15 +5233,15 @@ bool FHeaderParser::CompileDeclaration(TArray<UDelegateFunction*>& DelegatesToFi
 
 	if (Token.Matches(TEXT("UDELEGATE"), ESearchCase::CaseSensitive))
 	{
-		UDelegateFunction* Delegate = CompileDelegateDeclaration(Token.Identifier, EDelegateSpecifierAction::Parse);
-		DelegatesToFixup.Add(Delegate);
+		FUnrealFunctionDefinitionInfo &DelegateDef = CompileDelegateDeclaration(Token.Identifier, EDelegateSpecifierAction::Parse);
+		DelegatesToFixup.Add(&DelegateDef);
 		return true;
 	}
 
 	if (IsValidDelegateDeclaration(Token)) // Legacy delegate parsing - it didn't need a UDELEGATE
 	{
-		UDelegateFunction* Delegate = CompileDelegateDeclaration(Token.Identifier);
-		DelegatesToFixup.Add(Delegate);
+		FUnrealFunctionDefinitionInfo& DelegateDef = CompileDelegateDeclaration(Token.Identifier);
+		DelegatesToFixup.Add(&DelegateDef);
 		return true;
 	}
 
@@ -5881,7 +5289,7 @@ bool FHeaderParser::CompileDeclaration(TArray<UDelegateFunction*>& DelegatesToFi
 		bEncounteredNewStyleClass_UnmatchedBrackets = false;
 
 		UClass* CurrentClass = GetCurrentClass();
-		FUnrealClassDefinitionInfo& CurrentClassDef = GetCurrentClassDef().AsClassChecked();
+		FUnrealClassDefinitionInfo& CurrentClassDef = UHTCastChecked<FUnrealClassDefinitionInfo>(GetCurrentClassDef());
 
 		// Pop nesting here to allow other non UClass declarations in the header file.
 		if (CurrentClass->ClassFlags & CLASS_Interface)
@@ -6067,7 +5475,7 @@ bool FHeaderParser::CompileDeclaration(TArray<UDelegateFunction*>& DelegatesToFi
 						{
 							FString EnclosingDefine = CompilerDirectiveStack.Num() > 0 ? TEXT("WITH_EDITORONLY_DATA") : TEXT("");
 
-							FUnrealClassDefinitionInfo& ClassDef = GetCurrentClassDef().AsClassChecked();
+							FUnrealClassDefinitionInfo& ClassDef = UHTCastChecked<FUnrealClassDefinitionInfo>(GetCurrentClassDef());
 							ClassDef.AddArchiveType(ArchiveType);
 							ClassDef.SetEnclosingDefine(MoveTemp(EnclosingDefine));
 						}
@@ -6233,7 +5641,7 @@ FUnrealClassDefinitionInfo& FHeaderParser::ParseClassNameDeclaration(FString& De
 	check(FoundClass);
 	FoundClass->ClassCastFlags |= ClassCastFlagMap::Get().GetCastFlag(DeclaredClassName);
 
-	FUnrealClassDefinitionInfo& ClassDef = GTypeDefinitionInfoMap.FindChecked(FoundClass).AsClassChecked();
+	FUnrealClassDefinitionInfo& ClassDef = GTypeDefinitionInfoMap.FindChecked<FUnrealClassDefinitionInfo>(FoundClass);
 
 	// Skip optional final keyword
 	MatchIdentifier(TEXT("final"), ESearchCase::CaseSensitive);
@@ -6242,16 +5650,15 @@ FUnrealClassDefinitionInfo& FHeaderParser::ParseClassNameDeclaration(FString& De
 	ParseInheritance(TEXT("class"), [](const TCHAR* ClassName, bool bIsSuperClass) {}); // Eat the results, already been parsed
 
 	// Collect data about the super and base classes
-	if (FUnrealStructDefinitionInfo* SuperStruct = ClassDef.GetSuperClassInfo().Struct)
+	if (FUnrealClassDefinitionInfo* SuperClassDef = UHTCast<FUnrealClassDefinitionInfo>(ClassDef.GetSuperClassInfo().Struct))
 	{
-		UClass* SuperClass = SuperStruct->AsClassChecked().GetClass();
-		FoundClass->ClassCastFlags |= SuperClass->ClassCastFlags;
+		FoundClass->ClassCastFlags |= SuperClassDef->GetClass()->ClassCastFlags;
 	}
 	for (FUnrealStructDefinitionInfo::FBaseClassInfo& BaseClassInfo : ClassDef.GetBaseClassInfo())
 	{
-		if (BaseClassInfo.Struct != nullptr)
+		if (FUnrealClassDefinitionInfo* BaseClassDef = UHTCast<FUnrealClassDefinitionInfo>(BaseClassInfo.Struct))
 		{
-			UClass* BaseClass = BaseClassInfo.Struct->AsClassChecked().GetClass();
+			UClass* BaseClass = BaseClassDef->GetClass();
 			if (!BaseClass->HasAnyClassFlags(CLASS_Interface))
 			{
 				FError::Throwf(TEXT("Implements: Class %s is not an interface; Can only inherit from non-UObjects or UInterface derived interfaces"), *BaseClass->GetName());
@@ -6268,10 +5675,12 @@ FUnrealClassDefinitionInfo& FHeaderParser::ParseClassNameDeclaration(FString& De
 /**
  * Setups basic class settings after parsing.
  */
-void PostParsingClassSetup(UClass* Class)
+void PostParsingClassSetup(FUnrealClassDefinitionInfo& ClassDef)
 {
+	UClass* Class = ClassDef.GetClass();
+
 	// Cleanup after first pass.
-	FHeaderParser::ComputeFunctionParametersSize(Class);
+	FHeaderParser::ComputeFunctionParametersSize(ClassDef);
 
 	// Set all optimization ClassFlags based on property types
 	auto HasAllOptimizationClassFlags = [Class]()
@@ -6281,15 +5690,16 @@ void PostParsingClassSetup(UClass* Class)
 
 	if (!HasAllOptimizationClassFlags())
 	{
-		for (TFieldIterator<FProperty> It(Class, EFieldIteratorFlags::ExcludeSuper); It; ++It)
+		for (FUnrealPropertyDefinitionInfo* PropertyDef : ClassDef.GetProperties())
 		{
-			if ((It->PropertyFlags & CPF_Config) != 0)
+			FProperty* Property = PropertyDef->GetProperty();
+			if ((Property->PropertyFlags & CPF_Config) != 0)
 			{
 				Class->ClassFlags |= CLASS_Config;
 				if (HasAllOptimizationClassFlags()) break;
 			}
 
-			if (It->ContainsInstancedObjectProperty())
+			if (Property->ContainsInstancedObjectProperty())
 			{
 				Class->ClassFlags |= CLASS_HasInstancedReference;
 				if (HasAllOptimizationClassFlags()) break;
@@ -6375,7 +5785,7 @@ UClass* FHeaderParser::CompileClassDeclaration()
 	ClassDeclarationData.MergeClassCategories(Class);
 
 	// Class attributes.
-	FStructMetaData& StructData = GTypeDefinitionInfoMap.FindChecked(Class).AsStructChecked().GetStructMetaData();
+	FStructMetaData& StructData = GTypeDefinitionInfoMap.FindChecked<FUnrealStructDefinitionInfo>(Class).GetStructMetaData();
 	StructData.SetPrologLine(PrologFinishLine);
 
 	ClassDeclarationData.MergeAndValidateClassFlags(DeclaredClassName, PrevClassFlags, Class);
@@ -6395,7 +5805,7 @@ UClass* FHeaderParser::CompileClassDeclaration()
 	AddModuleRelativePathToMetadata(ClassDeclarationData, MetaData);
 
 	// Register the metadata
-	AddMetaDataToClassData(Class, MoveTemp(MetaData));
+	AddMetaDataToClassData(ClassDeclarationData, MoveTemp(MetaData));
 
 	// Handle the start of the rest of the class
 	RequireSymbol( TEXT('{'), TEXT("'Class'") );
@@ -6421,7 +5831,7 @@ UClass* FHeaderParser::CompileClassDeclaration()
 		// if this base class corresponds to an interface class, assign the vtable FProperty in the class's Interfaces map now...
 		if (FUnrealStructDefinitionInfo* BaseClassStruct = BaseClassInfos[ParentIndex].Struct)
 		{
-			UClass* BaseClass = BaseClassStruct->AsClassChecked().GetClass();
+			UClass* BaseClass = UHTCastChecked<FUnrealClassDefinitionInfo>(BaseClassStruct).GetClass();
 			if (BaseClass->HasAnyClassFlags(CLASS_Native))
 			{
 				FImplementedInterface* Found = Class->Interfaces.FindByPredicate([&BaseClass](const FImplementedInterface& Impl) { return Impl.Class == BaseClass; });
@@ -6551,7 +5961,7 @@ void FHeaderParser::CompileInterfaceDeclaration()
 	// New style files have the interface name / extends afterwards
 	RequireIdentifier(TEXT("class"), ESearchCase::CaseSensitive, TEXT("Interface declaration"));
 	FClass* InterfaceClass = ParseInterfaceNameDeclaration(/*out*/ DeclaredInterfaceName, /*out*/ RequiredAPIMacroIfPresent);
-	FUnrealClassDefinitionInfo& ClassDef = GTypeDefinitionInfoMap.FindChecked(InterfaceClass).AsClassChecked();
+	FUnrealClassDefinitionInfo& ClassDef = GTypeDefinitionInfoMap.FindChecked<FUnrealClassDefinitionInfo>(InterfaceClass);
 	ClassDefinitionRanges.Add(InterfaceClass, ClassDefinitionRange(&Input[InputPos], nullptr));
 
 	// Record that this interface is RequiredAPI if the CORE_API style macro was present
@@ -6615,12 +6025,12 @@ void FHeaderParser::CompileInterfaceDeclaration()
 	}
 
 	// Try parsing metadata for the interface
-	FStructMetaData& StructData = GTypeDefinitionInfoMap.FindChecked(InterfaceClass).AsStructChecked().GetStructMetaData();
+	FStructMetaData& StructData = GTypeDefinitionInfoMap.FindChecked<FUnrealStructDefinitionInfo>(InterfaceClass).GetStructMetaData();
 	StructData.SetPrologLine(PrologFinishLine);
 
 	// Register the metadata
 	AddModuleRelativePathToMetadata(ClassDef, MetaData);
-	AddMetaDataToClassData(InterfaceClass, MoveTemp(MetaData));
+	AddMetaDataToClassData(ClassDef, MoveTemp(MetaData));
 
 	// Handle the start of the rest of the interface
 	RequireSymbol( TEXT('{'), TEXT("'Class'") );
@@ -6771,8 +6181,9 @@ const TCHAR* FHeaderParser::GetRefText = TEXT("GetRef");
 const TCHAR* FHeaderParser::GetFixedArrayText = TEXT("GetFixedArray");
 const TCHAR* FHeaderParser::GetDynamicArrayText = TEXT("GetDynamicArray");
 
-void FHeaderParser::ParseRigVMMethodParameters(UStruct* Struct)
+void FHeaderParser::ParseRigVMMethodParameters(FUnrealStructDefinitionInfo& StructDef)
 {
+	UStruct* Struct = StructDef.GetStruct();
 	FRigVMStructInfo* StructRigVMInfo = StructRigVMMap.Find(Struct);
 	if (StructRigVMInfo == nullptr)
 	{
@@ -6780,9 +6191,9 @@ void FHeaderParser::ParseRigVMMethodParameters(UStruct* Struct)
 	}
 
 	// validate the property types for this struct
-	for (TFieldIterator<FProperty> It(Struct); It; ++It)
+	for (FUnrealPropertyDefinitionInfo* PropertyDef : TUHTFieldRange<FUnrealPropertyDefinitionInfo>(StructDef))
 	{
-		FProperty const* const Prop = *It;
+		FProperty const* const Prop = PropertyDef->GetProperty();
 		FString MemberCPPType;
 		FString ExtendedCPPType;
 		MemberCPPType = Prop->GetCPPType(&ExtendedCPPType);
@@ -6885,9 +6296,9 @@ void FHeaderParser::ParseParameterList(FUnrealFunctionDefinitionInfo& FunctionDe
 	do
 	{
 		// Get parameter type.
-		FToken Property(CPT_None);
+		FPropertyBase Property(CPT_None);
 		EVariableCategory::Type VariableCategory = (Function->FunctionFlags & FUNC_Net) ? EVariableCategory::ReplicatedParameter : EVariableCategory::RegularParameter;
-		GetVarType(GetCurrentScope(), Property, ~(CPF_ParmFlags | CPF_AutoWeak | CPF_RepSkip | CPF_UObjectWrapper | CPF_NativeAccessSpecifiers), NULL, EPropertyDeclarationStyle::None, VariableCategory);
+		GetVarType(GetCurrentScope(), Property, ~(CPF_ParmFlags | CPF_AutoWeak | CPF_RepSkip | CPF_UObjectWrapper | CPF_NativeAccessSpecifiers), nullptr, nullptr, EPropertyDeclarationStyle::None, VariableCategory);
 		Property.PropertyFlags |= CPF_Parm;
 
 		if (bExpectCommaBeforeName)
@@ -6895,7 +6306,8 @@ void FHeaderParser::ParseParameterList(FUnrealFunctionDefinitionInfo& FunctionDe
 			RequireSymbol(TEXT(','), TEXT("Delegate definitions require a , between the parameter type and parameter name"));
 		}
 
-		FProperty* Prop = GetVarNameAndDim(FunctionDef, Property, VariableCategory);
+		FUnrealPropertyDefinitionInfo& PropDef = GetVarNameAndDim(FunctionDef, Property, VariableCategory);
+		FProperty* Prop = PropDef.GetProperty();
 
 		Function->NumParms++;
 
@@ -6990,7 +6402,7 @@ void FHeaderParser::ParseParameterList(FUnrealFunctionDefinitionInfo& FunctionDe
 
 				for (const FString& InvalidName : InvalidParamNames)
 				{
-					if (Property.Matches(*InvalidName, ESearchCase::IgnoreCase))
+					if (FCString::Stricmp(*PropDef.GetNameCPP(), *InvalidName) == 0)
 					{
 						UE_LOG_ERROR_UHT(TEXT("Paramater name '%s' in function is invalid, '%s' is a reserved name."), *InvalidName, *InvalidName);
 					}
@@ -7042,7 +6454,7 @@ void FHeaderParser::ParseParameterList(FUnrealFunctionDefinitionInfo& FunctionDe
 				if (!MetaData->Contains(KeyName))
 				{
 					FString InnerDefaultValue;
-					const bool bDefaultValueParsed = DefaultValueStringCppFormatToInnerFormat(Prop, DefaultArgText, InnerDefaultValue);
+					const bool bDefaultValueParsed = FPropertyTraits::DefaultValueStringCppFormatToInnerFormat(PropDef, DefaultArgText, InnerDefaultValue);
 					if (!bDefaultValueParsed)
 					{
 						FError::Throwf(TEXT("C++ Default parameter not parsed: %s \"%s\" "), *Prop->GetName(), *DefaultArgText);
@@ -7056,7 +6468,7 @@ void FHeaderParser::ParseParameterList(FUnrealFunctionDefinitionInfo& FunctionDe
 	} while( MatchSymbol(TEXT(',')) );
 	RequireSymbol( TEXT(')'), TEXT("parameter list") );
 }
-UDelegateFunction* FHeaderParser::CompileDelegateDeclaration(const TCHAR* DelegateIdentifier, EDelegateSpecifierAction::Type SpecifierAction)
+FUnrealFunctionDefinitionInfo& FHeaderParser::CompileDelegateDeclaration(const TCHAR* DelegateIdentifier, EDelegateSpecifierAction::Type SpecifierAction)
 {
 	const TCHAR* CurrentScopeName = TEXT("Delegate Declaration");
 
@@ -7135,11 +6547,11 @@ UDelegateFunction* FHeaderParser::CompileDelegateDeclaration(const TCHAR* Delega
 	RequireSymbol(TEXT('('), CurrentScopeName);
 
 	// Parse the return value type
-	FToken ReturnType( CPT_None );
+	FPropertyBase ReturnType( CPT_None );
 
 	if (bHasReturnValue)
 	{
-		GetVarType(GetCurrentScope(), ReturnType, CPF_None, nullptr, EPropertyDeclarationStyle::None, EVariableCategory::Return);
+		GetVarType(GetCurrentScope(), ReturnType, CPF_None, nullptr, nullptr, EPropertyDeclarationStyle::None, EVariableCategory::Return);
 		RequireSymbol(TEXT(','), CurrentScopeName);
 	}
 
@@ -7176,16 +6588,9 @@ UDelegateFunction* FHeaderParser::CompileDelegateDeclaration(const TCHAR* Delega
 		FCString::Strcpy( FuncInfo.Function.Identifier, *Name );
 	}
 
-	UDelegateFunction* DelegateSignatureFunction = (bIsSparse ? CreateDelegateFunction<USparseDelegateFunction>(FuncInfo) : CreateDelegateFunction<UDelegateFunction>(FuncInfo));
-	
-	DelegateSignatureFunction->FunctionFlags |= FuncInfo.FunctionFlags;
-
-	FuncInfo.FunctionReference = DelegateSignatureFunction;
-	FuncInfo.SetFunctionNames();
-	if (FuncInfo.FunctionReference->HasAnyFunctionFlags(FUNC_Delegate) && !GetCurrentScope()->IsFileScope())
-	{
-		GetCurrentClassData().MarkContainsDelegate();
-	}
+	FUnrealFunctionDefinitionInfo& FuncDef = bIsSparse ? CreateDelegateFunction<USparseDelegateFunction>(MoveTemp(FuncInfo)) : CreateDelegateFunction<UDelegateFunction>(MoveTemp(FuncInfo));
+	UDelegateFunction* DelegateSignatureFunction = CastChecked<UDelegateFunction>(FuncDef.GetFunction());
+	FFuncInfo& FuncDefFuncInfo = FuncDef.GetFunctionData();
 
 	// determine whether this function should be 'const'
 	if (bDeclaredConst)
@@ -7216,22 +6621,14 @@ UDelegateFunction* FHeaderParser::CompileDelegateDeclaration(const TCHAR* Delega
 		SDF->DelegateName = DelegateName.Identifier;
 	}
 
-
-	// Save off function identifier in case we need it to error later, since we are moving the FuncInfo
-	TCHAR FuncInfoFunctionIdentifier[NAME_SIZE];
-	FCString::Strncpy(FuncInfoFunctionIdentifier, FuncInfo.Function.Identifier, NAME_SIZE);
-
-	TSharedRef<FUnrealFunctionDefinitionInfo> FuncDef = MakeShareable(new FUnrealFunctionDefinitionInfo(SourceFile, InputLine, FString(FuncInfoFunctionIdentifier), MoveTemp(FuncInfo)));
-	FuncDef->SetObject(FuncDef->GetFunctionData().FunctionReference);
-	GTypeDefinitionInfoMap.Add(FuncDef->GetFunctionData().FunctionReference, FuncDef);
-	GetCurrentScope()->AddType(*FuncDef);
+	FuncDef.SetLineNumber(InputLine);
 
 	// Get parameter list.
 	if (FoundParamCount)
 	{
 		RequireSymbol(TEXT(','), CurrentScopeName);
 
-		ParseParameterList(*FuncDef, /*bExpectCommaBeforeName=*/ true);
+		ParseParameterList(FuncDef, /*bExpectCommaBeforeName=*/ true);
 
 		// Check the expected versus actual number of parameters
 		int32 ParamCount = UE_PTRDIFF_TO_INT32(FoundParamCount - UHTConfig.DelegateParameterCountStrings.GetData()) + 1;
@@ -7247,13 +6644,13 @@ UDelegateFunction* FHeaderParser::CompileDelegateDeclaration(const TCHAR* Delega
 	}
 
 	// The macro line must be set here
-	FuncDef->GetFunctionData().MacroLine = InputLine;
+	FuncDef.GetFunctionData().MacroLine = InputLine;
 
 	// Create the return value property
 	if (bHasReturnValue)
 	{
 		ReturnType.PropertyFlags |= CPF_Parm | CPF_OutParm | CPF_ReturnParm;
-		FProperty* ReturnProp = GetVarNameAndDim(*FuncDef, ReturnType, EVariableCategory::Return);
+		FUnrealPropertyDefinitionInfo& PropDef = GetVarNameAndDim(FuncDef, ReturnType, EVariableCategory::Return);
 
 		DelegateSignatureFunction->NumParms++;
 	}
@@ -7263,104 +6660,30 @@ UDelegateFunction* FHeaderParser::CompileDelegateDeclaration(const TCHAR* Delega
 
 	AddFormattedPrevCommentAsTooltipMetaData(MetaData);
 
-	AddMetaDataToClassData(DelegateSignatureFunction, MoveTemp(MetaData));
-
-	// Optionally consume a semicolon, it's not required for the delegate macro since it contains one internally
-	MatchSemi();
-
-	// Bind the function.
-	DelegateSignatureFunction->Bind();
-
-	// End the nesting
-	PostPopFunctionDeclaration(DelegateSignatureFunction);
-
-	// Don't allow delegate signatures to be redefined.
-	auto FunctionIterator = GetCurrentScope()->GetTypeIterator<FUnrealFunctionDefinitionInfo>();
-	while (FunctionIterator.MoveNext())
-	{
-		FUnrealFunctionDefinitionInfo* TestFuncDef = *FunctionIterator;
-		UFunction* TestFunc = TestFuncDef->GetFunction();
-		if ((TestFunc->GetFName() == DelegateSignatureFunction->GetFName()) && (TestFunc != DelegateSignatureFunction))
-		{
-			FError::Throwf(TEXT("Can't override delegate signature function '%s'"), FuncInfoFunctionIdentifier);
-		}
-	}
-
-	return DelegateSignatureFunction;
-}
-
-// Compares the properties of two functions to see if they have the same signature.
-bool AreFunctionSignaturesEqual(const UFunction* Lhs, const UFunction* Rhs)
-{
-	auto LhsPropIter = TFieldIterator<FProperty>(Lhs);
-	auto RhsPropIter = TFieldIterator<FProperty>(Rhs);
-
-	for (;;)
-	{
-		bool bEndOfLhsFunction = !LhsPropIter;
-		bool bEndOfRhsFunction = !RhsPropIter;
-
-		if (bEndOfLhsFunction != bEndOfRhsFunction)
-		{
-			// The functions have different numbers of parameters
-			return false;
-		}
-
-		if (bEndOfLhsFunction)
-		{
-			// We've compared all the parameters
-			return true;
-		}
-
-		const FProperty* LhsProp = *LhsPropIter;
-		const FProperty* RhsProp = *RhsPropIter;
-
-		FFieldClass* LhsClass = LhsProp->GetClass();
-		FFieldClass* RhsClass = RhsProp->GetClass();
+	AddMetaDataToClassData(FuncDef, MoveTemp(MetaData));
 
-		if (LhsClass != RhsClass)
-		{
-			// The properties have different types
-			return false;
-		}
+	// Optionally consume a semicolon, it's not required for the delegate macro since it contains one internally
+	MatchSemi();
 
-		if (LhsClass == FArrayProperty::StaticClass())
-		{
-			const FArrayProperty* LhsArrayProp = (const FArrayProperty*)LhsProp;
-			const FArrayProperty* RhsArrayProp = (const FArrayProperty*)RhsProp;
+	// Bind the function.
+	DelegateSignatureFunction->Bind();
 
-			if (LhsArrayProp->Inner->GetClass() != RhsArrayProp->Inner->GetClass())
-			{
-				// The properties are arrays of different types
-				return false;
-			}
-		}
-		else if (LhsClass == FMapProperty::StaticClass())
-		{
-			const FMapProperty* LhsMapProp = (const FMapProperty*)LhsProp;
-			const FMapProperty* RhsMapProp = (const FMapProperty*)RhsProp;
+	// End the nesting
+	PostPopFunctionDeclaration(FuncDef);
 
-			if (LhsMapProp->KeyProp->GetClass() != RhsMapProp->KeyProp->GetClass() || LhsMapProp->ValueProp->GetClass() != RhsMapProp->ValueProp->GetClass())
-			{
-				// The properties are maps of different types
-				return false;
-			}
-		}
-		else if (LhsClass == FSetProperty::StaticClass())
+	// Don't allow delegate signatures to be redefined.
+	auto FunctionIterator = GetCurrentScope()->GetTypeIterator<FUnrealFunctionDefinitionInfo>();
+	while (FunctionIterator.MoveNext())
+	{
+		FUnrealFunctionDefinitionInfo* TestFuncDef = *FunctionIterator;
+		UFunction* TestFunc = TestFuncDef->GetFunction();
+		if ((TestFunc->GetFName() == DelegateSignatureFunction->GetFName()) && (TestFunc != DelegateSignatureFunction))
 		{
-			const FSetProperty* LhsSetProp = (const FSetProperty*)LhsProp;
-			const FSetProperty* RhsSetProp = (const FSetProperty*)RhsProp;
-
-			if (LhsSetProp->ElementProp->GetClass() != RhsSetProp->ElementProp->GetClass())
-			{
-			// The properties are sets of different types
-			return false;
-			}
+			FError::Throwf(TEXT("Can't override delegate signature function '%s'"), FuncDefFuncInfo.Function.Identifier);
 		}
-
-		++LhsPropIter;
-		++RhsPropIter;
 	}
+
+	return FuncDef;
 }
 
 /**
@@ -7610,10 +6933,10 @@ void FHeaderParser::CompileFunctionDeclaration()
 	}
 
 	// Get return type.
-	FToken ReturnType( CPT_None );
+	FPropertyBase ReturnType( CPT_None );
 
 	// C++ style functions always have a return value type, even if it's void
-	GetVarType(GetCurrentScope(), ReturnType, CPF_None, nullptr, EPropertyDeclarationStyle::None, EVariableCategory::Return);
+	GetVarType(GetCurrentScope(), ReturnType, CPF_None, nullptr, nullptr, EPropertyDeclarationStyle::None, EVariableCategory::Return);
 	bool bHasReturnValue = ReturnType.Type != CPT_None;
 
 	// Skip whitespaces to get InputPos exactly on beginning of function name.
@@ -7671,30 +6994,18 @@ void FHeaderParser::CompileFunctionDeclaration()
 		}
 	}
 
-	UFunction* TopFunction = CreateFunction(FuncInfo);
-
-	TopFunction->FunctionFlags |= FuncInfo.FunctionFlags;
-
-	FuncInfo.FunctionReference = TopFunction;
-	FuncInfo.SetFunctionNames();
-
-	TSharedRef<FUnrealFunctionDefinitionInfo> FuncDef = MakeShareable(new FUnrealFunctionDefinitionInfo(SourceFile, InputLine, FString(FuncInfo.Function.Identifier), FFuncInfo(FuncInfo)));
-	FuncDef->SetObject(FuncDef->GetFunctionData().FunctionReference);
-	GTypeDefinitionInfoMap.Add(FuncDef->GetFunctionData().FunctionReference, FuncDef);
-	GetCurrentScope()->AddType(*FuncDef);
-	if (FuncInfo.FunctionReference->HasAnyFunctionFlags(FUNC_Delegate))
-	{
-		GetCurrentClassData().MarkContainsDelegate();
-	}
+	FUnrealFunctionDefinitionInfo& FuncDef = CreateFunction(MoveTemp(FuncInfo));
+	UFunction* TopFunction = FuncDef.GetFunction();
+	FFuncInfo& FuncDefFuncInfo = FuncDef.GetFunctionData();
 
 	// Get parameter list.
-	ParseParameterList(*FuncDef, false, &MetaData);
+	ParseParameterList(FuncDef, false, &MetaData);
 
 	// Get return type, if any.
 	if (bHasReturnValue)
 	{
 		ReturnType.PropertyFlags |= CPF_Parm | CPF_OutParm | CPF_ReturnParm;
-		FProperty* ReturnProp = GetVarNameAndDim(*FuncDef, ReturnType, EVariableCategory::Return);
+		FUnrealPropertyDefinitionInfo& PropDef = GetVarNameAndDim(FuncDef, ReturnType, EVariableCategory::Return);
 
 		TopFunction->NumParms++;
 	}
@@ -7703,9 +7014,9 @@ void FHeaderParser::CompileFunctionDeclaration()
 	bool bHasAnyOutputs = bHasReturnValue;
 	if (!bHasAnyOutputs)
 	{
-		for (TFieldIterator<FProperty> It(TopFunction); It; ++It)
+		for (FUnrealPropertyDefinitionInfo* ProperyDef : FuncDef.GetProperties())
 		{
-			FProperty const* const Param = *It;
+			FProperty const* const Param = ProperyDef->GetProperty();
 			if (!(Param->PropertyFlags & CPF_ReturnParm) && (Param->PropertyFlags & CPF_OutParm))
 			{
 				bHasAnyOutputs = true;
@@ -7722,14 +7033,14 @@ void FHeaderParser::CompileFunctionDeclaration()
 	}
 	if (SuperStruct)
 	{
-		if (UFunction* OverriddenFunction = ::FindUField<UFunction>(SuperStruct, FuncInfo.Function.Identifier))
+		if (UFunction* OverriddenFunction = ::FindUField<UFunction>(SuperStruct, FuncDefFuncInfo.Function.Identifier))
 		{
 			// Native function overrides should be done in CPP text, not in a UFUNCTION() declaration (you can't change flags, and it'd otherwise be a burden to keep them identical)
-			UE_LOG_ERROR_UHT(TEXT("%s: Override of UFUNCTION in parent class (%s) cannot have a UFUNCTION() declaration above it; it will use the same parameters as the original declaration."), FuncInfo.Function.Identifier, *OverriddenFunction->GetOuter()->GetName());
+			UE_LOG_ERROR_UHT(TEXT("%s: Override of UFUNCTION in parent class (%s) cannot have a UFUNCTION() declaration above it; it will use the same parameters as the original declaration."), FuncDefFuncInfo.Function.Identifier, *OverriddenFunction->GetOuter()->GetName());
 		}
 	}
 
-	if (!bHasAnyOutputs && (FuncInfo.FunctionFlags & (FUNC_BlueprintPure)))
+	if (!bHasAnyOutputs && (FuncDefFuncInfo.FunctionFlags & (FUNC_BlueprintPure)))
 	{
 		// This bad behavior would be treated as a warning in the Blueprint editor, so when converted assets generates these bad functions
 		// we don't want to prevent compilation:
@@ -7749,7 +7060,7 @@ void FHeaderParser::CompileFunctionDeclaration()
 			//FError::Throwf(TEXT("'const' may only be used for native functions"));
 		}
 
-		FuncInfo.FunctionFlags |= FUNC_Const;
+		FuncDefFuncInfo.FunctionFlags |= FUNC_Const;
 
 		// @todo: the presence of const and one or more outputs does not guarantee that there are
 		// no side effects. On GCC and clang we could use __attribure__((pure)) or __attribute__((const))
@@ -7757,9 +7068,9 @@ void FHeaderParser::CompileFunctionDeclaration()
 		// const identifier to determine purity is not desirable. We should remove the following logic:
 
 		// If its a const BlueprintCallable function with some sort of output and is not being marked as an BlueprintPure=false function, mark it as BlueprintPure as well
-		if ( bHasAnyOutputs && ((FuncInfo.FunctionFlags & FUNC_BlueprintCallable) != 0) && !FuncInfo.bForceBlueprintImpure)
+		if ( bHasAnyOutputs && ((FuncDefFuncInfo.FunctionFlags & FUNC_BlueprintCallable) != 0) && !FuncDefFuncInfo.bForceBlueprintImpure)
 		{
-			FuncInfo.FunctionFlags |= FUNC_BlueprintPure;
+			FuncDefFuncInfo.FunctionFlags |= FUNC_BlueprintPure;
 		}
 	}
 
@@ -7768,7 +7079,7 @@ void FHeaderParser::CompileFunctionDeclaration()
 
 	AddFormattedPrevCommentAsTooltipMetaData(MetaData);
 
-	AddMetaDataToClassData(TopFunction, MoveTemp(MetaData));
+	AddMetaDataToClassData(FuncDef, MoveTemp(MetaData));
 
 	// 'final' and 'override' can appear in any order before an optional '= 0' pure virtual specifier
 	bool bFoundFinal    = MatchIdentifier(TEXT("final"), ESearchCase::CaseSensitive);
@@ -7795,13 +7106,13 @@ void FHeaderParser::CompileFunctionDeclaration()
 	if (bFoundFinal)
 	{
 		// This is a final (prebinding, non-overridable) function
-		FuncInfo.FunctionFlags |= FUNC_Final;
-		FuncInfo.FunctionExportFlags |= FUNCEXPORT_Final;
+		FuncDefFuncInfo.FunctionFlags |= FUNC_Final;
+		FuncDefFuncInfo.FunctionExportFlags |= FUNCEXPORT_Final;
 		if (GetCurrentClass()->HasAnyClassFlags(CLASS_Interface))
 		{
 			FError::Throwf(TEXT("Interface functions cannot be declared 'final'"));
 		}
-		else if (FuncInfo.FunctionFlags & FUNC_BlueprintEvent)
+		else if (FuncDefFuncInfo.FunctionFlags & FUNC_BlueprintEvent)
 		{
 			FError::Throwf(TEXT("Blueprint events cannot be declared 'final'"));
 		}
@@ -7809,8 +7120,7 @@ void FHeaderParser::CompileFunctionDeclaration()
 
 	// Make sure any new flags made it to the function
 	//@TODO: UCREMOVAL: Ideally the flags didn't get copied midway thru parsing the function declaration, and we could avoid this
-	TopFunction->FunctionFlags |= FuncInfo.FunctionFlags;
-	FuncDef->UpdateFunctionData(FuncInfo);
+	TopFunction->FunctionFlags |= FuncDefFuncInfo.FunctionFlags;
 
 	// Bind the function.
 	TopFunction->Bind();
@@ -7833,9 +7143,9 @@ void FHeaderParser::CompileFunctionDeclaration()
 
 	if (TopFunction->FunctionFlags & (FUNC_BlueprintCallable | FUNC_BlueprintEvent))
 	{
-		for (TFieldIterator<FProperty> It(TopFunction); It; ++It)
+		for (FUnrealPropertyDefinitionInfo* ProperyDef : FuncDef.GetProperties())
 		{
-			FProperty const* const Param = *It;
+			FProperty const* const Param = ProperyDef->GetProperty();
 			if (Param->ArrayDim > 1)
 			{
 				FError::Throwf(TEXT("Static array cannot be exposed to blueprint. Function: %s Parameter %s\n"), *TopFunction->GetName(), *Param->GetName());
@@ -7851,7 +7161,7 @@ void FHeaderParser::CompileFunctionDeclaration()
 	}
 
 	// Just declaring a function, so end the nesting.
-	PostPopFunctionDeclaration(TopFunction);
+	PostPopFunctionDeclaration(FuncDef);
 
 	// See what's coming next
 	FToken Token;
@@ -7879,7 +7189,7 @@ void FHeaderParser::CompileFunctionDeclaration()
 	}
 
 	// perform documentation policy tests
-	CheckDocumentationPolicyForFunc(GetCurrentClass(), FuncInfo.FunctionReference);
+	CheckDocumentationPolicyForFunc(GetCurrentClass(), FuncDef);
 }
 
 /** Parses optional metadata text. */
@@ -7991,11 +7301,11 @@ bool FHeaderParser::IsBitfieldProperty(ELayoutMacroType LayoutMacroType)
 	return bIsBitfield;
 }
 
-void FHeaderParser::ValidatePropertyIsDeprecatedIfNecessary(const FPropertyBase& VarProperty, const FToken* OuterPropertyType)
+void FHeaderParser::ValidatePropertyIsDeprecatedIfNecessary(const FPropertyBase& VarProperty, const EPropertyFlags* OuterPropertyFlags)
 {
 	// check to see if we have a FClassProperty using a deprecated class
 	if ( VarProperty.MetaClass != NULL && VarProperty.MetaClass->HasAnyClassFlags(CLASS_Deprecated) && !(VarProperty.PropertyFlags & CPF_Deprecated) &&
-		(OuterPropertyType == NULL || !(OuterPropertyType->PropertyFlags & CPF_Deprecated)) )
+		(OuterPropertyFlags == NULL || !(*OuterPropertyFlags & CPF_Deprecated)) )
 	{
 		UE_LOG_ERROR_UHT(TEXT("Property is using a deprecated class: %s.  Property should be marked deprecated as well."), *VarProperty.MetaClass->GetPathName());
 	}
@@ -8005,7 +7315,7 @@ void FHeaderParser::ValidatePropertyIsDeprecatedIfNecessary(const FPropertyBase&
 	if ( (VarProperty.Type == CPT_ObjectReference || VarProperty.Type == CPT_WeakObjectReference || VarProperty.Type == CPT_LazyObjectReference || VarProperty.Type == CPT_ObjectPtrReference || VarProperty.Type == CPT_SoftObjectReference) && VarProperty.PropertyClass != NULL
 		&&	VarProperty.PropertyClass->HasAnyClassFlags(CLASS_Deprecated)	// and the object class being used has been deprecated
 		&& (VarProperty.PropertyFlags&CPF_Deprecated) == 0					// and this property isn't marked deprecated as well
-		&& (OuterPropertyType == NULL || !(OuterPropertyType->PropertyFlags & CPF_Deprecated)) ) // and this property isn't in an array that was marked deprecated either
+		&& (OuterPropertyFlags == NULL || !(*OuterPropertyFlags & CPF_Deprecated)) ) // and this property isn't in an array that was marked deprecated either
 	{
 		UE_LOG_ERROR_UHT(TEXT("Property is using a deprecated class: %s.  Property should be marked deprecated as well."), *VarProperty.PropertyClass->GetPathName());
 	}
@@ -8033,8 +7343,10 @@ bool FHeaderParser::ValidateScriptStructOkForNet(const FString& OriginStructName
 		}
 	}
 
-	for (FField* ChildProp = InStruct->ChildProperties; ChildProp != nullptr; ChildProp = ChildProp->Next)
+	FUnrealScriptStructDefinitionInfo& StructDef = GTypeDefinitionInfoMap.FindChecked<FUnrealScriptStructDefinitionInfo>(InStruct);
+	for (FUnrealPropertyDefinitionInfo* PropertyDef : StructDef.GetProperties())
 	{
+		FProperty* ChildProp = PropertyDef->GetProperty();
 		if (const FSetProperty* const SetProp = CastField<FSetProperty>(ChildProp))
 		{
 			if (!(SetProp->PropertyFlags & CPF_RepSkip))
@@ -8094,7 +7406,7 @@ struct FExposeOnSpawnValidator
 
 		if (!ProperNativeType && (CPT_Struct == Property.Type) && Property.Struct)
 		{
-			ProperNativeType |= Property.Struct->GetBoolMetaData(NAME_BlueprintType);
+			ProperNativeType |= Property.Struct->GetBoolMetaData(FHeaderParserNames::NAME_BlueprintType);
 		}
 
 		return ProperNativeType;
@@ -8111,7 +7423,7 @@ void FHeaderParser::CompileVariableDeclaration(FUnrealStructDefinitionInfo& Stru
 	FPropertyBase OriginalProperty(CPT_None);
 	FIndexRange TypeRange;
 	ELayoutMacroType LayoutMacroType = ELayoutMacroType::None;
-	GetVarType(&*StructDef.GetScope(), OriginalProperty, DisallowFlags, /*OuterPropertyType=*/ NULL, EPropertyDeclarationStyle::UPROPERTY, EVariableCategory::Member, &TypeRange, &LayoutMacroType);
+	GetVarType(&*StructDef.GetScope(), OriginalProperty, DisallowFlags, /*OuterPropertyType=*/ nullptr, nullptr, EPropertyDeclarationStyle::UPROPERTY, EVariableCategory::Member, &TypeRange, &LayoutMacroType);
 	OriginalProperty.PropertyFlags |= EdFlags;
 
 	FString* Category = OriginalProperty.MetaData.Find(NAME_Category);
@@ -8186,8 +7498,10 @@ void FHeaderParser::CompileVariableDeclaration(FUnrealStructDefinitionInfo& Stru
 	TArray<FProperty*> NewProperties;
 	for (;;)
 	{
-		FToken     Property    = OriginalProperty;
-		FProperty* NewProperty = GetVarNameAndDim(StructDef, Property, EVariableCategory::Member, LayoutMacroType);
+		//ETSTODO
+		FPropertyBase Property = OriginalProperty;
+		FUnrealPropertyDefinitionInfo& NewPropDef = GetVarNameAndDim(StructDef, Property, EVariableCategory::Member, LayoutMacroType);
+		FProperty* NewProperty = NewPropDef.GetProperty();
 
 		// Optionally consume the :1 at the end of a bitfield boolean declaration
 		if (Property.IsBool())
@@ -8232,7 +7546,7 @@ void FHeaderParser::CompileVariableDeclaration(FUnrealStructDefinitionInfo& Stru
 
 		if (NewProperty->HasAnyPropertyFlags(CPF_BlueprintVisible))
 		{
-			if (Struct->IsA<UScriptStruct>() && !Struct->GetBoolMetaDataHierarchical(NAME_BlueprintType))
+			if (Struct->IsA<UScriptStruct>() && !Struct->GetBoolMetaDataHierarchical(FHeaderParserNames::NAME_BlueprintType))
 			{
 				UE_LOG_ERROR_UHT(TEXT("Cannot expose property to blueprints in a struct that is not a BlueprintType. %s.%s"), *Struct->GetName(), *NewProperty->GetName());
 			}
@@ -8351,7 +7665,7 @@ void FHeaderParser::CompileVariableDeclaration(FUnrealStructDefinitionInfo& Stru
 // Compile a statement: Either a declaration or a command.
 // Returns 1 if success, 0 if end of file.
 //
-bool FHeaderParser::CompileStatement(TArray<UDelegateFunction*>& DelegatesToFixup)
+bool FHeaderParser::CompileStatement(TArray<FUnrealFunctionDefinitionInfo*>& DelegatesToFixup)
 {
 	// Get a token and compile it.
 	FToken Token;
@@ -8371,19 +7685,19 @@ bool FHeaderParser::CompileStatement(TArray<UDelegateFunction*>& DelegatesToFixu
 // Compute the function parameter size and save the return offset
 //
 //@TODO: UCREMOVAL: Need to rename ComputeFunctionParametersSize to reflect the additional work it's doing
-void FHeaderParser::ComputeFunctionParametersSize( UClass* Class )
+void FHeaderParser::ComputeFunctionParametersSize(FUnrealClassDefinitionInfo& ClassDef)
 {
 	// Recurse with all child states in this class.
-	for (TFieldIterator<UFunction> FuncIt(Class, EFieldIteratorFlags::ExcludeSuper); FuncIt; ++FuncIt)
+	for (FUnrealFunctionDefinitionInfo* FunctionDef : ClassDef.GetFunctions())
 	{
-		UFunction* ThisFunction = *FuncIt;
+		UFunction* ThisFunction = FunctionDef->GetFunction();
 
 		// Fix up any structs that were used as a parameter in a delegate before being defined
 		if (ThisFunction->HasAnyFunctionFlags(FUNC_Delegate))
 		{
-			for (TFieldIterator<FProperty> It(ThisFunction); It; ++It)
+			for (FUnrealPropertyDefinitionInfo* PropertyDef : FunctionDef->GetProperties())
 			{
-				FProperty* Param = *It;
+				FProperty* Param = PropertyDef->GetProperty();
 				if (FStructProperty* StructProp = CastField<FStructProperty>(Param))
 				{
 					if (StructProp->Struct->StructFlags & STRUCT_HasInstancedReference)
@@ -8398,9 +7712,9 @@ void FHeaderParser::ComputeFunctionParametersSize( UClass* Class )
 		// Compute the function parameter size, propagate some flags to the outer function, and save the return offset
 		// Must be done in a second phase, as StaticLink resets various fields again!
 		ThisFunction->ParmsSize = 0;
-		for (TFieldIterator<FProperty> It(ThisFunction); It; ++It)
+		for (FUnrealPropertyDefinitionInfo* PropertyDef : FunctionDef->GetProperties())
 		{
-			FProperty* Param = *It;
+			FProperty* Param = PropertyDef->GetProperty();
 
 			if (!(Param->PropertyFlags & CPF_ReturnParm) && (Param->PropertyFlags & CPF_OutParm))
 			{
@@ -8470,27 +7784,14 @@ void FHeaderParser::SkipStatements( int32 NestCount, const TCHAR* ErrorTag  )
 //
 // Finalize any script-exposed functions in the specified class
 //
-void FHeaderParser::FinalizeScriptExposedFunctions(UClass* Class)
+void FHeaderParser::FinalizeScriptExposedFunctions(FUnrealClassDefinitionInfo& ClassDef)
 {
 	// Finalize all of the children introduced in this class
-	for (TFieldIterator<UStruct> ChildIt(Class, EFieldIteratorFlags::ExcludeSuper); ChildIt; ++ChildIt)
+	UClass* Class = ClassDef.GetClass();
+	for (FUnrealFunctionDefinitionInfo* FunctionDef : ClassDef.GetFunctions())
 	{
-		UStruct* ChildStruct = *ChildIt;
-
-		if (UFunction* Function = Cast<UFunction>(ChildStruct))
-		{
-			// Add this function to the function map of its parent class
-			Class->AddFunctionToFunctionMap(Function, Function->GetFName());
-		}
-		else if (ChildStruct->IsA(UScriptStruct::StaticClass()))
-		{
-			// Ignore embedded structs
-		}
-		else
-		{
-			UE_LOG_WARNING_UHT(TEXT("Unknown and unexpected child named %s of type %s in %s\n"), *ChildStruct->GetName(), *ChildStruct->GetClass()->GetName(), *Class->GetName());
-			check(false);
-		}
+		UFunction* Function = FunctionDef->GetFunction();
+		Class->AddFunctionToFunctionMap(Function, Function->GetFName());
 	}
 }
 
@@ -8541,7 +7842,7 @@ ECompilationResult::Type FHeaderParser::ParseHeader()
 #endif
 	{
 		// Parse entire program.
-		TArray<UDelegateFunction*> DelegatesToFixup;
+		TArray<FUnrealFunctionDefinitionInfo*> DelegatesToFixup;
 		while (CompileStatement(DelegatesToFixup))
 		{
 			bEmptyFile = false;
@@ -8553,30 +7854,26 @@ ECompilationResult::Type FHeaderParser::ParseHeader()
 
 		PopNest(ENestType::GlobalScope, TEXT("Global scope"));
 
-		auto ScopeTypeIterator = SourceFile.GetScope()->GetTypeIterator();
-		while (ScopeTypeIterator.MoveNext())
+		// now validate all delegate variables declared in the class
 		{
-			FUnrealFieldDefinitionInfo* TypeDef = *ScopeTypeIterator;
-			UField* Type = TypeDef->GetField();
-
-			if (!Type->IsA<UScriptStruct>() && !Type->IsA<UClass>())
+			auto ScopeTypeIterator = SourceFile.GetScope()->GetTypeIterator();
+			while (ScopeTypeIterator.MoveNext())
 			{
-				continue;
+				FUnrealFieldDefinitionInfo* TypeDef = *ScopeTypeIterator;
+				if (TypeDef->AsScriptStruct() != nullptr || TypeDef->AsClass() != nullptr)
+				{
+					TMap<FName, UFunction*> DelegateCache;
+					FixupDelegateProperties(TypeDef->AsStructChecked(), *TypeDef->GetScope(), DelegateCache);
+				}
 			}
-
-			UStruct* Struct = Cast<UStruct>(Type);
-
-			// now validate all delegate variables declared in the class
-			TMap<FName, UFunction*> DelegateCache;
-			FixupDelegateProperties(Struct, *TypeDef->GetScope(), DelegateCache);
 		}
 
 		// Fix up any delegates themselves, if they refer to other delegates
 		{
 			TMap<FName, UFunction*> DelegateCache;
-			for (UDelegateFunction* Delegate : DelegatesToFixup)
+			for (FUnrealFunctionDefinitionInfo* DelegateDef : DelegatesToFixup)
 			{
-				FixupDelegateProperties(Delegate, SourceFile.GetScope().Get(), DelegateCache);
+				FixupDelegateProperties(*DelegateDef, SourceFile.GetScope().Get(), DelegateCache);
 			}
 		}
 
@@ -8613,16 +7910,17 @@ ECompilationResult::Type FHeaderParser::ParseHeader()
 		// First-pass success.
 		Result = ECompilationResult::Succeeded;
 
-		for (const TSharedRef<FUnrealTypeDefinitionInfo>& ClassDef : SourceFile.GetDefinedClasses())
+		for (TSharedRef<FUnrealTypeDefinitionInfo>& TypeDef : SourceFile.GetDefinedClasses())
 		{
-			UClass* Class = ClassDef->AsClassChecked().GetClass();
-			PostParsingClassSetup(Class);
+			FUnrealClassDefinitionInfo& ClassDef = UHTCastChecked<FUnrealClassDefinitionInfo>(TypeDef);
+			UClass* Class = ClassDef.GetClass();
+			PostParsingClassSetup(ClassDef);
 
 			// Clean up and exit.
 			Class->Bind();
 
 			// Finalize functions
-			FinalizeScriptExposedFunctions(Class);
+			FinalizeScriptExposedFunctions(ClassDef);
 
 			bNoExportClassesOnly = bNoExportClassesOnly && Class->HasAnyClassFlags(CLASS_NoExport);
 		}
@@ -9484,8 +8782,9 @@ TSharedRef<FUnrealTypeDefinitionInfo> FHeaderPreParser::ParseEnumDeclaration(FUn
 
 	// Get the enum specifier list
 	FToken EnumToken;
+	FMetaData EnumMetaData;
 	TArray<FPropertySpecifier> SpecifiersFound;
-	ReadSpecifierSetInsideMacro(SpecifiersFound, ErrorMsg, EnumToken.MetaData);
+	ReadSpecifierSetInsideMacro(SpecifiersFound, ErrorMsg, EnumMetaData);
 
 	// Check enum type. This can be global 'enum', 'namespace' or 'enum class' enums.
 	bool bReadEnumName = false;
@@ -9549,8 +8848,9 @@ TSharedRef<FUnrealTypeDefinitionInfo> FHeaderPreParser::ParseStructDeclaration(F
 
 	// Get the structure specifier list
 	FToken StructToken;
+	FMetaData StructMetaData;
 	TArray<FPropertySpecifier> SpecifiersFound;
-	ReadSpecifierSetInsideMacro(SpecifiersFound, ErrorMsg, StructToken.MetaData);
+	ReadSpecifierSetInsideMacro(SpecifiersFound, ErrorMsg, StructMetaData);
 
 	// Consume the struct keyword
 	RequireIdentifier(TEXT("struct"), ESearchCase::CaseSensitive, TEXT("Struct declaration specifier"));
@@ -9591,375 +8891,6 @@ TSharedRef<FUnrealTypeDefinitionInfo> FHeaderPreParser::ParseStructDeclaration(F
 	return StructDef;
 }
 
-//////////////////////////////////////////////////////////////////////////////////////////////////////////////
-
-bool FHeaderParser::DefaultValueStringCppFormatToInnerFormat(const FProperty* Property, const FString& CppForm, FString &OutForm)
-{
-	OutForm = FString();
-	if (!Property || CppForm.IsEmpty())
-	{
-		return false;
-	}
-
-	if (Property->IsA(FClassProperty::StaticClass()) || Property->IsA(FObjectPropertyBase::StaticClass()))
-	{
-		const bool bIsNull = FDefaultValueHelper::Is(CppForm, TEXT("NULL")) || FDefaultValueHelper::Is(CppForm, TEXT("nullptr")) || FDefaultValueHelper::Is(CppForm, TEXT("0"));
-		if (bIsNull)
-		{
-			OutForm = TEXT("None");
-		}
-		return bIsNull; // always return as null is the only the processing we can do for object defaults
-	}
-
-	auto ValidateEnumEntry = [Property, &CppForm](const UEnum* Enum, const FString& EnumValue)
-	{
-		const int32 EnumEntryIndex = Enum->GetIndexByName(*EnumValue);
-		if (EnumEntryIndex == INDEX_NONE)
-		{
-			return false;
-		}
-		if (Enum->HasMetaData(TEXT("Hidden"), EnumEntryIndex))
-		{
-			FError::Throwf(TEXT("Hidden enum entries cannot be used as default values: %s \"%s\" "), *Property->GetName(), *CppForm);
-		}
-		return true;
-	};
-
-	if( !Property->IsA(FStructProperty::StaticClass()) )
-	{
-		if( Property->IsA(FIntProperty::StaticClass()) )
-		{
-			int32 Value;
-			if( FDefaultValueHelper::ParseInt( CppForm, Value) )
-			{
-				OutForm = FString::FromInt(Value);
-			}
-		}
-		else if (Property->IsA(FInt64Property::StaticClass()))
-		{
-			int64 Value;
-			if (FDefaultValueHelper::ParseInt64(CppForm, Value))
-			{
-				OutForm = FString::Printf(TEXT("%lld"), Value);
-			}
-		}
-		else if( Property->IsA(FByteProperty::StaticClass()) )
-		{
-			const UEnum* Enum = CastFieldChecked<const FByteProperty>(const_cast<FProperty*>(Property))->Enum;
-			if( NULL != Enum )
-			{
-				OutForm = FDefaultValueHelper::GetUnqualifiedEnumValue(FDefaultValueHelper::RemoveWhitespaces(CppForm));
-				return ValidateEnumEntry(Enum, OutForm);
-			}
-			int32 Value;
-			if( FDefaultValueHelper::ParseInt( CppForm, Value) )
-			{
-				OutForm = FString::FromInt(Value);
-				return ( 0 <= Value ) && ( 255 >= Value );
-			}
-		}
-		else if( Property->IsA(FEnumProperty::StaticClass()) )
-		{
-			const FEnumProperty* EnumProp = CastFieldChecked<FEnumProperty>(const_cast<FProperty*>(Property));
-			if (const UEnum* Enum = CastFieldChecked<FEnumProperty>(const_cast<FProperty*>(Property))->GetEnum())
-			{
-				OutForm = FDefaultValueHelper::GetUnqualifiedEnumValue(FDefaultValueHelper::RemoveWhitespaces(CppForm));
-				return ValidateEnumEntry(Enum, OutForm);
-			}
-
-			int64 Value;
-			if (FDefaultValueHelper::ParseInt64(CppForm, Value))
-			{
-				OutForm = LexToString(Value);
-				return EnumProp->GetUnderlyingProperty()->CanHoldValue(Value);
-			}
-		}
-		else if( Property->IsA(FFloatProperty::StaticClass()) )
-		{
-			float Value;
-			if( FDefaultValueHelper::ParseFloat( CppForm, Value) )
-			{
-				OutForm = FString::Printf( TEXT("%f"), Value) ;
-			}
-		}
-		else if( Property->IsA(FDoubleProperty::StaticClass()) )
-		{
-			double Value;
-			if( FDefaultValueHelper::ParseDouble( CppForm, Value) )
-			{
-				OutForm = FString::Printf( TEXT("%f"), Value) ;
-			}
-		}
-		else if( Property->IsA(FBoolProperty::StaticClass()) )
-		{
-			if( FDefaultValueHelper::Is(CppForm, TEXT("true")) || 
-				FDefaultValueHelper::Is(CppForm, TEXT("false")) )
-			{
-				OutForm = FDefaultValueHelper::RemoveWhitespaces( CppForm );
-			}
-		}
-		else if( Property->IsA(FNameProperty::StaticClass()) )
-		{
-			if(FDefaultValueHelper::Is( CppForm, TEXT("NAME_None") ))
-			{
-				OutForm = TEXT("None");
-				return true;
-			}
-			return FDefaultValueHelper::StringFromCppString(CppForm, TEXT("FName"), OutForm);
-		}
-		else if( Property->IsA(FTextProperty::StaticClass()) )
-		{
-			// Handle legacy cases of FText::FromString being used as default values
-			// These should be replaced with INVTEXT as FText::FromString can produce inconsistent keys
-			if (FDefaultValueHelper::StringFromCppString(CppForm, TEXT("FText::FromString"), OutForm))
-			{
-				UE_LOG_WARNING_UHT(TEXT("FText::FromString should be replaced with INVTEXT for default parameter values"));
-				return true;
-			}
-
-			// Parse the potential value into an instance
-			FText ParsedText;
-			if (FDefaultValueHelper::Is(CppForm, TEXT("FText()")) || FDefaultValueHelper::Is(CppForm, TEXT("FText::GetEmpty()")))
-			{
-				ParsedText = FText::GetEmpty();
-			}
-			else
-			{
-				static const FString UHTDummyNamespace = TEXT("__UHT_DUMMY_NAMESPACE__");
-
-				if (!FTextStringHelper::ReadFromBuffer(*CppForm, ParsedText, *UHTDummyNamespace, nullptr, /*bRequiresQuotes*/true))
-				{
-					return false;
-				}
-
-				// If the namespace of the parsed text matches the default we gave then this was a LOCTEXT macro which we 
-				// don't allow in default values as they rely on an external macro that is known to C++ but not to UHT
-				// TODO: UHT could parse these if it tracked the current LOCTEXT_NAMESPACE macro as it parsed
-				if (TOptional<FString> ParsedTextNamespace = FTextInspector::GetNamespace(ParsedText))
-				{
-					if (ParsedTextNamespace.GetValue().Equals(UHTDummyNamespace))
-					{
-						FError::Throwf(TEXT("LOCTEXT default parameter values are not supported; use NSLOCTEXT instead: %s \"%s\" "), *Property->GetName(), *CppForm);
-					}
-				}
-			}
-
-			// Normalize the default value from the parsed value
-			FTextStringHelper::WriteToBuffer(OutForm, ParsedText, /*bRequiresQuotes*/false);
-			return true;
-		}
-		else if( Property->IsA(FStrProperty::StaticClass()) )
-		{
-			return FDefaultValueHelper::StringFromCppString(CppForm, TEXT("FString"), OutForm);
-		}
-	}
-	else 
-	{
-		// Cache off the struct types, in case we need them later
-		UPackage* CoreUObjectPackage = UObject::StaticClass()->GetOutermost();
-		static const UScriptStruct* VectorStruct = FindObjectChecked<UScriptStruct>(CoreUObjectPackage, TEXT("Vector"));
-		static const UScriptStruct* Vector2DStruct = FindObjectChecked<UScriptStruct>(CoreUObjectPackage, TEXT("Vector2D"));
-		static const UScriptStruct* RotatorStruct = FindObjectChecked<UScriptStruct>(CoreUObjectPackage, TEXT("Rotator"));
-		static const UScriptStruct* LinearColorStruct = FindObjectChecked<UScriptStruct>(CoreUObjectPackage, TEXT("LinearColor"));
-		static const UScriptStruct* ColorStruct = FindObjectChecked<UScriptStruct>(CoreUObjectPackage, TEXT("Color"));
-
-		const FStructProperty* StructProperty = CastFieldChecked<FStructProperty>(const_cast<FProperty*>(Property));
-		if( StructProperty->Struct == VectorStruct )
-		{
-			FString Parameters;
-			if(FDefaultValueHelper::Is( CppForm, TEXT("FVector::ZeroVector") ))
-			{
-				return true;
-			}
-			else if(FDefaultValueHelper::Is(CppForm, TEXT("FVector::UpVector")))
-			{
-				OutForm = FString::Printf(TEXT("%f,%f,%f"),
-					FVector::UpVector.X, FVector::UpVector.Y, FVector::UpVector.Z);
-			}
-			else if(FDefaultValueHelper::Is(CppForm, TEXT("FVector::ForwardVector")))
-			{
-				OutForm = FString::Printf(TEXT("%f,%f,%f"),
-					FVector::ForwardVector.X, FVector::ForwardVector.Y, FVector::ForwardVector.Z);
-			}
-			else if(FDefaultValueHelper::Is(CppForm, TEXT("FVector::RightVector")))
-			{
-				OutForm = FString::Printf(TEXT("%f,%f,%f"),
-					FVector::RightVector.X, FVector::RightVector.Y, FVector::RightVector.Z);
-			}
-			else if( FDefaultValueHelper::GetParameters(CppForm, TEXT("FVector"), Parameters) )
-			{
-				if( FDefaultValueHelper::Is(Parameters, TEXT("ForceInit")) )
-				{
-					return true;
-				}
-				FVector Vector;
-				float Value;
-				if (FDefaultValueHelper::ParseVector(Parameters, Vector))
-				{
-					OutForm = FString::Printf(TEXT("%f,%f,%f"),
-						Vector.X, Vector.Y, Vector.Z);
-				}
-				else if (FDefaultValueHelper::ParseFloat(Parameters, Value))
-				{
-					OutForm = FString::Printf(TEXT("%f,%f,%f"),
-						Value, Value, Value);
-				}
-			}
-		}
-		else if( StructProperty->Struct == RotatorStruct )
-		{
-			if(FDefaultValueHelper::Is( CppForm, TEXT("FRotator::ZeroRotator") ))
-			{
-				return true;
-			}
-			FString Parameters;
-			if( FDefaultValueHelper::GetParameters(CppForm, TEXT("FRotator"), Parameters) )
-			{
-				if( FDefaultValueHelper::Is(Parameters, TEXT("ForceInit")) )
-				{
-					return true;
-				}
-				FRotator Rotator;
-				if(FDefaultValueHelper::ParseRotator(Parameters, Rotator))
-				{
-					OutForm = FString::Printf(TEXT("%f,%f,%f"),
-						Rotator.Pitch, Rotator.Yaw, Rotator.Roll);
-				}
-			}
-		}
-		else if( StructProperty->Struct == Vector2DStruct )
-		{
-			if(FDefaultValueHelper::Is( CppForm, TEXT("FVector2D::ZeroVector") ))
-			{
-				return true;
-			}
-			if(FDefaultValueHelper::Is(CppForm, TEXT("FVector2D::UnitVector")))
-			{
-				OutForm = FString::Printf(TEXT("(X=%3.3f,Y=%3.3f)"),
-					FVector2D::UnitVector.X, FVector2D::UnitVector.Y);
-			}
-			FString Parameters;
-			if( FDefaultValueHelper::GetParameters(CppForm, TEXT("FVector2D"), Parameters) )
-			{
-				if( FDefaultValueHelper::Is(Parameters, TEXT("ForceInit")) )
-				{
-					return true;
-				}
-				FVector2D Vector2D;
-				if(FDefaultValueHelper::ParseVector2D(Parameters, Vector2D))
-				{
-					OutForm = FString::Printf(TEXT("(X=%3.3f,Y=%3.3f)"),
-						Vector2D.X, Vector2D.Y);
-				}
-			}
-		}
-		else if( StructProperty->Struct == LinearColorStruct )
-		{
-			if( FDefaultValueHelper::Is( CppForm, TEXT("FLinearColor::White") ) )
-			{
-				OutForm = FLinearColor::White.ToString();
-			}
-			else if ( FDefaultValueHelper::Is( CppForm, TEXT("FLinearColor::Gray") ) )
-			{
-				OutForm = FLinearColor::Gray.ToString();
-			}
-			else if ( FDefaultValueHelper::Is( CppForm, TEXT("FLinearColor::Black") ) )
-			{
-				OutForm = FLinearColor::Black.ToString();
-			}
-			else if ( FDefaultValueHelper::Is( CppForm, TEXT("FLinearColor::Transparent") ) )
-			{
-				OutForm = FLinearColor::Transparent.ToString();
-			}
-			else if ( FDefaultValueHelper::Is( CppForm, TEXT("FLinearColor::Red") ) )
-			{
-				OutForm = FLinearColor::Red.ToString();
-			}
-			else if ( FDefaultValueHelper::Is( CppForm, TEXT("FLinearColor::Green") ) )
-			{
-				OutForm = FLinearColor::Green.ToString();
-			}
-			else if ( FDefaultValueHelper::Is( CppForm, TEXT("FLinearColor::Blue") ) )
-			{
-				OutForm = FLinearColor::Blue.ToString();
-			}
-			else if ( FDefaultValueHelper::Is( CppForm, TEXT("FLinearColor::Yellow") ) )
-			{
-				OutForm = FLinearColor::Yellow.ToString();
-			}
-			else
-			{
-				FString Parameters;
-				if( FDefaultValueHelper::GetParameters(CppForm, TEXT("FLinearColor"), Parameters) )
-				{
-					if( FDefaultValueHelper::Is(Parameters, TEXT("ForceInit")) )
-					{
-						return true;
-					}
-					FLinearColor Color;
-					if( FDefaultValueHelper::ParseLinearColor(Parameters, Color) )
-					{
-						OutForm = Color.ToString();
-					}
-				}
-			}
-		}
-		else if( StructProperty->Struct == ColorStruct )
-		{
-			if( FDefaultValueHelper::Is( CppForm, TEXT("FColor::White") ) )
-			{
-				OutForm = FColor::White.ToString();
-			}
-			else if ( FDefaultValueHelper::Is( CppForm, TEXT("FColor::Black") ) )
-			{
-				OutForm = FColor::Black.ToString();
-			}
-			else if ( FDefaultValueHelper::Is( CppForm, TEXT("FColor::Red") ) )
-			{
-				OutForm = FColor::Red.ToString();
-			}
-			else if ( FDefaultValueHelper::Is( CppForm, TEXT("FColor::Green") ) )
-			{
-				OutForm = FColor::Green.ToString();
-			}
-			else if ( FDefaultValueHelper::Is( CppForm, TEXT("FColor::Blue") ) )
-			{
-				OutForm = FColor::Blue.ToString();
-			}
-			else if (FDefaultValueHelper::Is(CppForm, TEXT("FColor::Yellow")))
-			{
-				OutForm = FColor::Yellow.ToString();
-			}
-			else if ( FDefaultValueHelper::Is( CppForm, TEXT("FColor::Cyan") ) )
-			{
-				OutForm = FColor::Cyan.ToString();
-			}
-			else if ( FDefaultValueHelper::Is( CppForm, TEXT("FColor::Magenta") ) )
-			{
-				OutForm = FColor::Magenta.ToString();
-			}
-			else
-			{
-				FString Parameters;
-				if( FDefaultValueHelper::GetParameters(CppForm, TEXT("FColor"), Parameters) )
-				{
-					if( FDefaultValueHelper::Is(Parameters, TEXT("ForceInit")) )
-					{
-						return true;
-					}
-					FColor Color;
-					if( FDefaultValueHelper::ParseColor(Parameters, Color) )
-					{
-						OutForm = Color.ToString();
-					}
-				}
-			}
-		}
-	}
-
-	return !OutForm.IsEmpty();
-}
-
 bool FHeaderParser::TryToMatchConstructorParameterList(FToken Token)
 {
 	FToken PotentialParenthesisToken;
@@ -10156,31 +9087,47 @@ void FHeaderParser::ResetClassData()
 
 void FHeaderParser::PostPopNestClass(FUnrealClassDefinitionInfo& CurrentClassDef)
 {
-	UClass* CurrentClass = CurrentClassDef.GetClass();
-
 	// Validate all the rep notify events here, to make sure they're implemented
-	VerifyPropertyMarkups(CurrentClass);
+	VerifyPropertyMarkups(CurrentClassDef);
 
 	// Iterate over all the interfaces we claim to implement
-	for (FImplementedInterface& Impl : CurrentClass->Interfaces)
+	for (const FUnrealStructDefinitionInfo::FBaseClassInfo& BaseClassInfo : CurrentClassDef.GetBaseClassInfo())
 	{
+
+		// Skip unknown base classes or things that aren't interfaces
+		if (BaseClassInfo.Struct == nullptr)
+		{
+			continue;
+		}
+		FUnrealClassDefinitionInfo* InterfaceDef = UHTCast<FUnrealClassDefinitionInfo>(BaseClassInfo.Struct);
+		if (InterfaceDef == nullptr || !InterfaceDef->IsInterface())
+		{
+			continue;
+		}
+
 		// And their super-classes
-		for (UClass* Interface = Impl.Class; Interface; Interface = Interface->GetSuperClass())
+		for (; InterfaceDef; InterfaceDef = InterfaceDef->GetSuperClass())
 		{
 			// If this interface is a common ancestor, skip it
-			if (CurrentClass->IsChildOf(Interface))
+			if (CurrentClassDef.IsChildOf(*InterfaceDef))
 			{
 				continue;
 			}
 
+			UClass* Interface = InterfaceDef->GetClass();
+
 			// So iterate over all functions this interface declares
-			for (UFunction* InterfaceFunction : TFieldRange<UFunction>(Interface, EFieldIteratorFlags::ExcludeSuper))
+			for (FUnrealFunctionDefinitionInfo* InterfaceFunctionDef : TUHTFieldRange<FUnrealFunctionDefinitionInfo>(*InterfaceDef, EFieldIteratorFlags::ExcludeSuper))
 			{
+				UFunction* InterfaceFunction = InterfaceFunctionDef->GetFunction();
+
 				bool bImplemented = false;
 
 				// And try to find one that matches
-				for (UFunction* ClassFunction : TFieldRange<UFunction>(CurrentClass))
+				for (FUnrealFunctionDefinitionInfo* ClassFunctionDef : TUHTFieldRange<FUnrealFunctionDefinitionInfo>(CurrentClassDef))
 				{
+					UFunction* ClassFunction = ClassFunctionDef->GetFunction();
+
 					if (ClassFunction->GetFName() != InterfaceFunction->GetFName())
 					{
 						continue;
@@ -10199,23 +9146,25 @@ void FHeaderParser::PostPopNestClass(FUnrealClassDefinitionInfo& CurrentClassDef
 					// Making sure all the parameters match up correctly
 					bImplemented = true;
 
-					if (ClassFunction->NumParms != InterfaceFunction->NumParms)
+					if (ClassFunctionDef->GetProperties().Num() != InterfaceFunctionDef->GetProperties().Num())
 					{
-						FError::Throwf(TEXT("Implementation of function '%s' conflicts with interface '%s' - different number of parameters (%i/%i)"), *InterfaceFunction->GetName(), *Interface->GetName(), ClassFunction->NumParms, InterfaceFunction->NumParms);
+						FError::Throwf(TEXT("Implementation of function '%s' conflicts with interface '%s' - different number of parameters (%i/%i)"), *InterfaceFunction->GetName(), *Interface->GetName(), ClassFunctionDef->GetProperties().Num(), InterfaceFunctionDef->GetProperties().Num());
 					}
 
-					int32 Count = 0;
-					for (TFieldIterator<FProperty> It1(InterfaceFunction), It2(ClassFunction); Count < ClassFunction->NumParms; ++It1, ++It2, Count++)
+					for (int32 Index = 0, End = ClassFunctionDef->GetProperties().Num(); Index != End; ++Index)
 					{
-						if (!FPropertyBase(*It1).MatchesType(FPropertyBase(*It2), 1))
+						FUnrealPropertyDefinitionInfo* InterfaceParamDef = InterfaceFunctionDef->GetProperties()[Index];
+						FUnrealPropertyDefinitionInfo* ClassParamDef = ClassFunctionDef->GetProperties()[Index];
+						if (!InterfaceParamDef->MatchesType(*ClassParamDef, true))
 						{
-							if (It1->PropertyFlags & CPF_ReturnParm)
+							FProperty* Property = InterfaceParamDef->GetProperty();
+							if (Property->HasAnyPropertyFlags(CPF_ReturnParm))
 							{
 								FError::Throwf(TEXT("Implementation of function '%s' conflicts only by return type with interface '%s'"), *InterfaceFunction->GetName(), *Interface->GetName());
 							}
 							else
 							{
-								FError::Throwf(TEXT("Implementation of function '%s' conflicts with interface '%s' - parameter %i '%s'"), *InterfaceFunction->GetName(), *Interface->GetName(), Count, *It1->GetName());
+								FError::Throwf(TEXT("Implementation of function '%s' conflicts with interface '%s' - parameter %i '%s'"), *InterfaceFunction->GetName(), *Interface->GetName(), End, *Property->GetName());
 							}
 						}
 					}
@@ -10240,25 +9189,24 @@ void FHeaderParser::PostPopNestClass(FUnrealClassDefinitionInfo& CurrentClassDef
 	}
 }
 
-void FHeaderParser::PostPopFunctionDeclaration(UFunction* PoppedFunction)
+void FHeaderParser::PostPopFunctionDeclaration(FUnrealFunctionDefinitionInfo& PoppedFunctionDef)
 {
 	//@TODO: UCREMOVAL: Move this code to occur at delegate var declaration, and force delegates to be declared before variables that use them
-	if (!GetCurrentScope()->IsFileScope() && GetCurrentClassData().ContainsDelegates())
+	if (!GetCurrentScope()->IsFileScope() && GetCurrentClassDef().ContainsDelegates())
 	{
 		// now validate all delegate variables declared in the class
 		TMap<FName, UFunction*> DelegateCache;
-		FixupDelegateProperties(PoppedFunction, *GetCurrentScope(), DelegateCache);
+		FixupDelegateProperties(PoppedFunctionDef, *GetCurrentScope(), DelegateCache);
 	}
 }
 
 void FHeaderParser::PostPopNestInterface(FUnrealClassDefinitionInfo& CurrentInterfaceDef)
 {
-	FStructMetaData& StructData = CurrentInterfaceDef.GetStructMetaData();
-	if (StructData.ContainsDelegates())
+	if (CurrentInterfaceDef.ContainsDelegates())
 	{
 		UClass* CurrentInterface = CurrentInterfaceDef.GetClass();
 		TMap<FName, UFunction*> DelegateCache;
-		FixupDelegateProperties(CurrentInterface, FScope::GetTypeScope(ExactCast<UClass>(CurrentInterface)).Get(), DelegateCache);
+		FixupDelegateProperties(CurrentInterfaceDef, *CurrentInterfaceDef.GetScope(), DelegateCache);
 	}
 }
 
@@ -10345,10 +9293,11 @@ void FHeaderParser::CheckDocumentationPolicyForEnum(UEnum* Enum, const TMap<FNam
 	}
 }
 
-void FHeaderParser::CheckDocumentationPolicyForStruct(UStruct* Struct)
+void FHeaderParser::CheckDocumentationPolicyForStruct(FUnrealStructDefinitionInfo& StructDef)
 {
 	SCOPE_SECONDS_COUNTER_UHT(DocumentationPolicy);
 
+	UStruct* Struct = StructDef.GetStruct();
 	check(Struct != nullptr);
 
 	FDocumentationPolicy DocumentationPolicy = GetDocumentationPolicyForStruct(Struct);
@@ -10369,8 +9318,9 @@ void FHeaderParser::CheckDocumentationPolicyForStruct(UStruct* Struct)
 	if (DocumentationPolicy.bMemberToolTipsRequired)
 	{
 		TMap<FString, FName> ToolTipToPropertyName;
-		for (FProperty* Property : TFieldRange<FProperty>(Struct, EFieldIteratorFlags::ExcludeSuper))
+		for (FUnrealPropertyDefinitionInfo* PropertyDef : StructDef.GetProperties())
 		{
+			FProperty* Property = PropertyDef->GetProperty();
 			FString ToolTip = Property->GetToolTipText().ToString();
 			if (ToolTip.IsEmpty() || ToolTip.Equals(Property->GetDisplayNameText().ToString()))
 			{
@@ -10388,8 +9338,9 @@ void FHeaderParser::CheckDocumentationPolicyForStruct(UStruct* Struct)
 
 	if (DocumentationPolicy.bFloatRangesRequired)
 	{
-		for (FProperty* Property : TFieldRange<FProperty>(Struct, EFieldIteratorFlags::ExcludeSuper))
+		for (FUnrealPropertyDefinitionInfo* PropertyDef : StructDef.GetProperties())
 		{
+			FProperty* Property = PropertyDef->GetProperty();
 			if(DoesCPPTypeRequireDocumentation(Property->GetCPPType()))
 			{
 				const FString& UIMin = Property->GetMetaData(NAME_UIMin);
@@ -10406,12 +9357,13 @@ void FHeaderParser::CheckDocumentationPolicyForStruct(UStruct* Struct)
 	// also compare all tooltips to see if they are unique
 	if (DocumentationPolicy.bFunctionToolTipsRequired)
 	{
-		UClass* Class = Cast<UClass>(Struct);
-		if (Class != nullptr)
+		if (FUnrealClassDefinitionInfo* ClassDef = UHTCast<FUnrealClassDefinitionInfo>(StructDef))
 		{
+			UClass* Class = ClassDef->GetClass();
 			TMap<FString, FName> ToolTipToFunc;
-			for (UFunction* Func : TFieldRange<UFunction>(Class, EFieldIteratorFlags::ExcludeSuper))
+			for (FUnrealFunctionDefinitionInfo* FunctionDef : ClassDef->GetFunctions())
 			{
+				UFunction* Func = FunctionDef->GetFunction();
 				FString ToolTip = Func->GetToolTipText().ToString();
 				if (ToolTip.IsEmpty())
 				{
@@ -10435,10 +9387,11 @@ bool FHeaderParser::DoesCPPTypeRequireDocumentation(const FString& CPPType)
 }
 
 // Validates the documentation for a given method
-void FHeaderParser::CheckDocumentationPolicyForFunc(UClass* Class, UFunction* Func)
+void FHeaderParser::CheckDocumentationPolicyForFunc(UClass* Class, FUnrealFunctionDefinitionInfo& FunctionDef)
 {
 	SCOPE_SECONDS_COUNTER_UHT(DocumentationPolicy);
 
+	UFunction* Func = FunctionDef.GetFunction();
 	check(Class != nullptr);
 	check(Func != nullptr);
 
@@ -10477,8 +9430,9 @@ void FHeaderParser::CheckDocumentationPolicyForFunc(UClass* Class, UFunction* Fu
 		{
 			// ensure each parameter has a tooltip
 			TSet<FName> ExistingFields;
-			for (FProperty* Property : TFieldRange<FProperty>(Func))
+			for (FUnrealPropertyDefinitionInfo* PropertyDef : FunctionDef.GetProperties())
 			{
+				FProperty* Property = PropertyDef->GetProperty();
 				FName ParamName = Property->GetFName();
 				if (ParamName == NAME_ReturnValue)
 				{
@@ -10595,7 +9549,7 @@ void FHeaderParser::ConditionalLogPointerUsage(EPointerMemberBehavior PointerMem
 }
 
 template <class TFunctionType>
-TFunctionType* CreateFunctionImpl(const FFuncInfo& FuncInfo, UObject* Outer, FScope* CurrentScope)
+FUnrealFunctionDefinitionInfo& CreateFunctionImpl(FUnrealSourceFile& SourceFile, int32 InputLine, FFuncInfo&& FuncInfo, FUnrealObjectDefinitionInfo& OuterDef, FScope* CurrentScope)
 {
 	// Allocate local property frame, push nesting level and verify
 	// uniqueness at this scope level.
@@ -10611,31 +9565,40 @@ TFunctionType* CreateFunctionImpl(const FFuncInfo& FuncInfo, UObject* Outer, FSc
 		}
 	}
 
-	TFunctionType* Function = new(EC_InternalUseOnlyConstructor, Outer, FuncInfo.Function.Identifier, RF_Public) TFunctionType(FObjectInitializer(), nullptr);
+	TFunctionType* Function = new(EC_InternalUseOnlyConstructor, OuterDef.GetObject(), FuncInfo.Function.Identifier, RF_Public) TFunctionType(FObjectInitializer(), nullptr);
 	Function->ReturnValueOffset = MAX_uint16;
 	Function->FirstPropertyToInit = nullptr;
+	Function->FunctionFlags |= FuncInfo.FunctionFlags;
+	FuncInfo.FunctionReference = Function;
+	FuncInfo.SetFunctionNames();
+
+	TSharedRef<FUnrealFunctionDefinitionInfo> FuncDefRef = MakeShared<FUnrealFunctionDefinitionInfo>(SourceFile, InputLine, FString(FuncInfo.Function.Identifier), OuterDef, MoveTemp(FuncInfo));
+	FUnrealFunctionDefinitionInfo& FuncDef = *FuncDefRef;
+	FuncDef.SetObject(Function);
+	GTypeDefinitionInfoMap.Add(Function, MoveTemp(FuncDefRef));
+	CurrentScope->AddType(FuncDef);
 
 	if (!CurrentScope->IsFileScope())
 	{
-		UStruct* Struct = ((FStructScope*)CurrentScope)->GetStruct();
+		FUnrealStructDefinitionInfo& StructDef = ((FStructScope*)CurrentScope)->GetStructDef();
+		UStruct* Struct = StructDef.GetStruct();
 
+		StructDef.AddFunction(FuncDef);
 		Function->Next = Struct->Children;
 		Struct->Children = Function;
 	}
 
-	return Function;
+	return FuncDef;
 }
 
-UFunction* FHeaderParser::CreateFunction(const FFuncInfo &FuncInfo) const
+FUnrealFunctionDefinitionInfo& FHeaderParser::CreateFunction(FFuncInfo&& FuncInfo) const
 {
-	return CreateFunctionImpl<UFunction>(FuncInfo, GetCurrentClass(), GetCurrentScope());
+	return CreateFunctionImpl<UFunction>(SourceFile, InputLine, MoveTemp(FuncInfo), GetCurrentClassDef(), GetCurrentScope());
 }
 
 template<class T>
-UDelegateFunction* FHeaderParser::CreateDelegateFunction(const FFuncInfo &FuncInfo) const
+FUnrealFunctionDefinitionInfo& FHeaderParser::CreateDelegateFunction(FFuncInfo&& FuncInfo) const
 {
-	FFileScope* CurrentFileScope = GetCurrentFileScope();
-	FUnrealSourceFile* LocSourceFile = CurrentFileScope ? CurrentFileScope->GetSourceFile() : nullptr;
-	UObject* CurrentPackage = LocSourceFile ? LocSourceFile->GetPackage() : nullptr;
-	return CreateFunctionImpl<T>(FuncInfo, IsInAClass() ? (UObject*)GetCurrentClass() : CurrentPackage, GetCurrentScope());
+	FUnrealObjectDefinitionInfo& OuterDef = IsInAClass() ? static_cast<FUnrealObjectDefinitionInfo&>(GetCurrentClassDef()) : SourceFile.GetPackageDef();
+	return CreateFunctionImpl<T>(SourceFile, InputLine, MoveTemp(FuncInfo), OuterDef, GetCurrentScope());
 }
