Added new PropertyTypes which in the long run will be where different property types are implemented.

diff --git a/Engine/Source/Programs/UnrealHeaderTool/Private/HeaderParser.h b/Engine/Source/Programs/UnrealHeaderTool/Private/HeaderParser.h
--- a/Engine/Source/Programs/UnrealHeaderTool/Private/HeaderParser.h
+++ b/Engine/Source/Programs/UnrealHeaderTool/Private/HeaderParser.h
@@ -430,7 +430,7 @@ class FHeaderParser : public FBaseParser, public FContextSupplier
 {
 public:
 	// Compute the function parameter size and save the return offset
-	static void ComputeFunctionParametersSize(UClass* InClass);
+	static void ComputeFunctionParametersSize(FUnrealClassDefinitionInfo& ClassDef);
 
 	// Performs a preliminary parse of the text in the specified buffer, pulling out:
 	//   Class name and parent class name
@@ -449,17 +449,6 @@ public:
 	 */
 	static bool ClassNameHasValidPrefix(const FString& InNameToCheck, const FString& OriginalClassName);
 
-	/**
-	 * Transforms CPP-formated string containing default value, to inner formated string
-	 * If it cannot be transformed empty string is returned.
-	 *
-	 * @param Property The property that owns the default value.
-	 * @param CppForm A CPP-formated string.
-	 * @param out InnerForm Inner formated string
-	 * @return true on success, false otherwise.
-	 */
-	static bool DefaultValueStringCppFormatToInnerFormat(const FProperty* Property, const FString& CppForm, FString &InnerForm);
-
 	/**
 	 * Parse the given header.
 	 * 
@@ -670,8 +659,8 @@ public:
 	*
 	* @return	a pointer to a UField with a name matching InIdentifier, or NULL if it wasn't found
 	*/
-	static UField* FindField( UStruct* InScope, const TCHAR* InIdentifier, bool bIncludeParents=true, UClass* FieldClass=UField::StaticClass(), const TCHAR* Thing=nullptr );
-	static FField* FindProperty(UStruct* InScope, const TCHAR* InIdentifier, bool bIncludeParents = true, FFieldClass* FieldClass = FField::StaticClass(), const TCHAR* Thing = nullptr);
+	static UField* FindField(FUnrealStructDefinitionInfo& InScope, const TCHAR* InIdentifier, bool bIncludeParents=true, UClass* FieldClass=UField::StaticClass(), const TCHAR* Thing=nullptr );
+	static FField* FindProperty(FUnrealStructDefinitionInfo& InScope, const TCHAR* InIdentifier, bool bIncludeParents = true, FFieldClass* FieldClass = FField::StaticClass(), const TCHAR* Thing = nullptr);
 
 	// Checks ToValidate to make sure that its associated sparse class data struct, if one exists, is a valid structure to use for storing sparse class data.
 	static void CheckSparseClassData(const UStruct* ToValidate);
@@ -711,10 +700,10 @@ protected:
 	 */
 	ECompilationResult::Type ParseHeader();
 	void CompileDirective();
-	void FinalizeScriptExposedFunctions(UClass* Class);
+	void FinalizeScriptExposedFunctions(FUnrealClassDefinitionInfo& ClassDef);
 	UEnum* CompileEnum();
 	UScriptStruct* CompileStructDeclaration();
-	bool CompileDeclaration(TArray<UDelegateFunction*>& DelegatesToFixup, FToken& Token);
+	bool CompileDeclaration(TArray<FUnrealFunctionDefinitionInfo*>& DelegatesToFixup, FToken& Token);
 
 	/** Skip C++ (noexport) declaration. */
 	bool SkipDeclaration(FToken& Token);
@@ -740,26 +729,26 @@ protected:
 	/**
 	 * Create new function object based on given info structure.
 	 */
-	UFunction* CreateFunction(const FFuncInfo &FuncInfo) const;
+	FUnrealFunctionDefinitionInfo& CreateFunction(FFuncInfo&& FuncInfo) const;
 
 	/**
 	 * Create new delegate function object based on given info structure.
 	 */
 	template<typename T>
-	UDelegateFunction* CreateDelegateFunction(const FFuncInfo &FuncInfo) const;	
+	FUnrealFunctionDefinitionInfo& CreateDelegateFunction(FFuncInfo&& FuncInfo) const;
 
 	UClass* CompileClassDeclaration();
-	UDelegateFunction* CompileDelegateDeclaration(const TCHAR* DelegateIdentifier, EDelegateSpecifierAction::Type SpecifierAction = EDelegateSpecifierAction::DontParse);
+	FUnrealFunctionDefinitionInfo& CompileDelegateDeclaration(const TCHAR* DelegateIdentifier, EDelegateSpecifierAction::Type SpecifierAction = EDelegateSpecifierAction::DontParse);
 	void CompileFunctionDeclaration();
 	void CompileVariableDeclaration (FUnrealStructDefinitionInfo& StructDef);
 	void CompileInterfaceDeclaration();
 	void CompileRigVMMethodDeclaration(UStruct* Struct);
-	void ParseRigVMMethodParameters(UStruct* Struct);
+	void ParseRigVMMethodParameters(FUnrealStructDefinitionInfo& StructDef);
 
 	FClass* ParseInterfaceNameDeclaration(FString& DeclaredInterfaceName, FString& RequiredAPIMacroIfPresent);
 	bool TryParseIInterfaceClass();
 
-	bool CompileStatement(TArray<UDelegateFunction*>& DelegatesToFixup);
+	bool CompileStatement(TArray<FUnrealFunctionDefinitionInfo*>& DelegatesToFixup);
 
 	// Checks to see if a particular kind of command is allowed on this nesting level.
 	bool IsAllowedInThisNesting(ENestAllowFlags AllowFlags);
@@ -769,8 +758,6 @@ protected:
 	// nesting level.
 	void CheckAllow(const TCHAR* Thing, ENestAllowFlags AllowFlags);
 
-	UStruct* GetSuperScope( UStruct* CurrentScope, const FName& SearchName );	
-
 	void SkipStatements( int32 SubCount, const TCHAR* ErrorTag );
 
 	/**
@@ -789,6 +776,7 @@ protected:
 		FPropertyBase&                  VarProperty,
 		EPropertyFlags                  Disallow,
 		const FToken*                   OuterPropertyType,
+		const EPropertyFlags*			OuterPropertyFlags,
 		EPropertyDeclarationStyle::Type PropertyDeclarationStyle,
 		EVariableCategory::Type         VariableCategory,
 		FIndexRange*                    ParsedVarIndexRange = nullptr,
@@ -801,11 +789,11 @@ protected:
 	 * @param	VarProperty			type and propertyflag info for the new property (inout)
 	 * @param   VariableCategory	what kind of variable is being created
 	 *
-	 * @return	a pointer to the new FProperty if successful, or NULL if there was no property to parse
+	 * @return	a reference to the new property if successful, or NULL if there was no property to parse
 	 */
-	FProperty* GetVarNameAndDim(
+	FUnrealPropertyDefinitionInfo& GetVarNameAndDim(
 		FUnrealStructDefinitionInfo& ParentStruct,
-		FToken& VarProperty,
+		FPropertyBase& VarProperty,
 		EVariableCategory::Type VariableCategory,
 		ELayoutMacroType LayoutMacroType = ELayoutMacroType::None);
 	
@@ -819,11 +807,6 @@ protected:
 	 */
 	bool AllowReferenceToClass(UStruct* Scope, UClass* CheckClass) const;
 
-	/**
-	 * @return	true if Scope has FProperty objects in its list of fields
-	 */
-	static bool HasMemberProperties( const UStruct* Scope );
-
 	/**
 	 * Parses optional metadata text.
 	 *
@@ -869,7 +852,7 @@ protected:
 	 *
 	 * @param PoppedFunction Function that have just been popped.
 	 */
-	void PostPopFunctionDeclaration(UFunction* PoppedFunction);
+	void PostPopFunctionDeclaration(FUnrealFunctionDefinitionInfo& PoppedFunctionDef);
 
 	/**
 	 * Tasks that need to be done after popping interface definition
@@ -897,13 +880,13 @@ protected:
 	 * @param	Scope				the current scope
 	 * @param	DelegateCache		cached map of delegates that have already been found; used for faster lookup.
 	 */
-	void FixupDelegateProperties(UStruct* ValidationScope, FScope& Scope, TMap<FName, UFunction*>& DelegateCache);
+	void FixupDelegateProperties(FUnrealStructDefinitionInfo& StructDef, FScope& Scope, TMap<FName, UFunction*>& DelegateCache);
 
 	// Retry functions.
 	void InitScriptLocation( FScriptLocation& Retry );
 	void ReturnToLocation( const FScriptLocation& Retry, bool Binary=1, bool Text=1 );
 
-	static void ValidatePropertyIsDeprecatedIfNecessary(const FPropertyBase& VarProperty, const FToken* OuterPropertyType);
+	static void ValidatePropertyIsDeprecatedIfNecessary(const FPropertyBase& VarProperty, const EPropertyFlags* OuterPropertyFlags);
 
 	// Cache of ScriptStructs that have been validated for Net Replication and RPC
 	TSet<UScriptStruct*> ScriptStructsValidForNet;
@@ -942,16 +925,16 @@ private:
 	void CompileVersionDeclaration(UStruct* Struct);
 
 	// Verifies that all specified class's UProperties with function associations have valid targets
-	void VerifyPropertyMarkups( UClass* TargetClass );
+	void VerifyPropertyMarkups(FUnrealClassDefinitionInfo& TargetClassDef);
 
 	// Verifies the target function meets the criteria for a blueprint property getter
-	void VerifyBlueprintPropertyGetter(FProperty* Property, UFunction* TargetFunction);
+	void VerifyBlueprintPropertyGetter(FProperty* Property, FUnrealFunctionDefinitionInfo* TargetFuncDef);
 
 	// Verifies the target function meets the criteria for a blueprint property setter
-	void VerifyBlueprintPropertySetter(FProperty* Property, UFunction* TargetFunction);
+	void VerifyBlueprintPropertySetter(FProperty* Property, FUnrealFunctionDefinitionInfo* TargetFuncDef);
 
 	// Verifies the target function meets the criteria for a replication notify callback
-	void VerifyRepNotifyCallback(FProperty* Property, UFunction* TargetFunction);
+	void VerifyRepNotifyCallback(FProperty* Property, FUnrealFunctionDefinitionInfo* TargetFuncDef);
 
 	// Constructs the policy from a string
 	static FDocumentationPolicy GetDocumentationPolicyFromName(const FString& PolicyName);
@@ -969,10 +952,10 @@ private:
 	void CheckDocumentationPolicyForEnum(UEnum* Enum, const TMap<FName, FString>& MetaData, const TArray<TMap<FName, FString>>& Entries);
 
 	// Validates the documentation for a given struct
-	void CheckDocumentationPolicyForStruct(UStruct* Struct);
+	void CheckDocumentationPolicyForStruct(FUnrealStructDefinitionInfo& StructDef);
 
 	// Validates the documentation for a given method
-	void CheckDocumentationPolicyForFunc(UClass* Class, UFunction* Func);
+	void CheckDocumentationPolicyForFunc(UClass* Class, FUnrealFunctionDefinitionInfo& FunctionDef);
 
 	// Checks if a valid range has been found on the provided metadata
 	bool CheckUIMinMaxRangeFromMetaData(const FString& UIMin, const FString& UIMax);
diff --git a/Engine/Source/Programs/UnrealHeaderTool/Private/NativeClassExporter.h b/Engine/Source/Programs/UnrealHeaderTool/Private/NativeClassExporter.h
--- a/Engine/Source/Programs/UnrealHeaderTool/Private/NativeClassExporter.h
+++ b/Engine/Source/Programs/UnrealHeaderTool/Private/NativeClassExporter.h
@@ -19,6 +19,7 @@ class UDelegateFunction;
 class FStructMetaData;
 class FOutputDevice;
 class FUnrealPackageDefinitionInfo;
+class FUnrealPropertyDefinitionInfo;
 struct FFuncInfo;
 struct FGeneratedFileInfo;
 
@@ -74,14 +75,14 @@ ENUM_CLASS_FLAGS(EExportClassOutFlags);
 
 struct FPropertyNamePointerPair
 {
-	FPropertyNamePointerPair(FString InName, FProperty* InProp)
+	FPropertyNamePointerPair(FString InName, FUnrealPropertyDefinitionInfo& InPropDef)
 		: Name(MoveTemp(InName))
-		, Prop(InProp)
+		, PropDef(&InPropDef)
 	{
 	}
 
 	FString Name;
-	FProperty* Prop;
+	FUnrealPropertyDefinitionInfo* PropDef;
 };
 
 FString CreateUTF8LiteralString(const FString& Str);
@@ -327,7 +328,7 @@ private:
 	 * @param	Struct			UStruct to export properties
 	 * @param	TextIndent		Current text indentation
 	 */
-	static void ExportProperties(FOutputDevice& Out, UStruct* Struct, int32 TextIndent);
+	static void ExportProperties(FOutputDevice& Out, FUnrealStructDefinitionInfo& StructDef, int32 TextIndent);
 
 	/** Return the name of the singleton function */
 	static const FString& GetPackageSingletonName(UPackage* Item, TSet<FString>* UniqueCrossModuleReferences);
@@ -358,7 +359,7 @@ private:
 	/**
 	 * Export functions used to find and call C++ or script implementation of a script function in the interface 
 	 */
-	void ExportInterfaceCallFunctions(FOutputDevice& OutCpp, FUHTStringBuilder& Out, FReferenceGatherers& OutReferenceGatherers, const TArray<UFunction*>& CallbackFunctions, const TCHAR* ClassName) const;
+	void ExportInterfaceCallFunctions(FOutputDevice& OutCpp, FUHTStringBuilder& Out, FReferenceGatherers& OutReferenceGatherers, const TArray<FUnrealFunctionDefinitionInfo*>& CallbackFunctions, const TCHAR* ClassName) const;
 
 private:
 
@@ -379,7 +380,7 @@ private:
 		FOutputDevice&           OutputGetter,
 		FOutputDevice&           OutDeclarations,
 		FReferenceGatherers&     OutReferenceGatherers, 
-		FClass*                  Class,
+		FUnrealClassDefinitionInfo&			 ClassDef,
 		const FUnrealSourceFile& SourceFile,
 		EExportClassOutFlags&    OutFlags
 	) const;
@@ -431,7 +432,7 @@ private:
 	 * @param	Structs		the struct to export
 	 * @param	TextIndent	the current indentation of the header exporter
 	 */
-	static void ExportMirrorsForNoexportStruct(FOutputDevice& Out, UScriptStruct* Struct, int32 TextIndent);
+	static void ExportMirrorsForNoexportStruct(FOutputDevice& Out, FUnrealScriptStructDefinitionInfo& ScriptStructDef, int32 TextIndent);
 
 	/**heade
 	 * Exports the parameter struct declarations for the list of functions specified
@@ -439,7 +440,7 @@ private:
 	 * @param	Function	the function that (may) have parameters which need to be exported
 	 * @return	true		if the structure generated is not completely empty
 	 */
-	static bool WillExportEventParms( UFunction* Function );
+	static bool WillExportEventParms(FUnrealFunctionDefinitionInfo& FunctionDef);
 
 	/**
 	 * Exports C++ type declarations for delegates
@@ -448,7 +449,7 @@ private:
 	 * @param	SourceFile			Source file of the delegate.
 	 * @param	DelegateFunctions	the functions that have parameters which need to be exported
 	 */
-	void ExportDelegateDeclaration(FOutputDevice& Out, FReferenceGatherers& OutReferenceGatherers, const FUnrealSourceFile& SourceFile, UFunction* Function) const;
+	void ExportDelegateDeclaration(FOutputDevice& Out, FReferenceGatherers& OutReferenceGatherers, const FUnrealSourceFile& SourceFile, FUnrealFunctionDefinitionInfo& FunctionDef) const;
 
 	/**
 	 * Exports C++ type definitions for delegates
@@ -457,7 +458,7 @@ private:
 	 * @param	SourceFile			Source file of the delegate.
 	 * @param	DelegateFunctions	the functions that have parameters which need to be exported
 	 */
-	void ExportDelegateDefinition(FOutputDevice& Out, FReferenceGatherers& OutReferenceGatherers, const FUnrealSourceFile& SourceFile, UFunction* Function) const;
+	void ExportDelegateDefinition(FOutputDevice& Out, FReferenceGatherers& OutReferenceGatherers, const FUnrealSourceFile& SourceFile, FUnrealFunctionDefinitionInfo& FunctionDef) const;
 
 	/**
 	 * Exports the parameter struct declarations for the given function.
@@ -467,7 +468,7 @@ private:
 	 * @param	Indent				number of spaces to put before each line
 	 * @param	bOutputConstructor	If true, output a constructor for the param struct
 	 */
-	static void ExportEventParm(FUHTStringBuilder& Out, TSet<FString>& PropertyFwd, UFunction* Function, int32 Indent, bool bOutputConstructor, EExportingState ExportingState);
+	static void ExportEventParm(FUHTStringBuilder& Out, TSet<FString>& PropertyFwd, FUnrealFunctionDefinitionInfo& FunctionDef, int32 Indent, bool bOutputConstructor, EExportingState ExportingState);
 
 	/**
 	* Move the temp header files into the .h files
@@ -486,7 +487,7 @@ private:
 	 *
 	 * @return	the intrinsic null value for the property (0 for ints, TEXT("") for strings, etc.)
 	 */
-	static FString GetNullParameterValue( FProperty* Prop, bool bInitializer = false );
+	static FString GetNullParameterValue(FUnrealPropertyDefinitionInfo& PropertyDef, bool bInitializer = false );
 
 	/**
 	 * Exports a native function prototype
@@ -500,7 +501,8 @@ private:
 	static void ExportNativeFunctionHeader(
 		FOutputDevice&                   Out,
 		TSet<FString>&                   OutFwdDecls,
-		const FFuncInfo&                 FunctionData,
+		FUnrealFunctionDefinitionInfo&                 FunctionDef,
+		const FFuncInfo&				 FunctionData,
 		EExportFunctionType::Type        FunctionType,
 		EExportFunctionHeaderStyle::Type FunctionHeaderStyle,
 		const TCHAR*                     ExtraParam,
@@ -516,7 +518,7 @@ private:
 	* @param	ValidatePosition		Position in source file of _Validate function for function described by FunctionData.
 	* @param	SourceFile				Currently analyzed source file.
 	*/
-	void CheckRPCFunctions(FReferenceGatherers& OutReferenceGatherers, const FFuncInfo& FunctionData, const FString& ClassName, int32 ImplementationPosition, int32 ValidatePosition, const FUnrealSourceFile& SourceFile) const;
+	void CheckRPCFunctions(FReferenceGatherers& OutReferenceGatherers, FUnrealFunctionDefinitionInfo& FunctionDef, const FString& ClassName, int32 ImplementationPosition, int32 ValidatePosition, const FUnrealSourceFile& SourceFile) const;
 
 	/**
 	 * Exports the native stubs for the list of functions specified
@@ -525,7 +527,7 @@ private:
 	 * @param Class			class
 	 * @param ClassData		class data
 	 */
-	void ExportNativeFunctions(FOutputDevice& OutGeneratedHeaderText, FOutputDevice& OutGeneratedCPPText, FOutputDevice& OutMacroCalls, FOutputDevice& OutNoPureDeclsMacroCalls, FReferenceGatherers& OutReferenceGatherers, const FUnrealSourceFile& SourceFile, UClass* Class, FStructMetaData& StructData) const;
+	void ExportNativeFunctions(FOutputDevice& OutGeneratedHeaderText, FOutputDevice& OutGeneratedCPPText, FOutputDevice& OutMacroCalls, FOutputDevice& OutNoPureDeclsMacroCalls, FReferenceGatherers& OutReferenceGatherers, const FUnrealSourceFile& SourceFile, FUnrealClassDefinitionInfo& ClassDef) const;
 
 	/**
 	 * Export the actual internals to a standard thunk function
@@ -536,10 +538,10 @@ private:
 	 * @param Parameters list of parameters in the function
 	 * @param Return return parameter for the function
 	 */
-	void ExportFunctionThunk(FUHTStringBuilder& RPCWrappers, FReferenceGatherers& OutReferenceGatherers, UFunction* Function, const FFuncInfo& FunctionData, const TArray<FProperty*>& Parameters, FProperty* Return) const;
+	void ExportFunctionThunk(FUHTStringBuilder& RPCWrappers, FReferenceGatherers& OutReferenceGatherers, FUnrealFunctionDefinitionInfo& FunctionDef, const TArray<FUnrealPropertyDefinitionInfo*>& ParameterDefs, FUnrealPropertyDefinitionInfo* ReturnDef) const;
 
 	/** Exports the native function registration code for the given class. */
-	static void ExportNatives(FOutputDevice& Out, FClass* Class);
+	static void ExportNatives(FOutputDevice& Out, FUnrealClassDefinitionInfo& ClassDef);
 
 	/**
 	 * Exports generated singleton functions for UObjects that used to be stored in .u files.
@@ -549,7 +551,7 @@ private:
 	 * @param	Class			Class to export
 	 * @param	OutFriendText	(Output parameter) Friend text
 	 */
-	void ExportNativeGeneratedInitCode(FOutputDevice& Out, FOutputDevice& OutDeclarations, FReferenceGatherers& OutReferenceGatherers, const FUnrealSourceFile& SourceFile, FClass* Class, FUHTStringBuilder& OutFriendText) const;
+	void ExportNativeGeneratedInitCode(FOutputDevice& Out, FOutputDevice& OutDeclarations, FReferenceGatherers& OutReferenceGatherers, const FUnrealSourceFile& SourceFile, FUnrealClassDefinitionInfo& ClassDef, FUHTStringBuilder& OutFriendText) const;
 
 	/**
 	 * Export given function.
@@ -558,7 +560,7 @@ private:
 	 * @param	Function		Given function.
 	 * @param	bIsNoExport		Is in NoExport class.
 	 */
-	void ExportFunction(FOutputDevice& Out, FReferenceGatherers& OutReferenceGatherers, const FUnrealSourceFile& SourceFile, UFunction* Function, bool bIsNoExport) const;
+	void ExportFunction(FOutputDevice& Out, FReferenceGatherers& OutReferenceGatherers, const FUnrealSourceFile& SourceFile, FUnrealFunctionDefinitionInfo& FunctionDef, bool bIsNoExport) const;
 
 	/**
 	 * Exports a generated singleton function to setup the package for compiled-in classes.
@@ -592,7 +594,7 @@ private:
 	 * @param	DeclSpaces		String of spaces to use as an indent for the declaration
 	 * @param	Spaces			String of spaces to use as an indent
 	**/
-	void OutputProperty(FOutputDevice& DeclOut, FOutputDevice& Out, FReferenceGatherers& OutReferenceGatherers, const TCHAR* Scope, TArray<FPropertyNamePointerPair>& PropertyNamesAndPointers, FProperty* Prop, const TCHAR* DeclSpaces, const TCHAR* Spaces) const;
+	void OutputProperty(FOutputDevice& DeclOut, FOutputDevice& Out, FReferenceGatherers& OutReferenceGatherers, const TCHAR* Scope, TArray<FPropertyNamePointerPair>& PropertyNamesAndPointers, FUnrealPropertyDefinitionInfo& PropertyDef, const TCHAR* DeclSpaces, const TCHAR* Spaces) const;
 
 	/**
 	 * Function to output the C++ code necessary to set up a property, including an array property and its inner, array dimensions, etc.
@@ -606,7 +608,7 @@ private:
 	 * @param	Spaces			String of spaces to use as an indent
 	 * @param	SourceStruct	Structure that the property offset is relative to
 	**/
-	void PropertyNew(FOutputDevice& DeclOut, FOutputDevice& Out, FReferenceGatherers& OutReferenceGatherers, FProperty* Prop, const TCHAR* OffsetStr, const TCHAR* Name, const TCHAR* DeclSpaces, const TCHAR* Spaces, const TCHAR* SourceStruct = nullptr) const;
+	void PropertyNew(FOutputDevice& DeclOut, FOutputDevice& Out, FReferenceGatherers& OutReferenceGatherers, FUnrealPropertyDefinitionInfo& PropertyDef, const TCHAR* OffsetStr, const TCHAR* Name, const TCHAR* DeclSpaces, const TCHAR* Spaces, const TCHAR* SourceStruct = nullptr) const;
 
 	/**
 	 * Function to generate the property tag
@@ -633,7 +635,7 @@ private:
 		FOutputDevice&            OutGeneratedHeaderText,
 		FOutputDevice&            Out,
 		TSet<FString>&            OutFwdDecls,
-		const TArray<UFunction*>& CallbackFunctions,
+		const TArray<FUnrealFunctionDefinitionInfo*>& CallbackFunctions,
 		const TCHAR*              CallbackWrappersMacroName,
 		EExportCallbackType       ExportCallbackType,
 		const TCHAR*              APIString
@@ -682,7 +684,7 @@ private:
 	 * @param	ConstructorsMacroPrefix					Prefix for constructors macro.
 	 * @param	Class									Class for which to export macros.
 	 */
-	static void ExportConstructorsMacros(FOutputDevice& OutGeneratedHeaderText, FOutputDevice& VTableOut, FOutputDevice& StandardUObjectConstructorsMacroCall, FOutputDevice& EnhancedUObjectConstructorsMacroCall, const FString& ConstructorsMacroPrefix, FClass* Class, const TCHAR* APIArg);
+	static void ExportConstructorsMacros(FOutputDevice& OutGeneratedHeaderText, FOutputDevice& VTableOut, FOutputDevice& StandardUObjectConstructorsMacroCall, FOutputDevice& EnhancedUObjectConstructorsMacroCall, const FString& ConstructorsMacroPrefix, FUnrealClassDefinitionInfo& ClassDef, const TCHAR* APIArg);
 
 	/**
 	 * Gets string with function return type.
@@ -690,7 +692,7 @@ private:
 	 * @param Function Function to get return type of.
 	 * @return FString with function return type.
 	 */
-	static FString GetFunctionReturnString(UFunction* Function, FReferenceGatherers& OutReferenceGatherers);
+	static FString GetFunctionReturnString(FUnrealFunctionDefinitionInfo& FunctionDef, FReferenceGatherers& OutReferenceGatherers);
 
 	/**
 	* Gets string with function parameters (with names).
@@ -698,7 +700,7 @@ private:
 	* @param Function Function to get parameters of.
 	* @return FString with function parameters.
 	*/
-	static FString GetFunctionParameterString(UFunction* Function, FReferenceGatherers& OutReferenceGatherers);
+	static FString GetFunctionParameterString(FUnrealFunctionDefinitionInfo& FunctionDef, FReferenceGatherers& OutReferenceGatherers);
 
 	/**
 	 * Checks if function is missing "virtual" specifier.
diff --git a/Engine/Source/Programs/UnrealHeaderTool/Private/ParserClass.cpp b/Engine/Source/Programs/UnrealHeaderTool/Private/ParserClass.cpp
--- a/Engine/Source/Programs/UnrealHeaderTool/Private/ParserClass.cpp
+++ b/Engine/Source/Programs/UnrealHeaderTool/Private/ParserClass.cpp
@@ -49,26 +49,6 @@ FString FClass::GetNameWithPrefix(EEnforceInterfacePrefix::Type EnforceInterface
 	return FString::Printf(TEXT("%s%s"), Prefix, *GetName());
 }
 
-FClass* FClass::GetSuperClass() const
-{
-	return static_cast<FClass*>(static_cast<const UClass*>(this)->GetSuperClass());
-}
-
-FClass* FClass::GetClassWithin() const
-{
-	return (FClass*)ClassWithin;
-}
-
-TArray<FClass*> FClass::GetInterfaceTypes() const
-{
-	TArray<FClass*> Result;
-	for (const FImplementedInterface& i : Interfaces)
-	{
-		Result.Add((FClass*)i.Class);
-	}
-	return Result;
-}
-
 void FClass::GetHideCategories(TArray<FString>& OutHideCategories) const
 {
 	if (HasMetaData(FHeaderParserNames::NAME_HideCategories))
diff --git a/Engine/Source/Programs/UnrealHeaderTool/Private/ParserClass.h b/Engine/Source/Programs/UnrealHeaderTool/Private/ParserClass.h
--- a/Engine/Source/Programs/UnrealHeaderTool/Private/ParserClass.h
+++ b/Engine/Source/Programs/UnrealHeaderTool/Private/ParserClass.h
@@ -32,22 +32,6 @@ public:
 	 */
 	FString GetNameWithPrefix(EEnforceInterfacePrefix::Type EnforceInterfacePrefix = EEnforceInterfacePrefix::None) const;
 
-	/**
-	 * Returns the super class of this class, or NULL if there is no superclass.
-	 *
-	 * @return The super class of this class.
-	 */
-	FClass* GetSuperClass() const;
-
-	/**
-	 * Returns the 'within' class of this class.
-	 *
-	 * @return The 'within' class of this class.
-	 */
-	FClass* GetClassWithin() const;
-
-	TArray<FClass*> GetInterfaceTypes() const;
-
 	void GetHideCategories(TArray<FString>& OutHideCategories) const;
 	void GetShowCategories(TArray<FString>& OutShowCategories) const;
 	void GetSparseClassDataTypes(TArray<FString>& OutSparseClassDataTypes) const;
diff --git a/Engine/Source/Programs/UnrealHeaderTool/Private/ParserHelper.cpp b/Engine/Source/Programs/UnrealHeaderTool/Private/ParserHelper.cpp
--- a/Engine/Source/Programs/UnrealHeaderTool/Private/ParserHelper.cpp
+++ b/Engine/Source/Programs/UnrealHeaderTool/Private/ParserHelper.cpp
@@ -8,33 +8,6 @@
 #include "UnrealTypeDefinitionInfo.h"
 #include "ClassMaps.h"
 
-/////////////////////////////////////////////////////
-// FStructMetaData
-
-void FStructMetaData::AddProperty(FUnrealPropertyDefinitionInfo& PropertyDef)
-{
-	// update the optimization flags
-	if (!bContainsDelegates)
-	{
-		FProperty* Prop = PropertyDef.GetProperty();
-		if (Prop->IsA(FDelegateProperty::StaticClass()) || Prop->IsA(FMulticastDelegateProperty::StaticClass()))
-		{
-			bContainsDelegates = true;
-		}
-		else
-		{
-			FArrayProperty* ArrayProp = CastField<FArrayProperty>(Prop);
-			if (ArrayProp != NULL)
-			{
-				if (ArrayProp->Inner->IsA(FDelegateProperty::StaticClass()) || ArrayProp->Inner->IsA(FMulticastDelegateProperty::StaticClass()))
-				{
-					bContainsDelegates = true;
-				}
-			}
-		}
-	}
-}
-
 /////////////////////////////////////////////////////
 // FPropertyBase
 
@@ -89,14 +62,10 @@ const TCHAR* FPropertyBase::GetPropertyTypeText( EPropertyType Type )
  */
 FToken& FToken::operator=(const FToken& Other)
 {
-	FPropertyBase::operator=((FPropertyBase&)Other);
-
+	ConstantType = Other.ConstantType;
 	TokenType = Other.TokenType;
-	TokenName = Other.TokenName;
-	bTokenNameInitialized = Other.bTokenNameInitialized;
 	StartPos = Other.StartPos;
 	StartLine = Other.StartLine;
-	TokenProperty = Other.TokenProperty;
 
 	FCString::Strncpy(Identifier, Other.Identifier, NAME_SIZE);
 	FMemory::Memcpy(String, Other.String, sizeof(String));
@@ -106,14 +75,10 @@ FToken& FToken::operator=(const FToken& Other)
 
 FToken& FToken::operator=(FToken&& Other)
 {
-	FPropertyBase::operator=(MoveTemp(Other));
-
+	ConstantType = Other.ConstantType;
 	TokenType = Other.TokenType;
-	TokenName = Other.TokenName;
-	bTokenNameInitialized = Other.bTokenNameInitialized;
 	StartPos = Other.StartPos;
 	StartLine = Other.StartLine;
-	TokenProperty = Other.TokenProperty;
 
 	FCString::Strncpy(Identifier, Other.Identifier, NAME_SIZE);
 	FMemory::Memcpy(String, Other.String, sizeof(String));
diff --git a/Engine/Source/Programs/UnrealHeaderTool/Private/ParserHelper.h b/Engine/Source/Programs/UnrealHeaderTool/Private/ParserHelper.h
--- a/Engine/Source/Programs/UnrealHeaderTool/Private/ParserHelper.h
+++ b/Engine/Source/Programs/UnrealHeaderTool/Private/ParserHelper.h
@@ -105,6 +105,17 @@ enum EAccessSpecifier
 #define CASE_TEXT(txt) case txt: return TEXT(#txt)
 #endif
 
+inline bool IsBool(EPropertyType Type)
+{
+	return Type == CPT_Bool || Type == CPT_Bool8 || Type == CPT_Bool16 || Type == CPT_Bool32 || Type == CPT_Bool64;
+}
+
+inline bool IsObject(EPropertyType Type)
+{
+	return Type == CPT_ObjectReference || Type == CPT_Interface || Type == CPT_WeakObjectReference || Type == CPT_LazyObjectReference || Type == CPT_ObjectPtrReference || Type == CPT_SoftObjectReference;
+}
+
+
 /**
  * Basic information describing a type.
  */
@@ -250,226 +261,6 @@ public:
 	{
 	}
 
-	explicit FPropertyBase(FProperty* Property)
-	: PropertyExportFlags(PROPEXPORT_Public)
-	, DelegateName       (NAME_None)
-	, DelegateSignatureOwnerClass(nullptr)
-	, RepNotifyName      (NAME_None)
-	, IntType            (EIntType::None)
-	{
-		checkSlow(Property);
-
-		EArrayType::Type ArrType         = EArrayType::None;
-		EPropertyFlags   PropagateFlags  = CPF_None;
-		FFieldClass*     ClassOfProperty = Property->GetClass();
-
-		if( ClassOfProperty==FArrayProperty::StaticClass() )
-		{
-			ArrType = EArrayType::Dynamic;
-
-			// if we're an array, save up Parm flags so we can propagate them.
-			// below the array will be assigned the inner property flags. This allows propagation of Parm flags (out, optional..)
-			PropagateFlags = Property->PropertyFlags & CPF_ParmFlags;
-			Property = CastFieldChecked<FArrayProperty>(Property)->Inner;
-			ClassOfProperty = Property->GetClass();
-		}
-		else if (ClassOfProperty == FSetProperty::StaticClass())
-		{
-			ArrType = EArrayType::Set;
-
-			PropagateFlags = Property->PropertyFlags & CPF_ParmFlags;
-			Property = CastFieldChecked<FSetProperty>(Property)->ElementProp;
-			ClassOfProperty = Property->GetClass();
-		}
-		else if (ClassOfProperty == FMapProperty::StaticClass())
-		{
-			PropagateFlags = Property->PropertyFlags & CPF_ParmFlags;
-			MapKeyProp = MakeShared<FPropertyBase>(CastFieldChecked<FMapProperty>(Property)->KeyProp);
-			Property = CastFieldChecked<FMapProperty>(Property)->ValueProp;
-			ClassOfProperty = Property->GetClass();
-		}
-
-		if( ClassOfProperty==FByteProperty::StaticClass() )
-		{
-			*this = FPropertyBase(CPT_Byte);
-			Enum = CastField<FByteProperty>(Property)->Enum;
-			IntType = EIntType::Sized;
-		}
-		else if( ClassOfProperty==FEnumProperty::StaticClass() )
-		{
-			FEnumProperty* EnumProp = CastField<FEnumProperty>(Property);
-			FNumericProperty* UnderlyingProp = EnumProp->GetUnderlyingProperty();
-
-			*this = FPropertyBase(
-				  UnderlyingProp->IsA<FInt8Property>()   ? CPT_Int8
-				: UnderlyingProp->IsA<FInt16Property>()  ? CPT_Int16
-				: UnderlyingProp->IsA<FIntProperty>()    ? CPT_Int
-				: UnderlyingProp->IsA<FInt64Property>()  ? CPT_Int64
-				: UnderlyingProp->IsA<FByteProperty>()   ? CPT_Byte
-				: UnderlyingProp->IsA<FUInt16Property>() ? CPT_UInt16
-				: UnderlyingProp->IsA<FUInt32Property>() ? CPT_UInt32
-				: UnderlyingProp->IsA<FUInt64Property>() ? CPT_UInt64
-				:                                          CPT_None
-			);
-			check(this->Type != CPT_None);
-			Enum = EnumProp->Enum;
-			IntType = EIntType::Sized;
-		}
-		else if( ClassOfProperty==FInt8Property::StaticClass() )
-		{
-			*this = FPropertyBase(CPT_Int8);
-			IntType = EIntType::Sized;
-		}
-		else if( ClassOfProperty==FInt16Property::StaticClass() )
-		{
-			*this = FPropertyBase(CPT_Int16);
-			IntType = EIntType::Sized;
-		}
-		else if( ClassOfProperty==FIntProperty::StaticClass() )
-		{
-			*this = FPropertyBase(CPT_Int);
-			IntType = EIntType::Sized;
-		}
-		else if( ClassOfProperty==FInt64Property::StaticClass() )
-		{
-			*this = FPropertyBase(CPT_Int64);
-			IntType = EIntType::Sized;
-		}
-		else if( ClassOfProperty==FUInt16Property::StaticClass() )
-		{
-			*this = FPropertyBase(CPT_UInt16);
-			IntType = EIntType::Sized;
-		}
-		else if( ClassOfProperty==FUInt32Property::StaticClass() )
-		{
-			*this = FPropertyBase(CPT_UInt32);
-			IntType = EIntType::Sized;
-		}
-		else if( ClassOfProperty==FUInt64Property::StaticClass() )
-		{
-			*this = FPropertyBase(CPT_UInt64);
-			IntType = EIntType::Sized;
-		}
-		else if( ClassOfProperty==FBoolProperty::StaticClass() )
-		{
-			FBoolProperty* BoolProperty =  CastField<FBoolProperty>(Property);
-			if (BoolProperty->IsNativeBool())
-			{
-				*this = FPropertyBase(CPT_Bool);
-			}
-			else
-			{
-				switch( BoolProperty->ElementSize )
-				{
-				case sizeof(uint8):
-					*this = FPropertyBase(CPT_Bool8);
-					break;
-				case sizeof(uint16):
-					*this = FPropertyBase(CPT_Bool16);
-					break;
-				case sizeof(uint32):
-					*this = FPropertyBase(CPT_Bool32);
-					break;
-				case sizeof(uint64):
-					*this = FPropertyBase(CPT_Bool64);
-					break;
-				}
-			}
-		}
-		else if( ClassOfProperty==FFloatProperty::StaticClass() )
-		{
-			*this = FPropertyBase(CPT_Float);
-		}
-		else if( ClassOfProperty==FDoubleProperty::StaticClass() )
-		{
-			*this = FPropertyBase(CPT_Double);
-		}
-		else if( ClassOfProperty==FClassProperty::StaticClass() )
-		{
-			*this = FPropertyBase(CPT_ObjectReference);
-			PropertyClass = CastField<FClassProperty>(Property)->PropertyClass;
-			MetaClass = CastField<FClassProperty>(Property)->MetaClass;
-		}
-		else if( ClassOfProperty==FObjectProperty::StaticClass() )
-		{
-			*this = FPropertyBase(CPT_ObjectReference);
-			PropertyClass = CastField<FObjectProperty>(Property)->PropertyClass;
-		}
-		else if( ClassOfProperty==FWeakObjectProperty::StaticClass() )
-		{
-			*this = FPropertyBase(CPT_WeakObjectReference);
-			PropertyClass = CastField<FWeakObjectProperty>(Property)->PropertyClass;
-		}
-		else if( ClassOfProperty==FLazyObjectProperty::StaticClass() )
-		{
-			*this = FPropertyBase(CPT_LazyObjectReference);
-			PropertyClass = CastField<FLazyObjectProperty>(Property)->PropertyClass;
-		}
-		else if( ClassOfProperty==FObjectPtrProperty::StaticClass() )
-		{
-			*this = FPropertyBase(CPT_ObjectPtrReference);
-			PropertyClass = CastField<FObjectPtrProperty>(Property)->PropertyClass;
-		}
-		else if( ClassOfProperty==FSoftClassProperty::StaticClass() )
-		{
-			*this = FPropertyBase(CPT_SoftObjectReference);
-			PropertyClass = CastField<FSoftClassProperty>(Property)->PropertyClass;
-			MetaClass = CastField<FSoftClassProperty>(Property)->MetaClass;
-		}
-		else if( ClassOfProperty==FSoftObjectProperty::StaticClass() )
-		{
-			*this = FPropertyBase(CPT_SoftObjectReference);
-			PropertyClass = CastField<FSoftObjectProperty>(Property)->PropertyClass;
-		}
-		else if( ClassOfProperty==FNameProperty::StaticClass() )
-		{
-			*this = FPropertyBase(CPT_Name);
-		}
-		else if( ClassOfProperty==FStrProperty::StaticClass() )
-		{
-			*this = FPropertyBase(CPT_String);
-		}
-		else if( ClassOfProperty==FTextProperty::StaticClass() )
-		{
-			*this = FPropertyBase(CPT_Text);
-		}
-		else if( ClassOfProperty==FStructProperty::StaticClass() )
-		{
-			*this = FPropertyBase(CPT_Struct);
-			Struct = CastField<FStructProperty>(Property)->Struct;
-		}
-		else if( ClassOfProperty==FDelegateProperty::StaticClass() )
-		{
-			*this = FPropertyBase(CPT_Delegate);
-			Function = CastField<FDelegateProperty>(Property)->SignatureFunction;
-		}
-		else if( ClassOfProperty==FMulticastDelegateProperty::StaticClass() )
-		{
-			*this = FPropertyBase(CPT_MulticastDelegate);
-			// @todo delegate: Any other setup for calling multi-cast delegates from script needed?
-			Function = CastField<FMulticastDelegateProperty>(Property)->SignatureFunction;
-		}
-		else if ( ClassOfProperty==FInterfaceProperty::StaticClass() )
-		{
-			*this = FPropertyBase(CPT_Interface);
-			PropertyClass = CastField<FInterfaceProperty>(Property)->InterfaceClass;
-		}
-		else if (ClassOfProperty == FFieldPathProperty::StaticClass())
-		{
-			*this = FPropertyBase(CPT_FieldPath);
-			PropertyPathClass = CastField<FFieldPathProperty>(Property)->PropertyClass;
-		}
-		else
-		{
-			UE_LOG(LogCompile, Fatal, TEXT("Unknown property type '%s'"), *Property->GetFullName() );
-		}
-		ArrayType            = ArrType;
-		PropertyFlags        = Property->PropertyFlags | PropagateFlags;
-		ImpliedPropertyFlags = CPF_None;
-		RefQualifier         = ERefQualifier::None;
-		PointerType          = EPointerType::None;
-	}
-
 	explicit FPropertyBase(FFieldClass* InPropertyClass, EPropertyType InType)
 		: Type(InType)
 		, ArrayType(EArrayType::None)
@@ -496,7 +287,12 @@ public:
 	 */
 	bool IsObject() const
 	{
-		return Type == CPT_ObjectReference || Type == CPT_Interface || Type == CPT_WeakObjectReference || Type == CPT_LazyObjectReference || Type == CPT_ObjectPtrReference || Type == CPT_SoftObjectReference;
+		return ::IsObject(Type);
+	}
+
+	bool IsBool() const
+	{
+		return ::IsBool(Type);
 	}
 
 	bool IsContainer() const
@@ -691,19 +487,6 @@ public:
 			return true;
 		}
 	}
-
-	FString Describe()
-	{
-		return FString::Printf(
-			TEXT("Type:%s  Flags:%lli  ImpliedFlags:%lli  Enum:%s  PropertyClass:%s  Struct:%s  Function:%s  MetaClass:%s"),
-			GetPropertyTypeText(Type), PropertyFlags, ImpliedPropertyFlags,
-			Enum!=NULL?*Enum->GetName():TEXT(""),
-			PropertyClass!=NULL?*PropertyClass->GetName():TEXT("NULL"),
-			Struct!=NULL?*Struct->GetName():TEXT("NULL"),
-			Function!=NULL?*Function->GetName():TEXT("NULL"),
-			MetaClass!=NULL?*MetaClass->GetName():TEXT("NULL")
-			);
-	}
 	//@}
 
 	EIntType GetSizedIntTypeFromPropertyType(EPropertyType PropType)
@@ -748,32 +531,21 @@ enum ETokenType
 //
 // Information about a token that was just parsed.
 //
-class FToken : public FPropertyBase
+class FToken
 {
 public:
 	/** @name Variables */
 	//@{
 	/** Type of token. */
 	ETokenType TokenType;
-
-private:
-	/** Whether TokenName has been looked up */
-	mutable bool bTokenNameInitialized;
-	/** Name of token, lazily initialized */
-	mutable FName TokenName;
-
-public:
 	/** Starting position in script where this token came from. */
 	int32 StartPos;
 	/** Starting line in script. */
 	int32 StartLine;
 	/** Always valid. */
 	TCHAR Identifier[NAME_SIZE];
-	/** property that corresponds to this FToken - null if this Token doesn't correspond to a FProperty */
-	FProperty* TokenProperty;
-
-
-public:
+	/** Type of the constant value */
+	EPropertyType ConstantType = CPT_None;
 
 	union
 	{
@@ -796,7 +568,7 @@ public:
 	{
 		if (TokenType == TOKEN_Const)
 		{
-			switch (Type)
+			switch (ConstantType)
 			{
 				case CPT_Byte:
 					return FString::Printf(TEXT("%u"), Byte);
@@ -829,37 +601,22 @@ public:
 
 	// Constructors.
 	FToken()
-		: FPropertyBase(CPT_None)
-		, TokenProperty(NULL)
 	{
 		InitToken(CPT_None);
 	}
 
 	// copy constructors
 	explicit FToken(EPropertyType InType)
-		: FPropertyBase(InType)
-		, TokenProperty(NULL)
 	{
 		InitToken(InType);
 	}
 
-	// copy constructors
-	FToken(const FPropertyBase& InType)
-		: FPropertyBase(CPT_None)
-		, TokenProperty(NULL)
-	{
-		InitToken( CPT_None );
-		(FPropertyBase&)*this = InType;
-	}
-
 	FToken(const FToken& InToken)
-		: FPropertyBase(CPT_None)
 	{
 		*this = InToken;
 	}
 
 	FToken(FToken&& InToken)
-		: FPropertyBase(CPT_None)
 	{
 		*this = MoveTemp(InToken);
 	}
@@ -867,9 +624,8 @@ public:
 	// Inlines.
 	void InitToken( EPropertyType InType )
 	{
-		(FPropertyBase&)*this = FPropertyBase(InType);
+		ConstantType    = InType;
 		TokenType		= TOKEN_None;
-		TokenName		= NAME_None;
 		StartPos		= 0;
 		StartLine		= 0;
 		*Identifier		= 0;
@@ -883,24 +639,6 @@ public:
 	{
 		return (TokenType==TOKEN_Identifier || TokenType==TOKEN_Symbol) && ((SearchCase == ESearchCase::CaseSensitive) ? !FCString::Strcmp(Identifier, Str) : !FCString::Stricmp(Identifier, Str));
 	}
-	bool IsBool() const
-	{
-		return Type == CPT_Bool || Type == CPT_Bool8 || Type == CPT_Bool16 || Type == CPT_Bool32 || Type == CPT_Bool64;
-	}
-	FName GetTokenName() const
-	{
-		if (!bTokenNameInitialized)
-		{
-			TokenName = FName(Identifier, FNAME_Find);
-			bTokenNameInitialized = true;
-		}
-		return TokenName;
-	}
-	void ClearTokenName()
-	{
-		bTokenNameInitialized = false;
-		TokenName = NAME_None;
-	}
 
 	// Setters.
 	
@@ -909,49 +647,48 @@ public:
 		InitToken(CPT_None);
 		TokenType = TOKEN_Identifier;
 		FCString::Strncpy(Identifier, InString, NAME_SIZE);
-		bTokenNameInitialized = false;
 	}
 
 	void SetConstInt64( int64 InInt64 )
 	{
-		(FPropertyBase&)*this = FPropertyBase(CPT_Int64);
+		ConstantType    = CPT_Int64;
 		Int64			= InInt64;
 		TokenType		= TOKEN_Const;
 	}
 	void SetConstInt( int32 InInt )
 	{
-		(FPropertyBase&)*this = FPropertyBase(CPT_Int);
+		ConstantType    = CPT_Int;
 		Int				= InInt;
 		TokenType		= TOKEN_Const;
 	}
 	void SetConstBool( bool InBool )
 	{
-		(FPropertyBase&)*this = FPropertyBase(CPT_Bool);
+		ConstantType    = CPT_Bool;
 		NativeBool 		= InBool;
 		TokenType		= TOKEN_Const;
 	}
 	void SetConstFloat( float InFloat )
 	{
-		(FPropertyBase&)*this = FPropertyBase(CPT_Float);
+		ConstantType    = CPT_Float;
 		Float			= InFloat;
 		TokenType		= TOKEN_Const;
 	}
 	void SetConstDouble( double InDouble )
 	{
-		(FPropertyBase&)*this = FPropertyBase(CPT_Double);
+		ConstantType    = CPT_Double;
 		Double			= InDouble;
 		TokenType		= TOKEN_Const;
 	}
 	void SetConstName( FName InName )
 	{
-		(FPropertyBase&)*this = FPropertyBase(CPT_Name);
+		ConstantType    = CPT_Name;
 		*(FName *)NameBytes = InName;
 		TokenType		= TOKEN_Const;
 	}
 	void SetConstString( const TCHAR* InString, int32 MaxLength=MAX_STRING_CONST_SIZE )
 	{
 		check(MaxLength>0);
-		(FPropertyBase&)*this = FPropertyBase(CPT_String);
+		ConstantType = CPT_String;
 		if( InString != String )
 		{
 			FCString::Strncpy( String, InString, MaxLength );
@@ -961,7 +698,7 @@ public:
 	void SetConstChar(TCHAR InChar)
 	{
 		//@TODO: Treating this like a string for now, nothing consumes it
-		(FPropertyBase&)*this = FPropertyBase(CPT_String);
+		ConstantType = CPT_String;
 		String[0] = InChar;
 		String[1] = 0;
 		TokenType = TOKEN_Const;
@@ -971,27 +708,27 @@ public:
 	// Getters.
 	bool GetConstInt( int32& I ) const
 	{
-		if( TokenType==TOKEN_Const && Type==CPT_Int64 )
+		if( TokenType==TOKEN_Const && ConstantType ==CPT_Int64 )
 		{
 			I = (int32)Int64;
 			return 1;
 		}
-		else if( TokenType==TOKEN_Const && Type==CPT_Int )
+		else if( TokenType==TOKEN_Const && ConstantType ==CPT_Int )
 		{
 			I = Int;
 			return 1;
 		}
-		else if( TokenType==TOKEN_Const && Type==CPT_Byte )
+		else if( TokenType==TOKEN_Const && ConstantType ==CPT_Byte )
 		{
 			I = Byte;
 			return 1;
 		}
-		else if( TokenType==TOKEN_Const && Type==CPT_Float && Float==FMath::TruncToInt(Float))
+		else if( TokenType==TOKEN_Const && ConstantType ==CPT_Float && Float==FMath::TruncToInt(Float))
 		{
 			I = (int32) Float;
 			return 1;
 		}
-		else if (TokenType == TOKEN_Const && Type == CPT_Double && Double == FMath::TruncToInt((float)Double))
+		else if (TokenType == TOKEN_Const && ConstantType == CPT_Double && Double == FMath::TruncToInt((float)Double))
 		{
 			I = (int32) Double;
 			return 1;
@@ -1001,27 +738,27 @@ public:
 
 	bool GetConstInt64( int64& I ) const
 	{
-		if( TokenType==TOKEN_Const && Type==CPT_Int64 )
+		if( TokenType==TOKEN_Const && ConstantType ==CPT_Int64 )
 		{
 			I = Int64;
 			return 1;
 		}
-		else if( TokenType==TOKEN_Const && Type==CPT_Int )
+		else if( TokenType==TOKEN_Const && ConstantType ==CPT_Int )
 		{
 			I = Int;
 			return 1;
 		}
-		else if( TokenType==TOKEN_Const && Type==CPT_Byte )
+		else if( TokenType==TOKEN_Const && ConstantType ==CPT_Byte )
 		{
 			I = Byte;
 			return 1;
 		}
-		else if( TokenType==TOKEN_Const && Type==CPT_Float && Float==FMath::TruncToInt(Float))
+		else if( TokenType==TOKEN_Const && ConstantType ==CPT_Float && Float==FMath::TruncToInt(Float))
 		{
 			I = (int32) Float;
 			return 1;
 		}
-		else if (TokenType == TOKEN_Const && Type == CPT_Double && Double == FMath::TruncToInt((float)Double))
+		else if (TokenType == TOKEN_Const && ConstantType == CPT_Double && Double == FMath::TruncToInt((float)Double))
 		{
 			I = (int32) Double;
 			return 1;
@@ -1029,18 +766,6 @@ public:
 		else return 0;
 	}
 
-	FString Describe()
-	{
-		return FString::Printf(
-			TEXT("Property:%s  Type:%s  TokenName:%s  ConstValue:%s  Struct:%s  Flags:%lli  Implied:%lli"),
-			TokenProperty!=NULL?*TokenProperty->GetName():TEXT("NULL"),
-			GetPropertyTypeText(Type), *GetTokenName().ToString(), *GetConstantValue(),
-			Struct!=NULL?*Struct->GetName():TEXT("NULL"),
-			PropertyFlags,
-			ImpliedPropertyFlags
-			);
-	}
-
 	friend struct FTokenArchiveProxy;
 };
 
@@ -1184,9 +909,6 @@ enum class EParsedInterface
  */
 class FStructMetaData
 {
-	/** whether this class declares delegate functions or properties */
-	bool													bContainsDelegates;
-
 	/** The line of UCLASS/UINTERFACE macro in this class. */
 	int32 PrologLine;
 
@@ -1199,8 +921,7 @@ class FStructMetaData
 public:
 	/** Default constructor */
 	FStructMetaData()
-		: bContainsDelegates(false)
-		, PrologLine(-1)
+		: PrologLine(-1)
 		, GeneratedBodyLine(-1)
 		, InterfaceGeneratedBodyLine(-1)
 		, bConstructorDeclared(false)
@@ -1265,81 +986,6 @@ public:
 		InterfaceGeneratedBodyLine = Line;
 	}
 
-	/**
-	 * Sets contains delegates flag for this class.
-	 */
-	void MarkContainsDelegate()
-	{
-		bContainsDelegates = true;
-	}
-
-	/**
-	 * Adds a new property to be tracked.  Determines the correct list for the property based on
-	 * its owner (function, struct, etc).
-	 * 
-	 * @param	PropertyDef	the property to add
-	 */
-	void AddProperty(FUnrealPropertyDefinitionInfo& PropertyDef);
-
-	/**
-	 * Adds new editor-only metadata (key/value pairs) to the class or struct that
-	 * owns this property or function.
-	 * 
-	 * @param	Field		the property or function to add to
-	 * @param	InMetaData	the metadata to add
-	 */
-	static void AddMetaData(UField* Field, TMap<FName, FString>&& InMetaData)
-	{
-		// only add if we have some!
-		if (InMetaData.Num())
-		{
-			check(Field);
-
-			// get (or create) a metadata object for this package
-			UMetaData* MetaData = Field->GetOutermost()->GetMetaData();
-			TMap<FName, FString>* ExistingMetaData = MetaData->GetMapForObject(Field);
-			if (ExistingMetaData && ExistingMetaData->Num())
-			{
-				// Merge the existing metadata
-				TMap<FName, FString> MergedMetaData;
-				MergedMetaData.Reserve(InMetaData.Num() + ExistingMetaData->Num());
-				MergedMetaData.Append(*ExistingMetaData);
-				MergedMetaData.Append(InMetaData);
-				MetaData->SetObjectValues(Field, MoveTemp(MergedMetaData));
-			}
-			else
-			{
-				// set the metadata for this field
-				MetaData->SetObjectValues(Field, MoveTemp(InMetaData));
-			}
-		}
-	}
-	static void AddMetaData(FField* Field, TMap<FName, FString>&& InMetaData)
-	{
-		// only add if we have some!
-		if (InMetaData.Num())
-		{
-			check(Field);
-
-			UPackage* Package = Field->GetOutermost();
-			// get (or create) a metadata object for this package
-			UMetaData* MetaData = Package->GetMetaData();
-			
-			for (TPair<FName, FString>& MetaKeyValue : InMetaData)
-			{
-				Field->SetMetaData(MetaKeyValue.Key, MoveTemp(MetaKeyValue.Value));
-			}
-		}
-	}
-
-	/**
-	 * Returns whether this class contains any delegate properties which need to be fixed up.
-	 */
-	bool ContainsDelegates() const
-	{
-		return bContainsDelegates;
-	}
-
 	// Is constructor declared?
 	bool bConstructorDeclared;
 
