Added new PropertyTypes which in the long run will be where different property types are implemented.

diff --git a/Engine/Source/Programs/UnrealHeaderTool/Private/PropertyTypes.cpp b/Engine/Source/Programs/UnrealHeaderTool/Private/PropertyTypes.cpp
new file mode 100644
--- /dev/null
+++ b/Engine/Source/Programs/UnrealHeaderTool/Private/PropertyTypes.cpp
@@ -0,0 +1,1433 @@
+// Copyright Epic Games, Inc. All Rights Reserved.
+
+#include "PropertyTypes.h"
+#include "BaseParser.h"
+#include "Classes.h"
+#include "ClassMaps.h"
+#include "HeaderParser.h"
+#include "UnrealHeaderTool.h"
+#include "UnrealTypeDefinitionInfo.h"
+#include "Misc/DefaultValueHelper.h"
+#include "UObject/ObjectMacros.h"
+
+void AddEditInlineMetaData(TMap<FName, FString>& MetaData);
+void AddMetaDataToClassData(FUnrealTypeDefinitionInfo& TypeDef, TMap<FName, FString>&& InMetaData);
+
+// Following is the relationship between the property types
+//
+//		FProperty
+//			FNumericProperty
+//				FByteProperty
+//				FInt8Property
+//				FInt16Property
+//				FIntProperty
+//				FInt64Property
+//				FUInt16Property
+//				FUInt32Property
+//				FUInt64Property
+//				FFloatProperty
+//				FDoubleProperty
+//				FLargeWorldCoordinatesRealProperty
+//			FBoolProperty
+//			FEnumProperty
+//			TObjectPropertyBase
+//				FObjectProperty
+//					FClassProperty
+//						FClassPtrProperty
+//					FObjectPtrProperty
+//				FWeakObjectProperty
+//				FLazyObjectProperty
+//				FSoftObjectProperty
+//					FSoftClassProperty
+//			FInterfaceProperty
+//			FNameProperty
+//			FStrProperty
+//			FTextProperty
+//			FStructProperty
+//			FMulticastSparseDelegateProperty
+//			FMulticastInlineDelegateProperty
+//			FFieldPathProperty
+//			FArrayProperty
+//			FSetProperty
+//			FMapProperty
+
+struct FPropertyTypeTraitsByte;
+struct FPropertyTypeTraitsInt8;
+struct FPropertyTypeTraitsInt16;
+struct FPropertyTypeTraitsInt;
+struct FPropertyTypeTraitsInt64;
+struct FPropertyTypeTraitsUInt16;
+struct FPropertyTypeTraitsUInt32;
+struct FPropertyTypeTraitsUInt64;
+struct FPropertyTypeTraitsBool;
+struct FPropertyTypeTraitsBool8;
+struct FPropertyTypeTraitsBool16;
+struct FPropertyTypeTraitsBool32;
+struct FPropertyTypeTraitsBool64;
+struct FPropertyTypeTraitsFloat;
+struct FPropertyTypeTraitsDouble;
+struct FPropertyTypeTraitsLargeWorldCoordinatesReal;
+struct FPropertyTypeTraitsObjectReference;
+struct FPropertyTypeTraitsObjectReference;
+struct FPropertyTypeTraitsWeakObjectReference;
+struct FPropertyTypeTraitsLazyObjectReference;
+struct FPropertyTypeTraitsObjectPtrReference;
+struct FPropertyTypeTraitsSoftObjectReference;
+struct FPropertyTypeTraitsInterface;
+struct FPropertyTypeTraitsName;
+struct FPropertyTypeTraitsString;
+struct FPropertyTypeTraitsText;
+struct FPropertyTypeTraitsStruct;
+struct FPropertyTypeTraitsDelegate;
+struct FPropertyTypeTraitsMulticastDelegate;
+struct FPropertyTypeTraitsFieldPath;
+struct FPropertyTypeTraitsStaticArray;
+struct FPropertyTypeTraitsDynamicArray;
+struct FPropertyTypeTraitsSet;
+struct FPropertyTypeTraitsMap;
+struct FPropertyTypeTraitsEnum;
+
+namespace
+{
+	void PropagateFlagsFromInnerAndHandlePersistentInstanceMetadata(EPropertyFlags& DestFlags, const TMap<FName, FString>& InMetaData, FUnrealPropertyDefinitionInfo& InnerDef)
+	{
+		FProperty* Inner = InnerDef.GetProperty();
+
+		// Copy some of the property flags to the container property.
+		if (Inner->PropertyFlags & (CPF_ContainsInstancedReference | CPF_InstancedReference))
+		{
+			DestFlags |= CPF_ContainsInstancedReference;
+			DestFlags &= ~(CPF_InstancedReference | CPF_PersistentInstance); //this was propagated to the inner
+
+			if (Inner->PropertyFlags & CPF_PersistentInstance)
+			{
+				AddMetaDataToClassData(InnerDef, TMap<FName, FString>(InMetaData));
+			}
+		}
+	};
+
+	//----------------------------------------------------------------------------------------------------------------------------------------------------------
+	//----------------------------------------------------------------------------------------------------------------------------------------------------------
+	// Dispatch system
+	//----------------------------------------------------------------------------------------------------------------------------------------------------------
+	//----------------------------------------------------------------------------------------------------------------------------------------------------------
+
+	/**
+	 * Given just the property type, invoke the dispatch functor.
+	 *
+	 * NOTE: This method does not support container or enum types.
+	 */
+	template <template<typename PropertyTraits> typename FuncDispatch, typename RetValue, typename ... Args>
+	RetValue PropertyTypeDispatch(EPropertyType PropertyType, Args&& ... args)
+	{
+		switch (PropertyType)
+		{
+		case CPT_Byte:						return FuncDispatch<FPropertyTypeTraitsByte>()(std::forward<Args>(args)...);
+		case CPT_Int8:						return FuncDispatch<FPropertyTypeTraitsInt8>()(std::forward<Args>(args)...);
+		case CPT_Int16:						return FuncDispatch<FPropertyTypeTraitsInt16>()(std::forward<Args>(args)...);
+		case CPT_Int:						return FuncDispatch<FPropertyTypeTraitsInt>()(std::forward<Args>(args)...);
+		case CPT_Int64:						return FuncDispatch<FPropertyTypeTraitsInt64>()(std::forward<Args>(args)...);
+		case CPT_UInt16:					return FuncDispatch<FPropertyTypeTraitsUInt16>()(std::forward<Args>(args)...);
+		case CPT_UInt32:					return FuncDispatch<FPropertyTypeTraitsUInt32>()(std::forward<Args>(args)...);
+		case CPT_UInt64:					return FuncDispatch<FPropertyTypeTraitsUInt64>()(std::forward<Args>(args)...);
+		case CPT_Bool:						return FuncDispatch<FPropertyTypeTraitsBool>()(std::forward<Args>(args)...);
+		case CPT_Bool8:						return FuncDispatch<FPropertyTypeTraitsBool8>()(std::forward<Args>(args)...);
+		case CPT_Bool16:					return FuncDispatch<FPropertyTypeTraitsBool16>()(std::forward<Args>(args)...);
+		case CPT_Bool32:					return FuncDispatch<FPropertyTypeTraitsBool32>()(std::forward<Args>(args)...);
+		case CPT_Bool64:					return FuncDispatch<FPropertyTypeTraitsBool64>()(std::forward<Args>(args)...);
+		case CPT_Float:						return FuncDispatch<FPropertyTypeTraitsFloat>()(std::forward<Args>(args)...);
+		case CPT_Double:					return FuncDispatch<FPropertyTypeTraitsDouble>()(std::forward<Args>(args)...);
+		case CPT_FLargeWorldCoordinatesReal:return FuncDispatch<FPropertyTypeTraitsLargeWorldCoordinatesReal>()(std::forward<Args>(args)...);
+		case CPT_ObjectReference:			return FuncDispatch<FPropertyTypeTraitsObjectReference>()(std::forward<Args>(args)...);
+		case CPT_WeakObjectReference:		return FuncDispatch<FPropertyTypeTraitsWeakObjectReference>()(std::forward<Args>(args)...);
+		case CPT_LazyObjectReference:		return FuncDispatch<FPropertyTypeTraitsLazyObjectReference>()(std::forward<Args>(args)...);
+		case CPT_ObjectPtrReference:		return FuncDispatch<FPropertyTypeTraitsObjectPtrReference>()(std::forward<Args>(args)...);
+		case CPT_SoftObjectReference:		return FuncDispatch<FPropertyTypeTraitsSoftObjectReference>()(std::forward<Args>(args)...);
+		case CPT_Interface:					return FuncDispatch<FPropertyTypeTraitsInterface>()(std::forward<Args>(args)...);
+		case CPT_Name:						return FuncDispatch<FPropertyTypeTraitsName>()(std::forward<Args>(args)...);
+		case CPT_String:					return FuncDispatch<FPropertyTypeTraitsString>()(std::forward<Args>(args)...);
+		case CPT_Text:						return FuncDispatch<FPropertyTypeTraitsText>()(std::forward<Args>(args)...);
+		case CPT_Struct:					return FuncDispatch<FPropertyTypeTraitsStruct>()(std::forward<Args>(args)...);
+		case CPT_Delegate:					return FuncDispatch<FPropertyTypeTraitsDelegate>()(std::forward<Args>(args)...);
+		case CPT_MulticastDelegate:			return FuncDispatch<FPropertyTypeTraitsMulticastDelegate>()(std::forward<Args>(args)...);
+		case CPT_FieldPath:					return FuncDispatch<FPropertyTypeTraitsFieldPath>()(std::forward<Args>(args)...);
+		default:							FError::Throwf(TEXT("Unknown property type %i"), (uint8)PropertyType);
+		}
+	}
+
+	/**
+	 * Given just the property type, invoke the dispatch functor.
+	 */
+	template <template <typename PropertyTraits> typename FuncDispatch, bool bHandleContainers, typename RetValue, typename ... Args>
+	RetValue PropertyTypeDispatch(const FPropertyBase& VarProperty, Args&& ... args)
+	{
+		if constexpr (bHandleContainers)
+		{
+			switch (VarProperty.ArrayType)
+			{
+			case EArrayType::Static:
+				return FuncDispatch<FPropertyTypeTraitsStaticArray>()(std::forward<Args>(args)...);
+			case EArrayType::Dynamic:
+				return FuncDispatch<FPropertyTypeTraitsDynamicArray>()(std::forward<Args>(args)...);
+			case EArrayType::Set:
+				return FuncDispatch<FPropertyTypeTraitsSet>()(std::forward<Args>(args)...);
+			}
+
+			if (VarProperty.MapKeyProp.IsValid())
+			{
+				return FuncDispatch<FPropertyTypeTraitsMap>()(std::forward<Args>(args)...);
+			}
+		}
+
+		// Check if it's an enum class property
+		// NOTE: VarProperty.Enum is a union and might not be an enum
+		if (FUnrealEnumDefinitionInfo* EnumDef = GTypeDefinitionInfoMap.Find<FUnrealEnumDefinitionInfo>(VarProperty.Enum))
+		{
+			return FuncDispatch<FPropertyTypeTraitsEnum>()(std::forward<Args>(args)...);
+		}
+
+		return PropertyTypeDispatch<FuncDispatch, RetValue>(VarProperty.Type, std::forward<Args>(args)...);
+	}
+
+	//----------------------------------------------------------------------------------------------------------------------------------------------------------
+	//----------------------------------------------------------------------------------------------------------------------------------------------------------
+	// Dispatch Functors
+	//----------------------------------------------------------------------------------------------------------------------------------------------------------
+	//----------------------------------------------------------------------------------------------------------------------------------------------------------
+
+	template<typename TraitsType>
+	struct DefaultValueStringCppFormatToInnerFormatDispatch
+	{
+		bool operator()(const FUnrealPropertyDefinitionInfo& PropDef, const FString& CppForm, FString& OutForm)
+		{
+			return TraitsType::DefaultValueStringCppFormatToInnerFormat(PropDef, CppForm, OutForm);
+		}
+	};
+
+	template<typename TraitsType>
+	struct IsObjectDispatch
+	{
+		bool operator()()
+		{
+			return TraitsType::bIsObject;
+		}
+	};
+
+	template<typename TraitsType>
+	struct CreateEngineTypeDispatch
+	{
+		FProperty* operator()(FUnrealPropertyDefinitionInfo& PropDef, FFieldVariant Scope, const FName& Name, EObjectFlags ObjectFlags, EVariableCategory::Type VariableCategory, const TCHAR* Dimensions)
+		{
+			return TraitsType::CreateEngineType(PropDef, Scope, Name, ObjectFlags, VariableCategory, Dimensions);
+		}
+	};
+
+	template<typename TraitsType>
+	struct IsSupportedByBlueprintDispatch
+	{
+		bool operator()(const FUnrealPropertyDefinitionInfo& PropDef, bool bMemberVariable)
+		{
+			return TraitsType::IsSupportedByBlueprint(PropDef, bMemberVariable);
+		}
+	};
+
+	//----------------------------------------------------------------------------------------------------------------------------------------------------------
+	//----------------------------------------------------------------------------------------------------------------------------------------------------------
+	// Helper Methods
+	//----------------------------------------------------------------------------------------------------------------------------------------------------------
+	//----------------------------------------------------------------------------------------------------------------------------------------------------------
+
+	template <bool bHandleContainers>
+	FProperty* CreatePropertyHelper(FUnrealPropertyDefinitionInfo& PropDef, FFieldVariant Scope, const FName& Name, EObjectFlags ObjectFlags, EVariableCategory::Type VariableCategory, const TCHAR* Dimensions)
+	{
+		FProperty* Property = PropertyTypeDispatch<CreateEngineTypeDispatch, bHandleContainers, FProperty*>(PropDef.GetPropertyBase(), std::ref(PropDef), Scope, std::ref(Name), ObjectFlags, VariableCategory, Dimensions);
+		Property->PropertyFlags = PropDef.GetPropertyBase().PropertyFlags;
+		return Property;
+	}
+
+	bool IsSupportedByBlueprintSansContainers(const FUnrealPropertyDefinitionInfo& PropDef, bool bMemberVariable)
+	{
+		return PropertyTypeDispatch<IsSupportedByBlueprintDispatch, false, bool>(PropDef.GetPropertyBase(), std::ref(PropDef), bMemberVariable);
+	}
+}
+
+/**
+ * Every property type is required to implement the follow methods and constants or derive from a base class that has them.
+ */
+struct FPropertyTypeTraitsBase
+{
+	/**
+	 * If true, this property type is an object property
+	*/
+	static constexpr bool bIsObject = false;
+
+	/**
+	 * Transforms CPP-formated string containing default value, to inner formated string
+	 * If it cannot be transformed empty string is returned.
+	 *
+	 * @param Property The property that owns the default value.
+	 * @param CppForm A CPP-formated string.
+	 * @param out InnerForm Inner formated string
+	 * @return true on success, false otherwise.
+	 */
+	static bool DefaultValueStringCppFormatToInnerFormat(const FUnrealPropertyDefinitionInfo& PropDef, const FString& CppForm, FString& OutForm)
+	{
+		return false;
+	}
+
+	/**
+	 * Given a property definition with the property base data already populated, create the underlying engine type.
+	 *
+	 * NOTE: This method MUST be implemented by each of the property types.
+	 *
+	 * @param PropDef The definition of the property
+	 * @param Scope The parent object owning the property
+	 * @param Name The name of the property
+	 * @param ObjectFlags The flags associated with the property
+	 * @param VariableCategory The parsing context of the property
+	 * @param Dimensions When this is a static array, this represents the dimensions value
+	 * @return The pointer to the newly created property.  It will be attached to the definition by the caller
+	 */
+	 //static FProperty* CreateEngineType(FUnrealPropertyDefinitionInfo& PropDef, FFieldVariant Scope, const FName& Name, EObjectFlags ObjectFlags, EVariableCategory::Type VariableCategory, const TCHAR* Dimensions)
+
+	/**
+	 * Returns true if this property is supported by blueprints
+	 */
+	static bool IsSupportedByBlueprint(const FUnrealPropertyDefinitionInfo& PropDef, bool bMemberVariable)
+	{
+		return false;
+	}
+};
+
+//----------------------------------------------------------------------------------------------------------------------------------------------------------
+//----------------------------------------------------------------------------------------------------------------------------------------------------------
+// Numeric types
+//----------------------------------------------------------------------------------------------------------------------------------------------------------
+//----------------------------------------------------------------------------------------------------------------------------------------------------------
+
+struct FPropertyTypeTraitsNumericBase : public FPropertyTypeTraitsBase
+{
+};
+
+struct FPropertyTypeTraitsByte : public FPropertyTypeTraitsNumericBase
+{
+	static bool DefaultValueStringCppFormatToInnerFormat(const FUnrealPropertyDefinitionInfo& PropDef, const FString& CppForm, FString& OutForm)
+	{
+		int32 Value;
+		if (FDefaultValueHelper::ParseInt(CppForm, Value))
+		{
+			OutForm = FString::FromInt(Value);
+			return (0 <= Value) && (255 >= Value);
+		}
+		return false;
+	}
+
+	static FProperty* CreateEngineType(FUnrealPropertyDefinitionInfo& PropDef, FFieldVariant Scope, const FName& Name, EObjectFlags ObjectFlags, EVariableCategory::Type VariableCategory, const TCHAR* Dimensions)
+	{
+		const FPropertyBase& VarProperty = PropDef.GetPropertyBase();
+#if UHT_ENABLE_VALUE_PROPERTY_TAG
+		PropDef.GetUnrealSourceFile().AddTypeDefIncludeIfNeeded(VarProperty.Enum);
+#endif
+		FByteProperty* Result = new FByteProperty(Scope, Name, ObjectFlags);
+		Result->Enum = VarProperty.Enum;
+		check(VarProperty.IntType == EIntType::Sized);
+		return Result;
+	}
+
+	static bool IsSupportedByBlueprint(const FUnrealPropertyDefinitionInfo& PropDef, bool bMemberVariable)
+	{
+		return true;
+	}
+};
+
+struct FPropertyTypeTraitsInt8 : public FPropertyTypeTraitsNumericBase
+{
+	static FProperty* CreateEngineType(FUnrealPropertyDefinitionInfo& PropDef, FFieldVariant Scope, const FName& Name, EObjectFlags ObjectFlags, EVariableCategory::Type VariableCategory, const TCHAR* Dimensions)
+	{
+		const FPropertyBase& VarProperty = PropDef.GetPropertyBase();
+		FInt8Property* Result = new FInt8Property(Scope, Name, ObjectFlags);
+		check(VarProperty.IntType == EIntType::Sized);
+		return Result;
+	}
+};
+
+struct FPropertyTypeTraitsInt16 : public FPropertyTypeTraitsNumericBase
+{
+	static FProperty* CreateEngineType(FUnrealPropertyDefinitionInfo& PropDef, FFieldVariant Scope, const FName& Name, EObjectFlags ObjectFlags, EVariableCategory::Type VariableCategory, const TCHAR* Dimensions)
+	{
+		const FPropertyBase& VarProperty = PropDef.GetPropertyBase();
+		FInt16Property* Result = new FInt16Property(Scope, Name, ObjectFlags);
+		check(VarProperty.IntType == EIntType::Sized);
+		return Result;
+	}
+};
+
+struct FPropertyTypeTraitsInt : public FPropertyTypeTraitsNumericBase
+{
+	static bool DefaultValueStringCppFormatToInnerFormat(const FUnrealPropertyDefinitionInfo& PropDef, const FString& CppForm, FString& OutForm)
+	{
+		int32 Value;
+		if (FDefaultValueHelper::ParseInt(CppForm, Value))
+		{
+			OutForm = FString::FromInt(Value);
+			return true;
+		}
+		return false;
+	}
+
+	static FProperty* CreateEngineType(FUnrealPropertyDefinitionInfo& PropDef, FFieldVariant Scope, const FName& Name, EObjectFlags ObjectFlags, EVariableCategory::Type VariableCategory, const TCHAR* Dimensions)
+	{
+		const FPropertyBase& VarProperty = PropDef.GetPropertyBase();
+		FIntProperty* Result = new FIntProperty(Scope, Name, ObjectFlags);
+		PropDef.SetUnsized(VarProperty.IntType == EIntType::Unsized);
+		return Result;
+	}
+
+	static bool IsSupportedByBlueprint(const FUnrealPropertyDefinitionInfo& PropDef, bool bMemberVariable)
+	{
+		return true;
+	}
+};
+
+struct FPropertyTypeTraitsInt64 : public FPropertyTypeTraitsNumericBase
+{
+	static bool DefaultValueStringCppFormatToInnerFormat(const FUnrealPropertyDefinitionInfo& PropDef, const FString& CppForm, FString& OutForm)
+	{
+		int64 Value;
+		if (FDefaultValueHelper::ParseInt64(CppForm, Value))
+		{
+			OutForm = FString::Printf(TEXT("%lld"), Value);
+			return true;
+		}
+		return false;
+	}
+
+	static FProperty* CreateEngineType(FUnrealPropertyDefinitionInfo& PropDef, FFieldVariant Scope, const FName& Name, EObjectFlags ObjectFlags, EVariableCategory::Type VariableCategory, const TCHAR* Dimensions)
+	{
+		const FPropertyBase& VarProperty = PropDef.GetPropertyBase();
+		FInt64Property* Result = new FInt64Property(Scope, Name, ObjectFlags);
+		check(VarProperty.IntType == EIntType::Sized);
+		return Result;
+	}
+
+	static bool IsSupportedByBlueprint(const FUnrealPropertyDefinitionInfo& PropDef, bool bMemberVariable)
+	{
+		return true;
+	}
+};
+
+struct FPropertyTypeTraitsUInt16 : public FPropertyTypeTraitsNumericBase
+{
+	static FProperty* CreateEngineType(FUnrealPropertyDefinitionInfo& PropDef, FFieldVariant Scope, const FName& Name, EObjectFlags ObjectFlags, EVariableCategory::Type VariableCategory, const TCHAR* Dimensions)
+	{
+		const FPropertyBase& VarProperty = PropDef.GetPropertyBase();
+		FUInt16Property* Result = new FUInt16Property(Scope, Name, ObjectFlags);
+		check(VarProperty.IntType == EIntType::Sized);
+		return Result;
+	}
+};
+
+struct FPropertyTypeTraitsUInt32 : public FPropertyTypeTraitsNumericBase
+{
+	static FProperty* CreateEngineType(FUnrealPropertyDefinitionInfo& PropDef, FFieldVariant Scope, const FName& Name, EObjectFlags ObjectFlags, EVariableCategory::Type VariableCategory, const TCHAR* Dimensions)
+	{
+		const FPropertyBase& VarProperty = PropDef.GetPropertyBase();
+		FUInt32Property* Result = new FUInt32Property(Scope, Name, ObjectFlags);
+		PropDef.SetUnsized(VarProperty.IntType == EIntType::Unsized);
+		return Result;
+	}
+};
+
+struct FPropertyTypeTraitsUInt64 : public FPropertyTypeTraitsNumericBase
+{
+	static FProperty* CreateEngineType(FUnrealPropertyDefinitionInfo& PropDef, FFieldVariant Scope, const FName& Name, EObjectFlags ObjectFlags, EVariableCategory::Type VariableCategory, const TCHAR* Dimensions)
+	{
+		const FPropertyBase& VarProperty = PropDef.GetPropertyBase();
+		FUInt64Property* Result = new FUInt64Property(Scope, Name, ObjectFlags);
+		check(VarProperty.IntType == EIntType::Sized);
+		return Result;
+	}
+};
+
+struct FPropertyTypeTraitsFloat : public FPropertyTypeTraitsNumericBase
+{
+	static bool DefaultValueStringCppFormatToInnerFormat(const FUnrealPropertyDefinitionInfo& PropDef, const FString& CppForm, FString& OutForm)
+	{
+		float Value;
+		if (FDefaultValueHelper::ParseFloat(CppForm, Value))
+		{
+			OutForm = FString::Printf(TEXT("%f"), Value);
+			return true;
+		}
+		return false;
+	}
+
+	static FProperty* CreateEngineType(FUnrealPropertyDefinitionInfo& PropDef, FFieldVariant Scope, const FName& Name, EObjectFlags ObjectFlags, EVariableCategory::Type VariableCategory, const TCHAR* Dimensions)
+	{
+		FFloatProperty* Result = new FFloatProperty(Scope, Name, ObjectFlags);
+		return Result;
+	}
+
+	static bool IsSupportedByBlueprint(const FUnrealPropertyDefinitionInfo& PropDef, bool bMemberVariable)
+	{
+		return true;
+	}
+};
+
+struct FPropertyTypeTraitsDouble : public FPropertyTypeTraitsNumericBase
+{
+	static bool DefaultValueStringCppFormatToInnerFormat(const FUnrealPropertyDefinitionInfo& PropDef, const FString& CppForm, FString& OutForm)
+	{
+		double Value;
+		if (FDefaultValueHelper::ParseDouble(CppForm, Value))
+		{
+			OutForm = FString::Printf(TEXT("%f"), Value);
+			return true;
+		}
+		return false;
+	}
+
+	static FProperty* CreateEngineType(FUnrealPropertyDefinitionInfo& PropDef, FFieldVariant Scope, const FName& Name, EObjectFlags ObjectFlags, EVariableCategory::Type VariableCategory, const TCHAR* Dimensions)
+	{
+		FDoubleProperty* Result = new FDoubleProperty(Scope, Name, ObjectFlags);
+		return Result;
+	}
+
+	static bool IsSupportedByBlueprint(const FUnrealPropertyDefinitionInfo& PropDef, bool bMemberVariable)
+	{
+		return true;
+	}
+};
+
+struct FPropertyTypeTraitsLargeWorldCoordinatesReal : public FPropertyTypeTraitsNumericBase
+{
+	static bool DefaultValueStringCppFormatToInnerFormat(const FUnrealPropertyDefinitionInfo& PropDef, const FString& CppForm, FString& OutForm)
+	{
+		double Value;
+		if (FDefaultValueHelper::ParseDouble(CppForm, Value))
+		{
+			OutForm = FString::Printf(TEXT("%f"), Value);
+			return true;
+		}
+		return false;
+	}
+
+	static FProperty* CreateEngineType(FUnrealPropertyDefinitionInfo& PropDef, FFieldVariant Scope, const FName& Name, EObjectFlags ObjectFlags, EVariableCategory::Type VariableCategory, const TCHAR* Dimensions)
+	{
+		FLargeWorldCoordinatesRealProperty* Result = new FLargeWorldCoordinatesRealProperty(Scope, Name, ObjectFlags);
+		return Result;
+	}
+
+	static bool IsSupportedByBlueprint(const FUnrealPropertyDefinitionInfo& PropDef, bool bMemberVariable)
+	{
+		return true;
+	}
+};
+
+//----------------------------------------------------------------------------------------------------------------------------------------------------------
+//----------------------------------------------------------------------------------------------------------------------------------------------------------
+// Boolean types
+//----------------------------------------------------------------------------------------------------------------------------------------------------------
+//----------------------------------------------------------------------------------------------------------------------------------------------------------
+
+/**
+ * Base implementation for all boolean property types
+ */
+struct FPropertyTypeTraitsBooleanBase : public FPropertyTypeTraitsBase
+{
+	static bool DefaultValueStringCppFormatToInnerFormat(const FUnrealPropertyDefinitionInfo& PropDef, const FString& CppForm, FString& OutForm)
+	{
+		if (FDefaultValueHelper::Is(CppForm, TEXT("true")) ||
+			FDefaultValueHelper::Is(CppForm, TEXT("false")))
+		{
+			OutForm = FDefaultValueHelper::RemoveWhitespaces(CppForm);
+			return true;
+		}
+		return false;
+	}
+
+	template <typename SizeType, bool bIsNativeBool>
+	static FProperty* CreateEngineTypeHelper(FUnrealPropertyDefinitionInfo& PropDef, FFieldVariant Scope, const FName& Name, EObjectFlags ObjectFlags, EVariableCategory::Type VariableCategory, const TCHAR* Dimensions)
+	{
+		FBoolProperty* Result = new FBoolProperty(Scope, Name, ObjectFlags);
+		bool bActsLikeNativeBool = bIsNativeBool || VariableCategory == EVariableCategory::Return;
+		Result->SetBoolSize(bActsLikeNativeBool ? sizeof(bool) : sizeof(SizeType), bActsLikeNativeBool);
+		return Result;
+	}
+
+	static bool IsSupportedByBlueprint(const FUnrealPropertyDefinitionInfo& PropDef, bool bMemberVariable)
+	{
+		return true;
+	}
+};
+
+struct FPropertyTypeTraitsBool : public FPropertyTypeTraitsBooleanBase
+{
+	static FProperty* CreateEngineType(FUnrealPropertyDefinitionInfo& PropDef, FFieldVariant Scope, const FName& Name, EObjectFlags ObjectFlags, EVariableCategory::Type VariableCategory, const TCHAR* Dimensions)
+	{
+		return CreateEngineTypeHelper<bool, true>(PropDef, Scope, Name, ObjectFlags, VariableCategory, Dimensions);
+	}
+};
+
+struct FPropertyTypeTraitsBool8 : public FPropertyTypeTraitsBooleanBase
+{
+	static FProperty* CreateEngineType(FUnrealPropertyDefinitionInfo& PropDef, FFieldVariant Scope, const FName& Name, EObjectFlags ObjectFlags, EVariableCategory::Type VariableCategory, const TCHAR* Dimensions)
+	{
+		return CreateEngineTypeHelper<uint8, false>(PropDef, Scope, Name, ObjectFlags, VariableCategory, Dimensions);
+	}
+};
+
+struct FPropertyTypeTraitsBool16 : public FPropertyTypeTraitsBooleanBase
+{
+	static FProperty* CreateEngineType(FUnrealPropertyDefinitionInfo& PropDef, FFieldVariant Scope, const FName& Name, EObjectFlags ObjectFlags, EVariableCategory::Type VariableCategory, const TCHAR* Dimensions)
+	{
+		return CreateEngineTypeHelper<uint16, false>(PropDef, Scope, Name, ObjectFlags, VariableCategory, Dimensions);
+	}
+};
+
+struct FPropertyTypeTraitsBool32 : public FPropertyTypeTraitsBooleanBase
+{
+	static FProperty* CreateEngineType(FUnrealPropertyDefinitionInfo& PropDef, FFieldVariant Scope, const FName& Name, EObjectFlags ObjectFlags, EVariableCategory::Type VariableCategory, const TCHAR* Dimensions)
+	{
+		return CreateEngineTypeHelper<uint32, false>(PropDef, Scope, Name, ObjectFlags, VariableCategory, Dimensions);
+	}
+};
+
+struct FPropertyTypeTraitsBool64 : public FPropertyTypeTraitsBooleanBase
+{
+	static FProperty* CreateEngineType(FUnrealPropertyDefinitionInfo& PropDef, FFieldVariant Scope, const FName& Name, EObjectFlags ObjectFlags, EVariableCategory::Type VariableCategory, const TCHAR* Dimensions)
+	{
+		return CreateEngineTypeHelper<uint64, false>(PropDef, Scope, Name, ObjectFlags, VariableCategory, Dimensions);
+	}
+};
+
+//----------------------------------------------------------------------------------------------------------------------------------------------------------
+//----------------------------------------------------------------------------------------------------------------------------------------------------------
+// Enumeration types
+//----------------------------------------------------------------------------------------------------------------------------------------------------------
+//----------------------------------------------------------------------------------------------------------------------------------------------------------
+
+struct FPropertyTypeTraitsEnum : public FPropertyTypeTraitsBase
+{
+	static bool DefaultValueStringCppFormatToInnerFormat(const FUnrealPropertyDefinitionInfo& PropDef, const FString& CppForm, FString& OutForm)
+	{
+		UEnum* Enum = PropDef.GetPropertyBase().Enum;
+		OutForm = FDefaultValueHelper::GetUnqualifiedEnumValue(FDefaultValueHelper::RemoveWhitespaces(CppForm));
+
+		const int32 EnumEntryIndex = Enum->GetIndexByName(*OutForm);
+		if (EnumEntryIndex == INDEX_NONE)
+		{
+			return false;
+		}
+		if (Enum->HasMetaData(TEXT("Hidden"), EnumEntryIndex))
+		{
+			FError::Throwf(TEXT("Hidden enum entries cannot be used as default values: %s \"%s\" "), *PropDef.GetProperty()->GetName(), *CppForm);
+		}
+		return true;
+	}
+
+	static FProperty* CreateEngineType(FUnrealPropertyDefinitionInfo& PropDef, FFieldVariant Scope, const FName& Name, EObjectFlags ObjectFlags, EVariableCategory::Type VariableCategory, const TCHAR* Dimensions)
+	{
+		const FPropertyBase& VarProperty = PropDef.GetPropertyBase();
+
+		if (VarProperty.Enum->GetCppForm() != UEnum::ECppForm::EnumClass)
+		{
+			check(VarProperty.Type == EPropertyType::CPT_Byte);
+			return FPropertyTypeTraitsByte::CreateEngineType(PropDef, Scope, Name, ObjectFlags, VariableCategory, Dimensions);
+		}
+
+		FUnrealEnumDefinitionInfo& EnumDef = GTypeDefinitionInfoMap.FindChecked<FUnrealEnumDefinitionInfo>(VarProperty.Enum);
+
+#if UHT_ENABLE_VALUE_PROPERTY_TAG
+		PropDef.GetUnrealSourceFile().AddTypeDefIncludeIfNeeded(EnumDef);
+#endif
+
+		FPropertyBase UnderlyingProperty = VarProperty;
+		UnderlyingProperty.Enum = nullptr;
+		UnderlyingProperty.PropertyFlags = EPropertyFlags::CPF_None;
+		UnderlyingProperty.ArrayType = EArrayType::None;
+		switch (EnumDef.GetUnderlyingType())
+		{
+		case EUnderlyingEnumType::int8:        UnderlyingProperty.Type = CPT_Int8;   break;
+		case EUnderlyingEnumType::int16:       UnderlyingProperty.Type = CPT_Int16;  break;
+		case EUnderlyingEnumType::int32:       UnderlyingProperty.Type = CPT_Int;    break;
+		case EUnderlyingEnumType::int64:       UnderlyingProperty.Type = CPT_Int64;  break;
+		case EUnderlyingEnumType::uint8:       UnderlyingProperty.Type = CPT_Byte;   break;
+		case EUnderlyingEnumType::uint16:      UnderlyingProperty.Type = CPT_UInt16; break;
+		case EUnderlyingEnumType::uint32:      UnderlyingProperty.Type = CPT_UInt32; break;
+		case EUnderlyingEnumType::uint64:      UnderlyingProperty.Type = CPT_UInt64; break;
+		case EUnderlyingEnumType::Unspecified:
+			UnderlyingProperty.Type = CPT_Int;
+			UnderlyingProperty.IntType = EIntType::Unsized;
+			break;
+
+		default:
+			check(false);
+		}
+
+		FEnumProperty* Result = new FEnumProperty(Scope, Name, ObjectFlags);
+		FUnrealPropertyDefinitionInfo& SubProp = FPropertyTraits::CreateProperty(UnderlyingProperty, Result, TEXT("UnderlyingType"), ObjectFlags, VariableCategory, Dimensions, PropDef.GetUnrealSourceFile(), PropDef.GetLineNumber(), PropDef.GetParsePosition());
+		PropDef.SetValuePropDef(SubProp);
+		Result->UnderlyingProp = CastFieldChecked<FNumericProperty>(SubProp.GetProperty());
+		Result->Enum = VarProperty.Enum;
+		return Result;
+	}
+
+	static bool IsSupportedByBlueprint(const FUnrealPropertyDefinitionInfo& PropDef, bool bMemberVariable)
+	{
+		return true;
+	}
+};
+
+//----------------------------------------------------------------------------------------------------------------------------------------------------------
+//----------------------------------------------------------------------------------------------------------------------------------------------------------
+// Object types
+//----------------------------------------------------------------------------------------------------------------------------------------------------------
+//----------------------------------------------------------------------------------------------------------------------------------------------------------
+
+/**
+ * Base class for all object based property types
+ */
+struct FPropertyTypeTraitsObjectBase : public FPropertyTypeTraitsBase
+{
+	static constexpr bool bIsObject = true;
+
+	static bool DefaultValueStringCppFormatToInnerFormat(const FUnrealPropertyDefinitionInfo& PropDef, const FString& CppForm, FString& OutForm)
+	{
+		const bool bIsNull = FDefaultValueHelper::Is(CppForm, TEXT("NULL")) || FDefaultValueHelper::Is(CppForm, TEXT("nullptr")) || FDefaultValueHelper::Is(CppForm, TEXT("0"));
+		if (bIsNull)
+		{
+			OutForm = TEXT("None");
+		}
+		return bIsNull; // always return as null is the only the processing we can do for object defaults
+	}
+};
+
+struct FPropertyTypeTraitsObjectReference : public FPropertyTypeTraitsObjectBase
+{
+	static FProperty* CreateEngineType(FUnrealPropertyDefinitionInfo& PropDef, FFieldVariant Scope, const FName& Name, EObjectFlags ObjectFlags, EVariableCategory::Type VariableCategory, const TCHAR* Dimensions)
+	{
+		FPropertyBase& VarProperty = PropDef.GetPropertyBase();
+		check(VarProperty.PropertyClass);
+
+#if UHT_ENABLE_PTR_PROPERTY_TAG
+		PropDef.GetUnrealSourceFile().AddTypeDefIncludeIfNeeded(VarProperty.PropertyClass);
+#endif
+
+		if (VarProperty.PropertyClass->IsChildOf(UClass::StaticClass()))
+		{
+			FClassProperty* Result = new FClassProperty(Scope, Name, ObjectFlags);
+			Result->MetaClass = VarProperty.MetaClass;
+			Result->PropertyClass = VarProperty.PropertyClass;
+			return Result;
+		}
+		else
+		{
+			//ETSTODO - This prevents us from having PropDef const
+			if (FUnrealClassDefinitionInfo::HierarchyHasAnyClassFlags(VarProperty.PropertyClass, CLASS_DefaultToInstanced))
+			{
+				VarProperty.PropertyFlags |= CPF_InstancedReference;
+				AddEditInlineMetaData(VarProperty.MetaData);
+			}
+
+			FObjectProperty* Result = new FObjectProperty(Scope, Name, ObjectFlags);
+			Result->PropertyClass = VarProperty.PropertyClass;
+			return Result;
+		}
+	}
+
+	static bool IsSupportedByBlueprint(const FUnrealPropertyDefinitionInfo& PropDef, bool bMemberVariable)
+	{
+		return true;
+	}
+};
+
+struct FPropertyTypeTraitsWeakObjectReference : public FPropertyTypeTraitsObjectBase
+{
+	static FProperty* CreateEngineType(FUnrealPropertyDefinitionInfo& PropDef, FFieldVariant Scope, const FName& Name, EObjectFlags ObjectFlags, EVariableCategory::Type VariableCategory, const TCHAR* Dimensions)
+	{
+		const FPropertyBase& VarProperty = PropDef.GetPropertyBase();
+		check(VarProperty.PropertyClass);
+
+#if UHT_ENABLE_PTR_PROPERTY_TAG
+		PropDef.GetUnrealSourceFile().AddTypeDefIncludeIfNeeded(VarProperty.PropertyClass);
+#endif
+
+		FWeakObjectProperty* Result = new FWeakObjectProperty(Scope, Name, ObjectFlags);
+		Result->PropertyClass = VarProperty.PropertyClass;
+		return Result;
+	}
+
+	static bool IsSupportedByBlueprint(const FUnrealPropertyDefinitionInfo& PropDef, bool bMemberVariable)
+	{
+		return bMemberVariable;
+	}
+};
+
+struct FPropertyTypeTraitsLazyObjectReference : public FPropertyTypeTraitsObjectBase
+{
+	static FProperty* CreateEngineType(FUnrealPropertyDefinitionInfo& PropDef, FFieldVariant Scope, const FName& Name, EObjectFlags ObjectFlags, EVariableCategory::Type VariableCategory, const TCHAR* Dimensions)
+	{
+		const FPropertyBase& VarProperty = PropDef.GetPropertyBase();
+		check(VarProperty.PropertyClass);
+
+#if UHT_ENABLE_PTR_PROPERTY_TAG
+		PropDef.GetUnrealSourceFile().AddTypeDefIncludeIfNeeded(VarProperty.PropertyClass);
+#endif
+
+		FLazyObjectProperty* Result = new FLazyObjectProperty(Scope, Name, ObjectFlags);
+		Result->PropertyClass = VarProperty.PropertyClass;
+		return Result;
+	}
+};
+
+struct FPropertyTypeTraitsObjectPtrReference : public FPropertyTypeTraitsObjectBase
+{
+	static FProperty* CreateEngineType(FUnrealPropertyDefinitionInfo& PropDef, FFieldVariant Scope, const FName& Name, EObjectFlags ObjectFlags, EVariableCategory::Type VariableCategory, const TCHAR* Dimensions)
+	{
+		FPropertyBase& VarProperty = PropDef.GetPropertyBase();
+		check(VarProperty.PropertyClass);
+
+#if UHT_ENABLE_PTR_PROPERTY_TAG
+		PropDef.GetUnrealSourceFile().AddTypeDefIncludeIfNeeded(VarProperty.PropertyClass);
+#endif
+
+		if (VarProperty.PropertyClass->IsChildOf(UClass::StaticClass()))
+		{
+			FClassPtrProperty* Result = new FClassPtrProperty(Scope, Name, ObjectFlags);
+			Result->MetaClass = VarProperty.MetaClass;
+			Result->PropertyClass = VarProperty.PropertyClass;
+			return Result;
+		}
+		else
+		{
+			//ETSTODO - This prevents us from having PropDef const
+			if (FUnrealClassDefinitionInfo::HierarchyHasAnyClassFlags(VarProperty.PropertyClass, CLASS_DefaultToInstanced))
+			{
+				VarProperty.PropertyFlags |= CPF_InstancedReference;
+				AddEditInlineMetaData(VarProperty.MetaData);
+			}
+
+			FObjectPtrProperty* Result = new FObjectPtrProperty(Scope, Name, ObjectFlags);
+			Result->PropertyClass = VarProperty.PropertyClass;
+			return Result;
+		}
+	}
+
+	static bool IsSupportedByBlueprint(const FUnrealPropertyDefinitionInfo& PropDef, bool bMemberVariable)
+	{
+		return true;
+	}
+};
+
+struct FPropertyTypeTraitsSoftObjectReference : public FPropertyTypeTraitsObjectBase
+{
+	static FProperty* CreateEngineType(FUnrealPropertyDefinitionInfo& PropDef, FFieldVariant Scope, const FName& Name, EObjectFlags ObjectFlags, EVariableCategory::Type VariableCategory, const TCHAR* Dimensions)
+	{
+		FPropertyBase& VarProperty = PropDef.GetPropertyBase();
+		check(VarProperty.PropertyClass);
+
+#if UHT_ENABLE_PTR_PROPERTY_TAG
+		PropDef.GetUnrealSourceFile().AddTypeDefIncludeIfNeeded(VarProperty.PropertyClass);
+#endif
+
+		if (VarProperty.PropertyClass->IsChildOf(UClass::StaticClass()))
+		{
+			FSoftClassProperty* Result = new FSoftClassProperty(Scope, Name, ObjectFlags);
+			Result->MetaClass = VarProperty.MetaClass;
+			Result->PropertyClass = VarProperty.PropertyClass;
+			return Result;
+		}
+		else
+		{
+			FSoftObjectProperty* Result = new FSoftObjectProperty(Scope, Name, ObjectFlags);
+			Result->PropertyClass = VarProperty.PropertyClass;
+			return Result;
+		}
+	}
+
+	static bool IsSupportedByBlueprint(const FUnrealPropertyDefinitionInfo& PropDef, bool bMemberVariable)
+	{
+		return PropDef.GetProperty()->IsA<FSoftObjectProperty>(); // Not SoftClass???
+	}
+};
+
+struct FPropertyTypeTraitsInterface : public FPropertyTypeTraitsBase
+{
+	static FProperty* CreateEngineType(FUnrealPropertyDefinitionInfo& PropDef, FFieldVariant Scope, const FName& Name, EObjectFlags ObjectFlags, EVariableCategory::Type VariableCategory, const TCHAR* Dimensions)
+	{
+		FPropertyBase& VarProperty = PropDef.GetPropertyBase();
+		check(VarProperty.PropertyClass);
+		check(VarProperty.PropertyClass->HasAnyClassFlags(CLASS_Interface));
+
+#if UHT_ENABLE_PTR_PROPERTY_TAG
+		PropDef.GetUnrealSourceFile().AddTypeDefIncludeIfNeeded(VarProperty.PropertyClass);
+#endif
+
+		FInterfaceProperty* Result = new  FInterfaceProperty(Scope, Name, ObjectFlags);
+		Result->InterfaceClass = VarProperty.PropertyClass;
+		return Result;
+	}
+
+	static bool IsSupportedByBlueprint(const FUnrealPropertyDefinitionInfo& PropDef, bool bMemberVariable)
+	{
+		return true;
+	}
+};
+
+//----------------------------------------------------------------------------------------------------------------------------------------------------------
+//----------------------------------------------------------------------------------------------------------------------------------------------------------
+// Other types
+//----------------------------------------------------------------------------------------------------------------------------------------------------------
+//----------------------------------------------------------------------------------------------------------------------------------------------------------
+
+struct FPropertyTypeTraitsName : public FPropertyTypeTraitsBase
+{
+	static bool DefaultValueStringCppFormatToInnerFormat(const FUnrealPropertyDefinitionInfo& PropDef, const FString& CppForm, FString& OutForm)
+	{
+		if (FDefaultValueHelper::Is(CppForm, TEXT("NAME_None")))
+		{
+			OutForm = TEXT("None");
+			return true;
+		}
+		return FDefaultValueHelper::StringFromCppString(CppForm, TEXT("FName"), OutForm);
+	}
+
+	static FProperty* CreateEngineType(FUnrealPropertyDefinitionInfo& PropDef, FFieldVariant Scope, const FName& Name, EObjectFlags ObjectFlags, EVariableCategory::Type VariableCategory, const TCHAR* Dimensions)
+	{
+		FNameProperty* Result = new FNameProperty(Scope, Name, ObjectFlags);
+		return Result;
+	}
+
+	static bool IsSupportedByBlueprint(const FUnrealPropertyDefinitionInfo& PropDef, bool bMemberVariable)
+	{
+		return true;
+	}
+};
+
+struct FPropertyTypeTraitsString : public FPropertyTypeTraitsBase
+{
+	static bool DefaultValueStringCppFormatToInnerFormat(const FUnrealPropertyDefinitionInfo& PropDef, const FString& CppForm, FString& OutForm)
+	{
+		return FDefaultValueHelper::StringFromCppString(CppForm, TEXT("FString"), OutForm);
+	}
+
+	static FProperty* CreateEngineType(FUnrealPropertyDefinitionInfo& PropDef, FFieldVariant Scope, const FName& Name, EObjectFlags ObjectFlags, EVariableCategory::Type VariableCategory, const TCHAR* Dimensions)
+	{
+		FStrProperty* Result = new FStrProperty(Scope, Name, ObjectFlags);
+		return Result;
+	}
+
+	static bool IsSupportedByBlueprint(const FUnrealPropertyDefinitionInfo& PropDef, bool bMemberVariable)
+	{
+		return true;
+	}
+};
+
+struct FPropertyTypeTraitsText : public FPropertyTypeTraitsBase
+{
+	static bool DefaultValueStringCppFormatToInnerFormat(const FUnrealPropertyDefinitionInfo& PropDef, const FString& CppForm, FString& OutForm)
+	{
+		// Handle legacy cases of FText::FromString being used as default values
+		// These should be replaced with INVTEXT as FText::FromString can produce inconsistent keys
+		if (FDefaultValueHelper::StringFromCppString(CppForm, TEXT("FText::FromString"), OutForm))
+		{
+			UE_LOG_WARNING_UHT(TEXT("FText::FromString should be replaced with INVTEXT for default parameter values"));
+			return true;
+		}
+
+		// Parse the potential value into an instance
+		FText ParsedText;
+		if (FDefaultValueHelper::Is(CppForm, TEXT("FText()")) || FDefaultValueHelper::Is(CppForm, TEXT("FText::GetEmpty()")))
+		{
+			ParsedText = FText::GetEmpty();
+		}
+		else
+		{
+			static const FString UHTDummyNamespace = TEXT("__UHT_DUMMY_NAMESPACE__");
+
+			if (!FTextStringHelper::ReadFromBuffer(*CppForm, ParsedText, *UHTDummyNamespace, nullptr, /*bRequiresQuotes*/true))
+			{
+				return false;
+			}
+
+			// If the namespace of the parsed text matches the default we gave then this was a LOCTEXT macro which we 
+			// don't allow in default values as they rely on an external macro that is known to C++ but not to UHT
+			// TODO: UHT could parse these if it tracked the current LOCTEXT_NAMESPACE macro as it parsed
+			if (TOptional<FString> ParsedTextNamespace = FTextInspector::GetNamespace(ParsedText))
+			{
+				if (ParsedTextNamespace.GetValue().Equals(UHTDummyNamespace))
+				{
+					FError::Throwf(TEXT("LOCTEXT default parameter values are not supported; use NSLOCTEXT instead: %s \"%s\" "), *PropDef.GetProperty()->GetName(), *CppForm);
+				}
+			}
+		}
+
+		// Normalize the default value from the parsed value
+		FTextStringHelper::WriteToBuffer(OutForm, ParsedText, /*bRequiresQuotes*/false);
+		return true;
+	}
+
+	static FProperty* CreateEngineType(FUnrealPropertyDefinitionInfo& PropDef, FFieldVariant Scope, const FName& Name, EObjectFlags ObjectFlags, EVariableCategory::Type VariableCategory, const TCHAR* Dimensions)
+	{
+		FTextProperty* Result = new FTextProperty(Scope, Name, ObjectFlags);
+		return Result;
+	}
+
+	static bool IsSupportedByBlueprint(const FUnrealPropertyDefinitionInfo& PropDef, bool bMemberVariable)
+	{
+		return true;
+	}
+};
+
+struct FPropertyTypeTraitsStruct : public FPropertyTypeTraitsBase
+{
+	static bool DefaultValueStringCppFormatToInnerFormat(const FUnrealPropertyDefinitionInfo& PropDef, const FString& CppForm, FString& OutForm)
+	{
+		// Cache off the struct types, in case we need them later
+		UPackage* CoreUObjectPackage = UObject::StaticClass()->GetOutermost();
+		static const UScriptStruct* VectorStruct = FClasses::FindObjectChecked<UScriptStruct>(CoreUObjectPackage, TEXT("Vector"));
+		static const UScriptStruct* Vector2DStruct = FClasses::FindObjectChecked<UScriptStruct>(CoreUObjectPackage, TEXT("Vector2D"));
+		static const UScriptStruct* RotatorStruct = FClasses::FindObjectChecked<UScriptStruct>(CoreUObjectPackage, TEXT("Rotator"));
+		static const UScriptStruct* LinearColorStruct = FClasses::FindObjectChecked<UScriptStruct>(CoreUObjectPackage, TEXT("LinearColor"));
+		static const UScriptStruct* ColorStruct = FClasses::FindObjectChecked<UScriptStruct>(CoreUObjectPackage, TEXT("Color"));
+
+		UScriptStruct* Struct = PropDef.GetPropertyBase().Struct;
+		if (Struct == VectorStruct)
+		{
+			FString Parameters;
+			if (FDefaultValueHelper::Is(CppForm, TEXT("FVector::ZeroVector")))
+			{
+				return true;
+			}
+			else if (FDefaultValueHelper::Is(CppForm, TEXT("FVector::UpVector")))
+			{
+				OutForm = FString::Printf(TEXT("%f,%f,%f"),
+					FVector::UpVector.X, FVector::UpVector.Y, FVector::UpVector.Z);
+			}
+			else if (FDefaultValueHelper::Is(CppForm, TEXT("FVector::ForwardVector")))
+			{
+				OutForm = FString::Printf(TEXT("%f,%f,%f"),
+					FVector::ForwardVector.X, FVector::ForwardVector.Y, FVector::ForwardVector.Z);
+			}
+			else if (FDefaultValueHelper::Is(CppForm, TEXT("FVector::RightVector")))
+			{
+				OutForm = FString::Printf(TEXT("%f,%f,%f"),
+					FVector::RightVector.X, FVector::RightVector.Y, FVector::RightVector.Z);
+			}
+			else if (FDefaultValueHelper::GetParameters(CppForm, TEXT("FVector"), Parameters))
+			{
+				if (FDefaultValueHelper::Is(Parameters, TEXT("ForceInit")))
+				{
+					return true;
+				}
+				FVector Vector;
+				float Value;
+				if (FDefaultValueHelper::ParseVector(Parameters, Vector))
+				{
+					OutForm = FString::Printf(TEXT("%f,%f,%f"),
+						Vector.X, Vector.Y, Vector.Z);
+				}
+				else if (FDefaultValueHelper::ParseFloat(Parameters, Value))
+				{
+					OutForm = FString::Printf(TEXT("%f,%f,%f"),
+						Value, Value, Value);
+				}
+			}
+		}
+		else if (Struct == RotatorStruct)
+		{
+			if (FDefaultValueHelper::Is(CppForm, TEXT("FRotator::ZeroRotator")))
+			{
+				return true;
+			}
+			FString Parameters;
+			if (FDefaultValueHelper::GetParameters(CppForm, TEXT("FRotator"), Parameters))
+			{
+				if (FDefaultValueHelper::Is(Parameters, TEXT("ForceInit")))
+				{
+					return true;
+				}
+				FRotator Rotator;
+				if (FDefaultValueHelper::ParseRotator(Parameters, Rotator))
+				{
+					OutForm = FString::Printf(TEXT("%f,%f,%f"),
+						Rotator.Pitch, Rotator.Yaw, Rotator.Roll);
+				}
+			}
+		}
+		else if (Struct == Vector2DStruct)
+		{
+			if (FDefaultValueHelper::Is(CppForm, TEXT("FVector2D::ZeroVector")))
+			{
+				return true;
+			}
+			if (FDefaultValueHelper::Is(CppForm, TEXT("FVector2D::UnitVector")))
+			{
+				OutForm = FString::Printf(TEXT("(X=%3.3f,Y=%3.3f)"),
+					FVector2D::UnitVector.X, FVector2D::UnitVector.Y);
+			}
+			FString Parameters;
+			if (FDefaultValueHelper::GetParameters(CppForm, TEXT("FVector2D"), Parameters))
+			{
+				if (FDefaultValueHelper::Is(Parameters, TEXT("ForceInit")))
+				{
+					return true;
+				}
+				FVector2D Vector2D;
+				if (FDefaultValueHelper::ParseVector2D(Parameters, Vector2D))
+				{
+					OutForm = FString::Printf(TEXT("(X=%3.3f,Y=%3.3f)"),
+						Vector2D.X, Vector2D.Y);
+				}
+			}
+		}
+		else if (Struct == LinearColorStruct)
+		{
+			if (FDefaultValueHelper::Is(CppForm, TEXT("FLinearColor::White")))
+			{
+				OutForm = FLinearColor::White.ToString();
+			}
+			else if (FDefaultValueHelper::Is(CppForm, TEXT("FLinearColor::Gray")))
+			{
+				OutForm = FLinearColor::Gray.ToString();
+			}
+			else if (FDefaultValueHelper::Is(CppForm, TEXT("FLinearColor::Black")))
+			{
+				OutForm = FLinearColor::Black.ToString();
+			}
+			else if (FDefaultValueHelper::Is(CppForm, TEXT("FLinearColor::Transparent")))
+			{
+				OutForm = FLinearColor::Transparent.ToString();
+			}
+			else if (FDefaultValueHelper::Is(CppForm, TEXT("FLinearColor::Red")))
+			{
+				OutForm = FLinearColor::Red.ToString();
+			}
+			else if (FDefaultValueHelper::Is(CppForm, TEXT("FLinearColor::Green")))
+			{
+				OutForm = FLinearColor::Green.ToString();
+			}
+			else if (FDefaultValueHelper::Is(CppForm, TEXT("FLinearColor::Blue")))
+			{
+				OutForm = FLinearColor::Blue.ToString();
+			}
+			else if (FDefaultValueHelper::Is(CppForm, TEXT("FLinearColor::Yellow")))
+			{
+				OutForm = FLinearColor::Yellow.ToString();
+			}
+			else
+			{
+				FString Parameters;
+				if (FDefaultValueHelper::GetParameters(CppForm, TEXT("FLinearColor"), Parameters))
+				{
+					if (FDefaultValueHelper::Is(Parameters, TEXT("ForceInit")))
+					{
+						return true;
+					}
+					FLinearColor Color;
+					if (FDefaultValueHelper::ParseLinearColor(Parameters, Color))
+					{
+						OutForm = Color.ToString();
+					}
+				}
+			}
+		}
+		else if (Struct == ColorStruct)
+		{
+			if (FDefaultValueHelper::Is(CppForm, TEXT("FColor::White")))
+			{
+				OutForm = FColor::White.ToString();
+			}
+			else if (FDefaultValueHelper::Is(CppForm, TEXT("FColor::Black")))
+			{
+				OutForm = FColor::Black.ToString();
+			}
+			else if (FDefaultValueHelper::Is(CppForm, TEXT("FColor::Red")))
+			{
+				OutForm = FColor::Red.ToString();
+			}
+			else if (FDefaultValueHelper::Is(CppForm, TEXT("FColor::Green")))
+			{
+				OutForm = FColor::Green.ToString();
+			}
+			else if (FDefaultValueHelper::Is(CppForm, TEXT("FColor::Blue")))
+			{
+				OutForm = FColor::Blue.ToString();
+			}
+			else if (FDefaultValueHelper::Is(CppForm, TEXT("FColor::Yellow")))
+			{
+				OutForm = FColor::Yellow.ToString();
+			}
+			else if (FDefaultValueHelper::Is(CppForm, TEXT("FColor::Cyan")))
+			{
+				OutForm = FColor::Cyan.ToString();
+			}
+			else if (FDefaultValueHelper::Is(CppForm, TEXT("FColor::Magenta")))
+			{
+				OutForm = FColor::Magenta.ToString();
+			}
+			else
+			{
+				FString Parameters;
+				if (FDefaultValueHelper::GetParameters(CppForm, TEXT("FColor"), Parameters))
+				{
+					if (FDefaultValueHelper::Is(Parameters, TEXT("ForceInit")))
+					{
+						return true;
+					}
+					FColor Color;
+					if (FDefaultValueHelper::ParseColor(Parameters, Color))
+					{
+						OutForm = Color.ToString();
+					}
+				}
+			}
+		}
+		return !OutForm.IsEmpty();
+	}
+
+	static FProperty* CreateEngineType(FUnrealPropertyDefinitionInfo& PropDef, FFieldVariant Scope, const FName& Name, EObjectFlags ObjectFlags, EVariableCategory::Type VariableCategory, const TCHAR* Dimensions)
+	{
+		FPropertyBase& VarProperty = PropDef.GetPropertyBase();
+
+#if UHT_ENABLE_VALUE_PROPERTY_TAG
+		PropDef.GetUnrealSourceFile().AddTypeDefIncludeIfNeeded(VarProperty.Struct);
+#endif
+
+		//ETSTODO - This is preventing PropDef from being const
+		if (VarProperty.Struct->StructFlags & STRUCT_HasInstancedReference)
+		{
+			VarProperty.PropertyFlags |= CPF_ContainsInstancedReference;
+		}
+
+		FStructProperty* Result = new FStructProperty(Scope, Name, ObjectFlags);
+		Result->Struct = VarProperty.Struct;
+		return Result;
+	}
+
+	static bool IsSupportedByBlueprint(const FUnrealPropertyDefinitionInfo& PropDef, bool bMemberVariable)
+	{
+		return CastField<const FStructProperty>(PropDef.GetProperty())->Struct->GetBoolMetaDataHierarchical(FHeaderParserNames::NAME_BlueprintType);
+	}
+};
+
+struct FPropertyTypeTraitsDelegate : public FPropertyTypeTraitsBase
+{
+	static FProperty* CreateEngineType(FUnrealPropertyDefinitionInfo& PropDef, FFieldVariant Scope, const FName& Name, EObjectFlags ObjectFlags, EVariableCategory::Type VariableCategory, const TCHAR* Dimensions)
+	{
+		FDelegateProperty* Result = new FDelegateProperty(Scope, Name, ObjectFlags);
+		return Result;
+	}
+
+	static bool IsSupportedByBlueprint(const FUnrealPropertyDefinitionInfo& PropDef, bool bMemberVariable)
+	{
+		return true;
+	}
+};
+
+struct FPropertyTypeTraitsMulticastDelegate : public FPropertyTypeTraitsBase
+{
+	static FProperty* CreateEngineType(FUnrealPropertyDefinitionInfo& PropDef, FFieldVariant Scope, const FName& Name, EObjectFlags ObjectFlags, EVariableCategory::Type VariableCategory, const TCHAR* Dimensions)
+	{
+		const FPropertyBase& VarProperty = PropDef.GetPropertyBase();
+
+		FMulticastDelegateProperty* Result;
+		if (VarProperty.Function->IsA<USparseDelegateFunction>())
+		{
+			Result = new FMulticastSparseDelegateProperty(Scope, Name, ObjectFlags);
+		}
+		else
+		{
+			Result = new FMulticastInlineDelegateProperty(Scope, Name, ObjectFlags);
+		}
+		return Result;
+	}
+
+	static bool IsSupportedByBlueprint(const FUnrealPropertyDefinitionInfo& PropDef, bool bMemberVariable)
+	{
+		return bMemberVariable;
+	}
+};
+
+struct FPropertyTypeTraitsFieldPath : public FPropertyTypeTraitsBase
+{
+	static FProperty* CreateEngineType(FUnrealPropertyDefinitionInfo& PropDef, FFieldVariant Scope, const FName& Name, EObjectFlags ObjectFlags, EVariableCategory::Type VariableCategory, const TCHAR* Dimensions)
+	{
+		const FPropertyBase& VarProperty = PropDef.GetPropertyBase();
+		FFieldPathProperty* Result = new FFieldPathProperty(Scope, Name, ObjectFlags);
+		Result->PropertyClass = VarProperty.PropertyPathClass;
+		return Result;
+	}
+
+	static bool IsSupportedByBlueprint(const FUnrealPropertyDefinitionInfo& PropDef, bool bMemberVariable)
+	{
+		return true;
+	}
+};
+
+//----------------------------------------------------------------------------------------------------------------------------------------------------------
+//----------------------------------------------------------------------------------------------------------------------------------------------------------
+// Container types
+//----------------------------------------------------------------------------------------------------------------------------------------------------------
+//----------------------------------------------------------------------------------------------------------------------------------------------------------
+
+struct FPropertyTypeTraitsStaticArray : public FPropertyTypeTraitsBase
+{
+	static FProperty* CreateEngineType(FUnrealPropertyDefinitionInfo& PropDef, FFieldVariant Scope, const FName& Name, EObjectFlags ObjectFlags, EVariableCategory::Type VariableCategory, const TCHAR* Dimensions)
+	{
+		FProperty* Property = CreatePropertyHelper<false>(PropDef, Scope, Name, ObjectFlags, VariableCategory, Dimensions);
+		Property->ArrayDim = 2;
+		PropDef.SetArrayDimensions(Dimensions);
+		return Property;
+	}
+
+	static bool IsSupportedByBlueprint(const FUnrealPropertyDefinitionInfo& PropDef, bool bMemberVariable)
+	{
+		return IsSupportedByBlueprintSansContainers(PropDef, bMemberVariable);
+	}
+};
+
+struct FPropertyTypeTraitsDynamicArray : public FPropertyTypeTraitsBase
+{
+	static FProperty* CreateEngineType(FUnrealPropertyDefinitionInfo& PropDef, FFieldVariant Scope, const FName& Name, EObjectFlags ObjectFlags, EVariableCategory::Type VariableCategory, const TCHAR* Dimensions)
+	{
+		FPropertyBase& VarProperty = PropDef.GetPropertyBase();
+
+		FArrayProperty* Array = new FArrayProperty(Scope, Name, ObjectFlags);
+
+		FPropertyBase InnerVarProperty = PropDef.GetPropertyBase();
+		InnerVarProperty.ArrayType = EArrayType::None;
+		FUnrealPropertyDefinitionInfo& InnerPropDef = FPropertyTraits::CreateProperty(InnerVarProperty, Array, Name, RF_Public, VariableCategory, Dimensions, PropDef.GetUnrealSourceFile(), PropDef.GetLineNumber(), PropDef.GetParsePosition());
+		FProperty* InnerProp = InnerPropDef.GetProperty();
+
+		Array->Inner = InnerProp;
+		VarProperty.PropertyFlags = Array->Inner->PropertyFlags;
+		VarProperty.MetaData = MoveTemp(InnerPropDef.GetPropertyBase().MetaData);
+		PropDef.SetAllocatorType(VarProperty.AllocatorType);
+		PropDef.SetValuePropDef(InnerPropDef);
+
+		// Propagate flags
+		InnerPropDef.GetPropertyBase().PropertyFlags = InnerProp->PropertyFlags = InnerProp->PropertyFlags & CPF_PropagateToArrayInner;
+
+		PropagateFlagsFromInnerAndHandlePersistentInstanceMetadata(VarProperty.PropertyFlags, PropDef.GetPropertyBase().MetaData, InnerPropDef);
+		return Array;
+	}
+
+	static bool IsSupportedByBlueprint(const FUnrealPropertyDefinitionInfo& PropDef, bool bMemberVariable)
+	{
+		return IsSupportedByBlueprintSansContainers(PropDef.GetValuePropDef(), false);
+	}
+};
+
+struct FPropertyTypeTraitsSet : public FPropertyTypeTraitsBase
+{
+	static FProperty* CreateEngineType(FUnrealPropertyDefinitionInfo& PropDef, FFieldVariant Scope, const FName& Name, EObjectFlags ObjectFlags, EVariableCategory::Type VariableCategory, const TCHAR* Dimensions)
+	{
+		FPropertyBase& VarProperty = PropDef.GetPropertyBase();
+
+		FSetProperty* Set = new FSetProperty(Scope, Name, ObjectFlags);
+
+		FPropertyBase InnerVarProperty = PropDef.GetPropertyBase();
+		InnerVarProperty.ArrayType = EArrayType::None;
+		FUnrealPropertyDefinitionInfo& InnerPropDef = FPropertyTraits::CreateProperty(InnerVarProperty, Set, Name, RF_Public, VariableCategory, Dimensions, PropDef.GetUnrealSourceFile(), PropDef.GetLineNumber(), PropDef.GetParsePosition());
+		FProperty* InnerProp = InnerPropDef.GetProperty();
+
+		Set->ElementProp = InnerProp;
+		VarProperty.PropertyFlags = InnerProp->PropertyFlags;
+		VarProperty.MetaData = MoveTemp(InnerPropDef.GetPropertyBase().MetaData);
+		PropDef.SetValuePropDef(InnerPropDef);
+
+		// Propagate flags
+		InnerPropDef.GetPropertyBase().PropertyFlags = InnerProp->PropertyFlags = InnerProp->PropertyFlags & CPF_PropagateToSetElement;
+
+		PropagateFlagsFromInnerAndHandlePersistentInstanceMetadata(VarProperty.PropertyFlags, PropDef.GetPropertyBase().MetaData, InnerPropDef);
+		return Set;
+	}
+
+	static bool IsSupportedByBlueprint(const FUnrealPropertyDefinitionInfo& PropDef, bool bMemberVariable)
+	{
+		return IsSupportedByBlueprintSansContainers(PropDef.GetValuePropDef(), false);
+	}
+};
+
+struct FPropertyTypeTraitsMap : public FPropertyTypeTraitsBase
+{
+	static FProperty* CreateEngineType(FUnrealPropertyDefinitionInfo& PropDef, FFieldVariant Scope, const FName& Name, EObjectFlags ObjectFlags, EVariableCategory::Type VariableCategory, const TCHAR* Dimensions)
+	{
+		FPropertyBase& VarProperty = PropDef.GetPropertyBase();
+
+		FMapProperty* Map = new FMapProperty(Scope, Name, ObjectFlags);
+
+		FUnrealPropertyDefinitionInfo& KeyPropDef = FPropertyTraits::CreateProperty(*PropDef.GetPropertyBase().MapKeyProp, Map, *(Name.ToString() + TEXT("_Key")), RF_Public, VariableCategory, Dimensions, PropDef.GetUnrealSourceFile(), PropDef.GetLineNumber(), PropDef.GetParsePosition());
+		FProperty* KeyProp = KeyPropDef.GetProperty();
+		FPropertyBase ValueVarProperty = PropDef.GetPropertyBase();
+		ValueVarProperty.ArrayType = EArrayType::None;
+		ValueVarProperty.MapKeyProp = nullptr;
+		FUnrealPropertyDefinitionInfo& ValuePropDef = FPropertyTraits::CreateProperty(ValueVarProperty, Map, Name, RF_Public, VariableCategory, Dimensions, PropDef.GetUnrealSourceFile(), PropDef.GetLineNumber(), PropDef.GetParsePosition());
+		FProperty* ValueProp = ValuePropDef.GetProperty();
+
+		Map->KeyProp = KeyProp;
+		Map->ValueProp = ValueProp;
+		VarProperty.PropertyFlags = ValueProp->PropertyFlags;
+		VarProperty.MetaData = MoveTemp(ValuePropDef.GetPropertyBase().MetaData);
+		PropDef.SetAllocatorType(VarProperty.AllocatorType);
+		PropDef.SetKeyPropDef(KeyPropDef);
+		PropDef.SetValuePropDef(ValuePropDef);
+
+		// Propagate flags
+		KeyPropDef.GetPropertyBase().PropertyFlags = KeyProp->PropertyFlags = KeyProp->PropertyFlags & CPF_PropagateToMapKey;
+		ValuePropDef.GetPropertyBase().PropertyFlags = ValueProp->PropertyFlags = ValueProp->PropertyFlags & CPF_PropagateToMapValue;
+
+		PropagateFlagsFromInnerAndHandlePersistentInstanceMetadata(VarProperty.PropertyFlags, PropDef.GetPropertyBase().MapKeyProp->MetaData, KeyPropDef);
+		PropagateFlagsFromInnerAndHandlePersistentInstanceMetadata(VarProperty.PropertyFlags, PropDef.GetPropertyBase().MetaData, ValuePropDef);
+		return Map;
+	}
+
+	static bool IsSupportedByBlueprint(const FUnrealPropertyDefinitionInfo& PropDef, bool bMemberVariable)
+	{
+		return
+			IsSupportedByBlueprintSansContainers(PropDef.GetValuePropDef(), false) &&
+			IsSupportedByBlueprintSansContainers(PropDef.GetKeyPropDef(), false);
+	}
+};
+
+//----------------------------------------------------------------------------------------------------------------------------------------------------------
+//----------------------------------------------------------------------------------------------------------------------------------------------------------
+// APIs
+//----------------------------------------------------------------------------------------------------------------------------------------------------------
+//----------------------------------------------------------------------------------------------------------------------------------------------------------
+
+bool FPropertyTraits::DefaultValueStringCppFormatToInnerFormat(const FUnrealPropertyDefinitionInfo& PropDef, const FString& CppForm, FString& OutForm)
+{
+	OutForm = FString();
+	if (CppForm.IsEmpty())
+	{
+		return false;
+	}
+
+	return PropertyTypeDispatch<DefaultValueStringCppFormatToInnerFormatDispatch, false, bool>(PropDef.GetPropertyBase(), std::ref(PropDef), std::ref(CppForm), std::ref(OutForm));
+}
+
+bool FPropertyTraits::IsObject(EPropertyType PropertyType)
+{
+	return PropertyTypeDispatch<IsObjectDispatch, bool>(PropertyType);
+}
+
+FUnrealPropertyDefinitionInfo& FPropertyTraits::CreateProperty(const FPropertyBase& VarProperty, FFieldVariant Scope, const FName& Name, EObjectFlags ObjectFlags, EVariableCategory::Type VariableCategory,
+	const TCHAR* Dimensions, FUnrealSourceFile& SourceFile, int LineNumber, int ParsePosition)
+{
+	TSharedRef<FUnrealPropertyDefinitionInfo> PropDefRef = MakeShared<FUnrealPropertyDefinitionInfo>(SourceFile, LineNumber, ParsePosition, VarProperty, Name.ToString());
+	FUnrealPropertyDefinitionInfo& PropDef = *PropDefRef;
+
+	FProperty* Property = CreatePropertyHelper<true>(PropDef, Scope, Name, ObjectFlags, VariableCategory, Dimensions);
+
+	PropDef.SetProperty(Property);
+	GTypeDefinitionInfoMap.Add(PropDef.GetProperty(), MoveTemp(PropDefRef));
+	return PropDef;
+}
+
+bool FPropertyTraits::IsSupportedByBlueprint(const FUnrealPropertyDefinitionInfo& PropDef, bool bMemberVariable)
+{
+	return PropertyTypeDispatch<IsSupportedByBlueprintDispatch, true, bool>(PropDef.GetPropertyBase(), std::ref(PropDef), bMemberVariable);
+}
diff --git a/Engine/Source/Programs/UnrealHeaderTool/Private/PropertyTypes.h b/Engine/Source/Programs/UnrealHeaderTool/Private/PropertyTypes.h
new file mode 100644
--- /dev/null
+++ b/Engine/Source/Programs/UnrealHeaderTool/Private/PropertyTypes.h
@@ -0,0 +1,56 @@
+// Copyright Epic Games, Inc. All Rights Reserved.
+
+#pragma once
+
+#include "CoreMinimal.h"
+#include "HeaderParser.h" // for EVariableCategory
+#include "UObject\Stack.h"
+
+class FProperty;
+class FString;
+class FUnrealPropertyDefinitionInfo;
+class FUnrealSourceFile;
+
+struct FPropertyTraits
+{
+	/**
+	 * Return true if this property type is an object property
+	 */
+	static bool IsObject(EPropertyType PropertyType);
+
+	/**
+	 * Transforms CPP-formated string containing default value, to inner formated string
+	 * If it cannot be transformed empty string is returned.
+	 *
+	 * @param PropDef The property that owns the default value.
+	 * @param CppForm A CPP-formated string.
+	 * @param out InnerForm Inner formated string
+	 * @return true on success, false otherwise.
+	 */
+	static bool DefaultValueStringCppFormatToInnerFormat(const FUnrealPropertyDefinitionInfo& PropDef, const FString& CppForm, FString& InnerForm);
+
+	/**
+	 * Given a property definition token, create the property definition and then underlying engine FProperty
+	 *
+	 * @param Token The definition of the property
+	 * @param Scope The parent object owning the property
+	 * @param Name The name of the property
+	 * @param ObjectFlags The flags associated with the property
+	 * @param VariableCategory The parsing context of the property
+	 * @param Dimensions When this is a static array, this represents the dimensions value
+	 * @param SourceFile The source file containing the property
+	 * @param LineNumber Line number of the property
+	 * @param ParsePosition Character position of the property in the header
+	 * @return The pointer to the newly created property.  It will be attached to the definition by the caller
+	 */
+	static FUnrealPropertyDefinitionInfo& CreateProperty(const FPropertyBase& VarProperty, FFieldVariant Scope, const FName& Name, EObjectFlags ObjectFlags, EVariableCategory::Type VariableCategory, const TCHAR* Dimensions, FUnrealSourceFile& SourceFile, int LineNumber, int ParsePosition);
+
+	/**
+	 * Test to see if the property can be used in a blueprint
+	 * 
+	 * @param PropDef The property in question
+	 * @param bMemberVariable If true, this is a member variable being tested
+	 * @return Return true if the property is supported in blueprints
+	 */
+	static bool IsSupportedByBlueprint(const FUnrealPropertyDefinitionInfo& PropDef, bool bMemberVariable);
+};
diff --git a/Engine/Source/Programs/UnrealHeaderTool/Private/Scope.cpp b/Engine/Source/Programs/UnrealHeaderTool/Private/Scope.cpp
--- a/Engine/Source/Programs/UnrealHeaderTool/Private/Scope.cpp
+++ b/Engine/Source/Programs/UnrealHeaderTool/Private/Scope.cpp
@@ -122,11 +122,6 @@ const FUnrealFieldDefinitionInfo* FScope::FindTypeByName(FName Name) const
 	return nullptr;
 }
 
-bool FScope::ContainsType(FUnrealFieldDefinitionInfo* Type)
-{
-	return FindTypeByName(Type->GetField()->GetFName()) != nullptr;
-}
-
 bool FScope::IsFileScope() const
 {
 	return Parent == nullptr;
@@ -170,4 +165,4 @@ FName FFileScope::GetName() const
 FName FStructScope::GetName() const
 {
 	return Struct->GetFName();
-}
\ No newline at end of file
+}
diff --git a/Engine/Source/Programs/UnrealHeaderTool/Private/Scope.h b/Engine/Source/Programs/UnrealHeaderTool/Private/Scope.h
--- a/Engine/Source/Programs/UnrealHeaderTool/Private/Scope.h
+++ b/Engine/Source/Programs/UnrealHeaderTool/Private/Scope.h
@@ -103,12 +103,6 @@ public:
 	{
 		return FindTypeByName(FName(Name, FNAME_Find));
 	}
-	/**
-	 * Checks if scope contains this type.
-	 *
-	 * @param Type Type to look for.
-	 */
-	bool ContainsType(FUnrealFieldDefinitionInfo* Type);
 
 	/**
 	 * Checks if scope contains type that satisfies a predicate.
diff --git a/Engine/Source/Programs/UnrealHeaderTool/Private/UnrealHeaderTool.h b/Engine/Source/Programs/UnrealHeaderTool/Private/UnrealHeaderTool.h
--- a/Engine/Source/Programs/UnrealHeaderTool/Private/UnrealHeaderTool.h
+++ b/Engine/Source/Programs/UnrealHeaderTool/Private/UnrealHeaderTool.h
@@ -17,4 +17,5 @@ struct FHeaderParserNames
 	static const FName NAME_ShowCategories;
 	static const FName NAME_SparseClassDataTypes;
 	static const FName NAME_AdvancedClassDisplay;
+	static const FName NAME_BlueprintType;
 };
\ No newline at end of file
diff --git a/Engine/Source/Programs/UnrealHeaderTool/Private/UnrealSourceFile.cpp b/Engine/Source/Programs/UnrealHeaderTool/Private/UnrealSourceFile.cpp
--- a/Engine/Source/Programs/UnrealHeaderTool/Private/UnrealSourceFile.cpp
+++ b/Engine/Source/Programs/UnrealHeaderTool/Private/UnrealSourceFile.cpp
@@ -145,7 +145,7 @@ bool FUnrealSourceFile::HasChanged() const
 
 FString FUnrealSourceFile::GetFileDefineName() const
 {
-	const FString API = FPackageName::GetShortName(Package).ToUpper();
+	const FString API = FPackageName::GetShortName(GetPackage()).ToUpper();
 	return FString::Printf(TEXT("%s_%s_generated_h"), *API, *GetStrippedFilename());
 }
 
diff --git a/Engine/Source/Programs/UnrealHeaderTool/Private/UnrealSourceFile.h b/Engine/Source/Programs/UnrealHeaderTool/Private/UnrealSourceFile.h
--- a/Engine/Source/Programs/UnrealHeaderTool/Private/UnrealSourceFile.h
+++ b/Engine/Source/Programs/UnrealHeaderTool/Private/UnrealSourceFile.h
@@ -38,10 +38,10 @@ class FUnrealSourceFile : public TSharedFromThis<FUnrealSourceFile>
 {
 public:
 	// Constructor.
-	FUnrealSourceFile(UPackage* InPackage, const FString& InFilename)
+	FUnrealSourceFile(FUnrealPackageDefinitionInfo& InPackageDef, const FString& InFilename)
 		: Scope                (MakeShareable(new FFileScope(*(FString(TEXT("__")) + FPaths::GetBaseFilename(InFilename) + FString(TEXT("__File"))), this)))
 		, Filename             (InFilename)
-		, Package              (InPackage)
+		, PackageDef           (InPackageDef)
 	{
 		if (GetStrippedFilename() != "NoExportTypes")
 		{
@@ -207,12 +207,20 @@ public:
 		return Scope;
 	}
 
+	/**
+	 * Gets package definition this file is in.
+	 */
+	FUnrealPackageDefinitionInfo& GetPackageDef() const
+	{
+		return PackageDef;
+	}
+
 	/**
 	 * Gets package this file is in.
 	 */
 	UPackage* GetPackage() const
 	{
-		return Package;
+		return PackageDef.GetPackage();
 	}
 
 	/**
@@ -458,7 +466,7 @@ private:
 	mutable FString FileId;
 
 	// Package of this file.
-	UPackage* Package;
+	FUnrealPackageDefinitionInfo& PackageDef;
 
 	// File name of the generated header file associated with this file.
 	FString GeneratedFilename;
diff --git a/Engine/Source/Programs/UnrealHeaderTool/Private/UnrealTypeDefinitionInfo.cpp b/Engine/Source/Programs/UnrealHeaderTool/Private/UnrealTypeDefinitionInfo.cpp
--- a/Engine/Source/Programs/UnrealHeaderTool/Private/UnrealTypeDefinitionInfo.cpp
+++ b/Engine/Source/Programs/UnrealHeaderTool/Private/UnrealTypeDefinitionInfo.cpp
@@ -1,6 +1,7 @@
 // Copyright Epic Games, Inc. All Rights Reserved.
 
 #include "UnrealTypeDefinitionInfo.h"
+#include "ClassMaps.h"
 #include "NativeClassExporter.h"
 #include "Scope.h"
 #include "UnrealSourceFile.h"
@@ -118,16 +119,6 @@ TSharedRef<FScope> FUnrealTypeDefinitionInfo::GetScope()
 	return GetUnrealSourceFile().GetScope();
 }
 
-// Constructor
-FUnrealPackageDefinitionInfo::FUnrealPackageDefinitionInfo(const FManifestModule& InModule, UPackage* InPackage)
-	: FUnrealObjectDefinitionInfo(FString())
-	, Module(InModule)
-	, ShortUpperName(FPackageName::GetShortName(InPackage).ToUpper())
-	, API(FString::Printf(TEXT("%s_API "), *ShortUpperName))
-{
-	SetObject(InPackage);
-}
-
 void FUnrealTypeDefinitionInfo::SetHash(uint32 InHash)
 {
 	Hash = InHash;
@@ -158,6 +149,33 @@ void FUnrealTypeDefinitionInfo::GetHashTag(FUHTStringBuilder& Out) const
 	}
 }
 
+void FUnrealPropertyDefinitionInfo::AddMetaData(TMap<FName, FString>&& InMetaData)
+{
+	// only add if we have some!
+	if (InMetaData.Num())
+	{
+		check(Property);
+
+		UPackage* Package = Property->GetOutermost();
+		// get (or create) a metadata object for this package
+		UMetaData* MetaData = Package->GetMetaData();
+
+		for (TPair<FName, FString>& MetaKeyValue : InMetaData)
+		{
+			Property->SetMetaData(MetaKeyValue.Key, MoveTemp(MetaKeyValue.Value));
+		}
+	}
+}
+
+FUnrealPackageDefinitionInfo::FUnrealPackageDefinitionInfo(const FManifestModule& InModule, UPackage* InPackage)
+	: FUnrealObjectDefinitionInfo(FString())
+	, Module(InModule)
+	, ShortUpperName(FPackageName::GetShortName(InPackage).ToUpper())
+	, API(FString::Printf(TEXT("%s_API "), *ShortUpperName))
+{
+	SetObject(InPackage);
+}
+
 void FUnrealPackageDefinitionInfo::PostParseFinalize()
 {
 	UPackage* Package = GetPackage();
@@ -210,6 +228,82 @@ void FUnrealFieldDefinitionInfo::AddCrossModuleReference(TSet<FString>* UniqueCr
 	}
 }
 
+void FUnrealFieldDefinitionInfo::AddMetaData(TMap<FName, FString>&& InMetaData)
+{
+	// only add if we have some!
+	if (InMetaData.Num())
+	{
+		UField* Field = GetField();
+		check(Field);
+
+		// get (or create) a metadata object for this package
+		UMetaData* MetaData = Field->GetOutermost()->GetMetaData();
+		TMap<FName, FString>* ExistingMetaData = MetaData->GetMapForObject(Field);
+		if (ExistingMetaData && ExistingMetaData->Num())
+		{
+			// Merge the existing metadata
+			TMap<FName, FString> MergedMetaData;
+			MergedMetaData.Reserve(InMetaData.Num() + ExistingMetaData->Num());
+			MergedMetaData.Append(*ExistingMetaData);
+			MergedMetaData.Append(InMetaData);
+			MetaData->SetObjectValues(Field, MoveTemp(MergedMetaData));
+		}
+		else
+		{
+			// set the metadata for this field
+			MetaData->SetObjectValues(Field, MoveTemp(InMetaData));
+		}
+	}
+}
+
+FUnrealEnumDefinitionInfo::FUnrealEnumDefinitionInfo(FUnrealSourceFile& InSourceFile, int32 InLineNumber, FString&& InNameCPP)
+	: FUnrealFieldDefinitionInfo(InSourceFile, InLineNumber, MoveTemp(InNameCPP), InSourceFile.GetPackageDef())
+{ }
+
+void FUnrealStructDefinitionInfo::AddProperty(FUnrealPropertyDefinitionInfo& PropertyDef)
+{
+	Properties.Add(&PropertyDef);
+
+	// update the optimization flags
+	if (!bContainsDelegates)
+	{
+		FProperty* Prop = PropertyDef.GetProperty();
+		if (Prop->IsA(FDelegateProperty::StaticClass()) || Prop->IsA(FMulticastDelegateProperty::StaticClass()))
+		{
+			bContainsDelegates = true;
+		}
+		else
+		{
+			FArrayProperty* ArrayProp = CastField<FArrayProperty>(Prop);
+			if (ArrayProp != NULL)
+			{
+				if (ArrayProp->Inner->IsA(FDelegateProperty::StaticClass()) || ArrayProp->Inner->IsA(FMulticastDelegateProperty::StaticClass()))
+				{
+					bContainsDelegates = true;
+				}
+			}
+		}
+	}
+}
+
+void FUnrealStructDefinitionInfo::AddFunction(FUnrealFunctionDefinitionInfo& FunctionDef)
+{
+	Functions.Add(&FunctionDef);
+
+	// update the optimization flags
+	if (!bContainsDelegates)
+	{
+		if (FunctionDef.GetFunction()->HasAnyFunctionFlags(FUNC_Delegate))
+		{
+			bContainsDelegates = true;
+		}
+	}
+}
+
+FUnrealScriptStructDefinitionInfo::FUnrealScriptStructDefinitionInfo(FUnrealSourceFile& InSourceFile, int32 InLineNumber, FString&& InNameCPP)
+	: FUnrealStructDefinitionInfo(InSourceFile, InLineNumber, MoveTemp(InNameCPP), InSourceFile.GetPackageDef())
+{ }
+
 uint32 FUnrealScriptStructDefinitionInfo::GetHash(bool bIncludeNoExport) const
 {
 	if (!bIncludeNoExport)
@@ -249,6 +343,16 @@ void FUnrealStructDefinitionInfo::SetObject(UObject* InObject)
 	}
 }
 
+FUnrealClassDefinitionInfo::FUnrealClassDefinitionInfo(FUnrealSourceFile& InSourceFile, int32 InLineNumber, FString&& InNameCPP, bool bInIsInterface)
+	: FUnrealStructDefinitionInfo(InSourceFile, InLineNumber, MoveTemp(InNameCPP), InSourceFile.GetPackageDef())
+	, bIsInterface(bInIsInterface)
+{
+	if (bInIsInterface)
+	{
+		GetStructMetaData().ParsedInterface = EParsedInterface::ParsedUInterface;
+	}
+}
+
 uint32 FUnrealClassDefinitionInfo::GetHash(bool bIncludeNoExport) const
 {
 	if (!bIncludeNoExport)
@@ -273,4 +377,35 @@ void FUnrealClassDefinitionInfo::PostParseFinalize()
 	}
 
 	FUnrealStructDefinitionInfo::PostParseFinalize();
-}
\ No newline at end of file
+}
+
+bool FUnrealClassDefinitionInfo::HierarchyHasAnyClassFlags(UClass* Class, EClassFlags FlagsToCheck)
+{
+	if (FUnrealClassDefinitionInfo* ClassDef = UHTCast<FUnrealClassDefinitionInfo>(GTypeDefinitionInfoMap.FindChecked(Class)))
+	{
+		return ClassDef->HierarchyHasAnyClassFlags(CLASS_DefaultToInstanced);
+	}
+	return false;
+}
+
+bool FUnrealClassDefinitionInfo::HierarchyHasAllClassFlags(UClass* Class, EClassFlags FlagsToCheck)
+{
+	if (FUnrealClassDefinitionInfo* ClassDef = UHTCast<FUnrealClassDefinitionInfo>(GTypeDefinitionInfoMap.FindChecked(Class)))
+	{
+		return ClassDef->HierarchyHasAllClassFlags(CLASS_DefaultToInstanced);
+	}
+	return false;
+}
+
+void FUnrealFunctionDefinitionInfo::AddProperty(FUnrealPropertyDefinitionInfo& PropertyDef)
+{
+	const FProperty* Prop = PropertyDef.GetProperty();
+	check((Prop->PropertyFlags & CPF_Parm) != 0);
+
+	if ((Prop->PropertyFlags & CPF_ReturnParm) != 0)
+	{
+		check(ReturnProperty == nullptr);
+		ReturnProperty = &PropertyDef;
+	}
+	FUnrealStructDefinitionInfo::AddProperty(PropertyDef);
+}
diff --git a/Engine/Source/Programs/UnrealHeaderTool/Private/UnrealTypeDefinitionInfo.h b/Engine/Source/Programs/UnrealHeaderTool/Private/UnrealTypeDefinitionInfo.h
--- a/Engine/Source/Programs/UnrealHeaderTool/Private/UnrealTypeDefinitionInfo.h
+++ b/Engine/Source/Programs/UnrealHeaderTool/Private/UnrealTypeDefinitionInfo.h
@@ -3,8 +3,9 @@
 #pragma once
 
 #include "CoreMinimal.h"
-#include "Templates/SharedPointer.h"
 #include "ClassDeclarationMetaData.h"
+#include "Exceptions.h"
+#include "Templates/SharedPointer.h"
 #include <atomic>
 
 #include "ParserHelper.h" // TEMPORARY
@@ -259,6 +260,14 @@ public:
 		return LineNumber;
 	}
 
+	/**
+	 * Sets the input line in the rare case where the definition is created before fully parsed (sparse delegates)
+	 */
+	void SetLineNumber(int32 InLineNumber)
+	{
+		LineNumber = InLineNumber;
+	}
+
 	/**
 	 * Gets the reference to FUnrealSourceFile object that stores information about
 	 * source file this type was defined in.
@@ -290,6 +299,14 @@ public:
 	 */
 	void GetHashTag(FUHTStringBuilder& Out) const;
 
+	/**
+	 * Add meta data for the given definition
+	 */
+	virtual void AddMetaData(TMap<FName, FString>&& InMetaData)
+	{
+		FUHTException::Throwf(*this, TEXT("Meta data can not be set for a definition of this type."));
+	}
+
 protected:
 	explicit FUnrealTypeDefinitionInfo(FString&& InNameCPP)
 		: NameCPP(MoveTemp(InNameCPP))
@@ -311,25 +328,14 @@ private:
 /**
  * Class that stores information about type definitions derived from UProperty
  */
-class FUnrealPropertyDefinitionInfo 
+class FUnrealPropertyDefinitionInfo
 	: public FUnrealTypeDefinitionInfo
 {
 public:
-	FUnrealPropertyDefinitionInfo(FUnrealSourceFile& InSourceFile, int32 InLineNumber, FProperty* InProperty)
-		: FUnrealTypeDefinitionInfo(InSourceFile, InLineNumber, FString())
-		, Property(InProperty)
-	{ }
-
-	FUnrealPropertyDefinitionInfo(FUnrealSourceFile& InSourceFile, int32 InLineNumber, FProperty* InProperty, bool bInUnsized)
-		: FUnrealTypeDefinitionInfo(InSourceFile, InLineNumber, FString())
-		, Property(InProperty)
-		, bUnsized(bInUnsized)
-	{ }
-
-	FUnrealPropertyDefinitionInfo(FUnrealSourceFile& InSourceFile, int32 InLineNumber, FProperty* InProperty, EAllocatorType InAllocatorType)
-		: FUnrealTypeDefinitionInfo(InSourceFile, InLineNumber, FString())
-		, Property(InProperty)
-		, AllocatorType(InAllocatorType)
+	FUnrealPropertyDefinitionInfo(FUnrealSourceFile& InSourceFile, int32 InLineNumber, int32 InParsePosition, const FPropertyBase& InVarProperty, FString&& InNameCPP)
+		: FUnrealTypeDefinitionInfo(InSourceFile, InLineNumber, MoveTemp(InNameCPP))
+		, PropertyBase(InVarProperty)
+		, ParsePosition(InParsePosition)
 	{ }
 
 	virtual FUnrealPropertyDefinitionInfo* AsProperty() override
@@ -354,6 +360,14 @@ public:
 		return Property;
 	}
 
+	/**
+	 * Sets the engine type
+	 */
+	void SetProperty(FProperty* InProperty)
+	{
+		Property = InProperty;
+	}
+
 	/**
 	 * Set the string that represents the array dimensions
 	 */
@@ -379,6 +393,14 @@ public:
 		return bUnsized;
 	}
 
+	/**
+	 * Set the unsized flag
+	 */
+	void SetUnsized(bool bInUnsized)
+	{
+		bUnsized = bInUnsized;
+	}
+
 	/**
 	 * Return the allocator type
 	 */
@@ -387,22 +409,105 @@ public:
 		return AllocatorType;
 	}
 
+	/**
+	 * Set the allocator type
+	 */
+	void SetAllocatorType(EAllocatorType InAllocatorType)
+	{
+		AllocatorType = InAllocatorType;
+	}
+
 	/**
 	 * Return the token associated with the property parsing
 	 */
-	FToken& GetToken()
+	FPropertyBase& GetPropertyBase()
+	{
+		return PropertyBase;
+	}
+	const FPropertyBase& GetPropertyBase() const
+	{
+		return PropertyBase;
+	}
+
+	/**
+	 * Add meta data for the given definition
+	 */
+	virtual void AddMetaData(TMap<FName, FString>&& InMetaData) override;
+
+	/**
+	 * Get the associated key property definition (valid for maps)
+	 */
+	FUnrealPropertyDefinitionInfo& GetKeyPropDef() const
+	{
+		check(KeyPropDef);
+		return *KeyPropDef;
+	}
+
+	/**
+	 * Sets the associated key property definition (valid for maps)
+	 */
+	void SetKeyPropDef(FUnrealPropertyDefinitionInfo& InKeyPropDef)
+	{
+		check(KeyPropDef == nullptr);
+		KeyPropDef = &InKeyPropDef;
+	}
+
+	/**
+	 * Get the associated value property definition (valid for maps, sets, and dynamic arrays)
+	 */
+	FUnrealPropertyDefinitionInfo& GetValuePropDef() const
 	{
-		return Token;
+		check(ValuePropDef);
+		return *ValuePropDef;
 	}
-	const FToken& GetToken() const
+
+	/**
+	 * Sets the associated value property definition (valid for maps, sets, and dynamic arrays)
+	 */
+	void SetValuePropDef(FUnrealPropertyDefinitionInfo& InValuePropDef)
+	{
+		check(ValuePropDef == nullptr);
+		ValuePropDef = &InValuePropDef;
+	}
+
+	/**
+	 * Get the parsing position of the property
+	 */
+	int32 GetParsePosition() const
+	{
+		return ParsePosition;
+	}
+
+	/**
+	 * Determines whether this property's type is compatible with another property's type.
+	 *
+	 * @param	Other							the property to check against this one.
+	 *											Given the following example expressions, VarA is Other and VarB is 'this':
+	 *												VarA = VarB;
+	 *
+	 *												function func(type VarB) {}
+	 *												func(VarA);
+	 *
+	 *												static operator==(type VarB_1, type VarB_2) {}
+	 *												if ( VarA_1 == VarA_2 ) {}
+	 *
+	 * @param	bDisallowGeneralization			controls whether it should be considered a match if this property's type is a generalization
+	 *											of the other property's type (or vice versa, when dealing with structs
+	 * @param	bIgnoreImplementedInterfaces	controls whether two types can be considered a match if one type is an interface implemented
+	 *											by the other type.
+	 */
+	bool MatchesType(const FUnrealPropertyDefinitionInfo& Other, bool bDisallowGeneralization, bool bIgnoreImplementedInterfaces = false) const
 	{
-		return Token;
+		return GetPropertyBase().MatchesType(Other.GetPropertyBase(), bDisallowGeneralization, bIgnoreImplementedInterfaces);
 	}
 
 private:
-	FToken Token;
+	FPropertyBase PropertyBase;
 	FString ArrayDimensions;
+	FUnrealPropertyDefinitionInfo* KeyPropDef = nullptr;
+	FUnrealPropertyDefinitionInfo* ValuePropDef = nullptr;
 	FProperty* Property = nullptr;
+	int32 ParsePosition;
 	EAllocatorType AllocatorType = EAllocatorType::Default;
 	bool bUnsized = false;
 };
@@ -439,10 +544,27 @@ public:
 		Object = InObject;
 	}
 
+	/**
+	 * Return the "outer" object that contains the definition for this object
+	 */
+	FUnrealObjectDefinitionInfo* GetOuter() const
+	{
+		return Outer;
+	}
+
 protected:
-	using FUnrealTypeDefinitionInfo::FUnrealTypeDefinitionInfo;
+	// Used only by packages
+	explicit FUnrealObjectDefinitionInfo(FString&& InNameCPP)
+		: FUnrealTypeDefinitionInfo(MoveTemp(InNameCPP))
+	{ }
+
+	FUnrealObjectDefinitionInfo(FUnrealSourceFile& InSourceFile, int32 InLineNumber, FString&& InNameCPP, FUnrealObjectDefinitionInfo& InOuter)
+		: FUnrealTypeDefinitionInfo(InSourceFile, InLineNumber, MoveTemp(InNameCPP))
+		, Outer(&InOuter)
+	{ }
 
 private:
+	FUnrealObjectDefinitionInfo* Outer = nullptr;
 	UObject* Object = nullptr;
 };
 
@@ -637,6 +759,11 @@ public:
 		return ExternDecl[bRequiresValidObject];
 	}
 
+	/**
+	 * Add meta data for the given definition
+	 */
+	virtual void AddMetaData(TMap<FName, FString>&& InMetaData) override;
+
 private:
 	FString SingletonName[2];
 	FString SingletonNameChopped[2];
@@ -649,9 +776,7 @@ private:
 class FUnrealEnumDefinitionInfo : public FUnrealFieldDefinitionInfo
 {
 public:
-	FUnrealEnumDefinitionInfo(FUnrealSourceFile& InSourceFile, int32 InLineNumber, FString&& InNameCPP)
-		: FUnrealFieldDefinitionInfo(InSourceFile, InLineNumber, MoveTemp(InNameCPP))
-	{ }
+	FUnrealEnumDefinitionInfo(FUnrealSourceFile& InSourceFile, int32 InLineNumber, FString&& InNameCPP);
 
 	virtual FUnrealEnumDefinitionInfo* AsEnum() override
 	{
@@ -746,23 +871,37 @@ public:
 	/**
 	 * Add a new property to the structure
 	 */
-	virtual void AddProperty(FUnrealPropertyDefinitionInfo& PropertyDef)
-	{
-		Properties.Add(&PropertyDef);
-	}
+	virtual void AddProperty(FUnrealPropertyDefinitionInfo& PropertyDef);
 
 	/**
 	 * Get the collection of properties
 	 */
-	const TArray<FUnrealPropertyDefinitionInfo*> GetProperties() const
+	const TArray<FUnrealPropertyDefinitionInfo*>& GetProperties() const
 	{
 		return Properties;
 	}
-	TArray<FUnrealPropertyDefinitionInfo*> GetParameters()
+	TArray<FUnrealPropertyDefinitionInfo*>& GetProperties()
 	{
 		return Properties;
 	}
 
+	/**
+	 * Add a new function to the structure
+	 */
+	virtual void AddFunction(FUnrealFunctionDefinitionInfo& FunctionDef);
+
+	/**
+	 * Get the collection of functions
+	 */
+	const TArray<FUnrealFunctionDefinitionInfo*>& GetFunctions() const
+	{
+		return Functions;
+	}
+	TArray<FUnrealFunctionDefinitionInfo*>& GetFunctions()
+	{
+		return Functions;
+	}
+
 	/**
 	 * Return the class meta data information
 	 */
@@ -774,25 +913,72 @@ public:
 	/**
 	 * Return the super class information
 	 */
+	const FBaseClassInfo& GetSuperClassInfo() const
+	{
+		return SuperClassInfo;
+	}
+
 	FBaseClassInfo& GetSuperClassInfo()
 	{
 		return SuperClassInfo;
 	}
 
+	const TArray<FBaseClassInfo>& GetBaseClassInfo() const
+	{
+		return BaseClassInfo;
+	}
+
 	TArray<FBaseClassInfo>& GetBaseClassInfo()
 	{
 		return BaseClassInfo;
 	}
 
+	/**
+	 * Get the contains delegates flag
+	 */
+	bool ContainsDelegates() const
+	{
+		return bContainsDelegates;
+	}
+
+	/**
+	 * Sets contains delegates flag for this class.
+	 */
+	void MarkContainsDelegate()
+	{
+		bContainsDelegates = true;
+	}
+
+	/**
+	 * Test to see if this struct is a child of another struct
+	 */
+	bool IsChildOf(FUnrealStructDefinitionInfo& ParentStruct) const
+	{
+		for (const FUnrealStructDefinitionInfo* Current = this; Current; Current = Current->GetSuperClassInfo().Struct)
+		{
+			if (Current == &ParentStruct)
+			{
+				return true;
+			}
+		}
+		return false;
+	}
+
 private:
 	TSharedPtr<FScope> StructScope;
 
 	/** Properties of the structure */
 	TArray<FUnrealPropertyDefinitionInfo*> Properties;
 
+	/** Functions of the structure */
+	TArray<FUnrealFunctionDefinitionInfo*> Functions;
+
 	FStructMetaData StructMetaData;
 	FBaseClassInfo SuperClassInfo;
 	TArray<FBaseClassInfo> BaseClassInfo;
+
+	/** whether this struct declares delegate functions or properties */
+	bool bContainsDelegates = false;
 };
 
 /**
@@ -801,9 +987,7 @@ private:
 class FUnrealScriptStructDefinitionInfo : public FUnrealStructDefinitionInfo
 {
 public:
-	FUnrealScriptStructDefinitionInfo(FUnrealSourceFile& InSourceFile, int32 InLineNumber, FString&& InNameCPP)
-		: FUnrealStructDefinitionInfo(InSourceFile, InLineNumber, MoveTemp(InNameCPP))
-	{ }
+	FUnrealScriptStructDefinitionInfo(FUnrealSourceFile& InSourceFile, int32 InLineNumber, FString&& InNameCPP);
 
 	virtual FUnrealScriptStructDefinitionInfo* AsScriptStruct() override
 	{
@@ -835,8 +1019,8 @@ public:
 class FUnrealFunctionDefinitionInfo : public FUnrealStructDefinitionInfo
 {
 public:
-	FUnrealFunctionDefinitionInfo(FUnrealSourceFile& InSourceFile, int32 InLineNumber, FString&& InNameCPP, FFuncInfo&& InFuncInfo)
-		: FUnrealStructDefinitionInfo(InSourceFile, InLineNumber, MoveTemp(InNameCPP))
+	FUnrealFunctionDefinitionInfo(FUnrealSourceFile& InSourceFile, int32 InLineNumber, FString&& InNameCPP, FUnrealObjectDefinitionInfo& InOuter, FFuncInfo&& InFuncInfo)
+		: FUnrealStructDefinitionInfo(InSourceFile, InLineNumber, MoveTemp(InNameCPP), InOuter)
 		, FunctionData(MoveTemp(InFuncInfo))
 	{ }
 
@@ -868,35 +1052,12 @@ public:
 	FUnrealPropertyDefinitionInfo* GetReturn() const { return ReturnProperty; }
 	//@}
 
-	void UpdateFunctionData(FFuncInfo& UpdatedFuncData)
-	{
-		//@TODO: UCREMOVAL: Some more thorough evaluation should be done here
-		FunctionData.FunctionFlags |= UpdatedFuncData.FunctionFlags;
-		FunctionData.FunctionExportFlags |= UpdatedFuncData.FunctionExportFlags;
-	}
-
 	/**
 	 * Adds a new function property to be tracked.  Determines whether the property is a
 	 * function parameter, local property, or return value, and adds it to the appropriate
 	 * list
-	 *
-	 * @param	PropertyToken	the property to add
 	 */
-	virtual void AddProperty(FUnrealPropertyDefinitionInfo& PropertyDef) override
-	{
-		const FProperty* Prop = PropertyDef.GetProperty();
-		check((Prop->PropertyFlags & CPF_Parm) != 0);
-
-		if ((Prop->PropertyFlags & CPF_ReturnParm) != 0)
-		{
-			check(ReturnProperty == nullptr);
-			ReturnProperty = &PropertyDef;
-		}
-		else
-		{
-			FUnrealStructDefinitionInfo::AddProperty(PropertyDef);
-		}
-	}
+	virtual void AddProperty(FUnrealPropertyDefinitionInfo& PropertyDef) override;
 
 	/**
 	 * Sets the specified function export flags
@@ -931,15 +1092,7 @@ class FUnrealClassDefinitionInfo
 	, public FClassDeclarationMetaData
 {
 public:
-	FUnrealClassDefinitionInfo(FUnrealSourceFile& InSourceFile, int32 InLineNumber, FString&& InNameCPP, bool bInIsInterface)
-		: FUnrealStructDefinitionInfo(InSourceFile, InLineNumber, MoveTemp(InNameCPP))
-		, bIsInterface(bInIsInterface)
-	{
-		if (bInIsInterface)
-		{
-			GetStructMetaData().ParsedInterface = EParsedInterface::ParsedUInterface;
-		}
-	}
+	FUnrealClassDefinitionInfo(FUnrealSourceFile& InSourceFile, int32 InLineNumber, FString&& InNameCPP, bool bInIsInterface);
 
 	FUnrealClassDefinitionInfo(FString&& InNameCPP)
 		: FUnrealStructDefinitionInfo(MoveTemp(InNameCPP))
@@ -973,6 +1126,15 @@ public:
 		return static_cast<UClass*>(GetObject());
 	}
 
+	FUnrealClassDefinitionInfo* GetSuperClass() const
+	{
+		if (FUnrealStructDefinitionInfo* SuperClass = GetSuperClassInfo().Struct)
+		{
+			return &SuperClass->AsClassChecked();
+		}
+		return nullptr;
+	}
+
 	/**
 	 * Get the archive type
 	 */
@@ -1013,8 +1175,372 @@ public:
 		return bIsInterface;
 	}
 
+	/**
+	 * Used to safely check whether the passed in flag is set.
+	 *
+	 * @param	FlagsToCheck		Class flag(s) to check for
+	 * @return	true if the passed in flag is set, false otherwise
+	 *			(including no flag passed in, unless the FlagsToCheck is CLASS_AllFlags)
+	 */
+	FORCEINLINE bool HasAnyClassFlags(EClassFlags FlagsToCheck) const
+	{
+		return EnumHasAnyFlags(ClassFlags, FlagsToCheck) != 0 || GetClass()->HasAnyClassFlags(FlagsToCheck);
+	}
+
+	/**
+	 * Used to safely check whether all of the passed in flags are set.
+	 *
+	 * @param FlagsToCheck	Class flags to check for
+	 * @return true if all of the passed in flags are set (including no flags passed in), false otherwise
+	 */
+	FORCEINLINE bool HasAllClassFlags(EClassFlags FlagsToCheck) const
+	{
+		return EnumHasAllFlags(ClassFlags, FlagsToCheck) || GetClass()->HasAllClassFlags(FlagsToCheck);
+	}
+
+	/**
+	 * Gets the class flags.
+	 *
+	 * @return	The class flags.
+	 */
+	FORCEINLINE EClassFlags GetClassFlags() const
+	{
+		return ClassFlags;
+	}
+
+	/**
+	 * Used to safely check whether the passed in flag is set in the whole hierarchy
+	 *
+	 * @param	FlagsToCheck		Class flag(s) to check for
+	 * @return	true if the passed in flag is set, false otherwise
+	 *			(including no flag passed in, unless the FlagsToCheck is CLASS_AllFlags)
+	 */
+	bool HierarchyHasAnyClassFlags(EClassFlags FlagsToCheck) const
+	{
+		for (const FUnrealClassDefinitionInfo* ClassDef = this; ClassDef; ClassDef = ClassDef->GetSuperClass())
+		{
+			if (ClassDef->HasAnyClassFlags(FlagsToCheck))
+			{
+				return true;
+			}
+		}
+		return false;
+	}
+
+	/**
+	 * Used to safely check whether the passed in flag is set in the whole hierarchy
+	 *
+	 * @param	Class				Class to check
+	 * @param	FlagsToCheck		Class flag(s) to check for
+	 * @return	true if the passed in flag is set, false otherwise
+	 *			(including no flag passed in, unless the FlagsToCheck is CLASS_AllFlags)
+	 */
+	static bool HierarchyHasAnyClassFlags(UClass* Class, EClassFlags FlagsToCheck);
+
+	/**
+	 * Used to safely check whether all of the passed in flags are set.
+	 *
+	 * @param FlagsToCheck	Class flags to check for
+	 * @return true if all of the passed in flags are set (including no flags passed in), false otherwise
+	 */
+	bool HierarchyHasAllClassFlags(EClassFlags FlagsToCheck) const
+	{
+		for (const FUnrealClassDefinitionInfo* ClassDef = this; ClassDef; ClassDef = ClassDef->GetSuperClass())
+		{
+			if (ClassDef->HasAllClassFlags(FlagsToCheck))
+			{
+				return true;
+			}
+		}
+		return false;
+	}
+	/**
+	 * Used to safely check whether all of the passed in flags are set.
+	 *
+	 * @param	Class				Class to check
+	 * @param FlagsToCheck	Class flags to check for
+	 * @return true if all of the passed in flags are set (including no flags passed in), false otherwise
+	 */
+	static bool HierarchyHasAllClassFlags(UClass* Class, EClassFlags FlagsToCheck);
+
 private:
 	FString EnclosingDefine;
 	ESerializerArchiveType ArchiveType = ESerializerArchiveType::None;
 	bool bIsInterface = false;
 };
+
+template <typename To>
+struct FUHTCastImplTo
+{};
+
+template <>
+struct FUHTCastImplTo<FUnrealTypeDefinitionInfo>
+{
+	template <typename From>
+	static FUnrealTypeDefinitionInfo* CastImpl(From& Src)
+	{
+		return &Src;
+	}
+
+	template <>
+	static FUnrealTypeDefinitionInfo* CastImpl(FUnrealTypeDefinitionInfo& Src)
+	{
+		return &Src;
+	}
+};
+
+#define UHT_CAST_IMPL(TypeName, RoutineName) \
+	template <> \
+	struct FUHTCastImplTo<TypeName> \
+	{ \
+		template <typename From> \
+		static TypeName* CastImpl(From& Src) \
+		{ \
+			return Src.RoutineName(); \
+		} \
+		template <> \
+		static TypeName* CastImpl(TypeName& Src) \
+		{ \
+			return &Src; \
+		} \
+	};
+
+UHT_CAST_IMPL(FUnrealPropertyDefinitionInfo, AsProperty);
+UHT_CAST_IMPL(FUnrealObjectDefinitionInfo, AsObject);
+UHT_CAST_IMPL(FUnrealPackageDefinitionInfo, AsPackage);
+UHT_CAST_IMPL(FUnrealFieldDefinitionInfo, AsField);
+UHT_CAST_IMPL(FUnrealEnumDefinitionInfo, AsEnum);
+UHT_CAST_IMPL(FUnrealStructDefinitionInfo, AsStruct);
+UHT_CAST_IMPL(FUnrealScriptStructDefinitionInfo, AsScriptStruct);
+UHT_CAST_IMPL(FUnrealClassDefinitionInfo, AsClass);
+UHT_CAST_IMPL(FUnrealFunctionDefinitionInfo, AsFunction);
+
+#undef UHT_CAST_IMPL
+
+template <typename To, typename From>
+To* UHTCast(TSharedRef<From>* Src)
+{
+	return Src ? FUHTCastImplTo<To>::template CastImpl<From>(**Src) : nullptr;
+}
+
+template <typename To, typename From>
+To* UHTCast(TSharedRef<From>& Src)
+{
+	return FUHTCastImplTo<To>::template CastImpl<From>(*Src);
+}
+
+template <typename To, typename From>
+To* UHTCast(From* Src)
+{
+	return Src ? FUHTCastImplTo<To>::template CastImpl<From>(*Src) : nullptr;
+}
+
+template <typename To, typename From>
+const To* UHTCast(const From* Src)
+{
+	return Src ? FUHTCastImplTo<To>::template CastImpl<From>(const_cast<From&>(*Src)) : nullptr;
+}
+
+template <typename To, typename From>
+To* UHTCast(From& Src)
+{
+	return FUHTCastImplTo<To>::template CastImpl<From>(Src);
+}
+
+template <typename To, typename From>
+const To* UHTCast(const From& Src)
+{
+	return FUHTCastImplTo<To>::template CastImpl<From>(const_cast<From&>(Src));
+}
+
+template <typename To, typename From>
+To& UHTCastChecked(TSharedRef<From>* Src)
+{
+	To* Dst = UHTCast<To, From>(Src);
+	check(Dst);
+	return *Dst;
+}
+
+template <typename To, typename From>
+To& UHTCastChecked(TSharedRef<From>& Src)
+{
+	To* Dst = UHTCast<To, From>(Src);
+	check(Dst);
+	return *Dst;
+}
+
+template <typename To, typename From>
+To& UHTCastChecked(From* Src)
+{
+	To* Dst = UHTCast<To, From>(Src);
+	check(Dst);
+	return *Dst;
+}
+
+template <typename To, typename From>
+const To& UHTCastChecked(const From* Src)
+{
+	const To* Dst = UHTCast<To, From>(Src);
+	check(Dst);
+	return *Dst;
+}
+
+template <typename To, typename From>
+To& UHTCastChecked(From& Src)
+{
+	To* Dst = UHTCast<To, From>(Src);
+	check(Dst);
+	return *Dst;
+}
+
+template <typename To, typename From>
+const To& UHTCastChecked(const From& Src)
+{
+	const To* Dst = UHTCast<To, From>(Src);
+	check(Dst);
+	return *Dst;
+}
+
+template <class T>
+const TArray<T*>& GetFieldsFromDef(const FUnrealStructDefinitionInfo& InStructDef);
+
+template <>
+inline const TArray<FUnrealPropertyDefinitionInfo*>& GetFieldsFromDef(const FUnrealStructDefinitionInfo& InStructDef)
+{
+	return InStructDef.GetProperties();
+}
+
+template <>
+inline const TArray<FUnrealFunctionDefinitionInfo*>& GetFieldsFromDef(const FUnrealStructDefinitionInfo& InStructDef)
+{
+	return InStructDef.GetFunctions();
+}
+
+//
+// For iterating through a linked list of fields.
+//
+template <class T>
+class TUHTFieldIterator
+{
+private:
+	/** The object being searched for the specified field */
+	const FUnrealStructDefinitionInfo* StructDef;
+	/** The current location in the list of fields being iterated */
+	T* const* Field;
+	T* const* End;
+	/** Whether to include the super class or not */
+	const bool bIncludeSuper;
+
+public:
+	TUHTFieldIterator(const FUnrealStructDefinitionInfo* InStructDef, EFieldIteratorFlags::SuperClassFlags InSuperClassFlags = EFieldIteratorFlags::IncludeSuper)
+		: StructDef(InStructDef)
+		, Field(UpdateRange(InStructDef))
+		, bIncludeSuper(InSuperClassFlags == EFieldIteratorFlags::IncludeSuper)
+	{
+		IterateToNext();
+	}
+
+	TUHTFieldIterator(const TUHTFieldIterator& rhs) = default;
+
+	/** conversion to "bool" returning true if the iterator is valid. */
+	FORCEINLINE explicit operator bool() const
+	{
+		return Field != NULL;
+	}
+	/** inverse of the "bool" operator */
+	FORCEINLINE bool operator !() const
+	{
+		return !(bool)*this;
+	}
+
+	inline friend bool operator==(const TUHTFieldIterator<T>& Lhs, const TUHTFieldIterator<T>& Rhs) { return Lhs.Field == Rhs.Field; }
+	inline friend bool operator!=(const TUHTFieldIterator<T>& Lhs, const TUHTFieldIterator<T>& Rhs) { return Lhs.Field != Rhs.Field; }
+
+	inline void operator++()
+	{
+		checkSlow(Field != End);
+		++Field;
+		IterateToNext();
+	}
+	inline T* operator*()
+	{
+		checkSlow(Field != End);
+		return *Field;
+	}
+	inline const T* operator*() const
+	{
+		checkSlow(Field != End);
+		return *Field;
+	}
+	inline T* operator->()
+	{
+		checkSlow(Field != End);
+		return *Field;
+	}
+	inline const FUnrealStructDefinitionInfo* GetStructDef()
+	{
+		return StructDef;
+	}
+protected:
+	inline T* const * UpdateRange(const FUnrealStructDefinitionInfo* InStructDef)
+	{
+		if (InStructDef)
+		{
+			auto& Array = GetFieldsFromDef<T>(*InStructDef);
+			T* const * Out = Array.GetData();
+			End = Out + Array.Num();
+			return Out;
+		}
+		else
+		{
+			End = nullptr;
+			return nullptr;
+		}
+	}
+
+	inline void IterateToNext()
+	{
+		T* const * CurrentField = Field;
+		const FUnrealStructDefinitionInfo* CurrentStructDef = StructDef;
+
+		while (CurrentStructDef)
+		{
+			if (CurrentField != End)
+			{
+				StructDef = CurrentStructDef;
+				Field = CurrentField;
+				return;
+			}
+
+			if (bIncludeSuper)
+			{
+				CurrentStructDef = CurrentStructDef->GetSuperClassInfo().Struct;
+				if (CurrentStructDef)
+				{
+					CurrentField = UpdateRange(CurrentStructDef);
+					continue;
+				}
+			}
+
+			break;
+		}
+
+		StructDef = CurrentStructDef;
+		Field = nullptr;
+	}
+};
+
+template <typename T>
+struct TUHTFieldRange
+{
+	TUHTFieldRange(const FUnrealStructDefinitionInfo& InStruct, EFieldIteratorFlags::SuperClassFlags InSuperClassFlags = EFieldIteratorFlags::IncludeSuper)
+		: Begin(&InStruct, InSuperClassFlags)
+	{
+	}
+
+	friend TUHTFieldIterator<T> begin(const TUHTFieldRange& Range) { return Range.Begin; }
+	friend TUHTFieldIterator<T> end(const TUHTFieldRange& Range) { return TUHTFieldIterator<T>(nullptr); }
+
+	TUHTFieldIterator<T> Begin;
+};
+
