Fixing Mac compile issue

#rnx

diff --git a/Engine/Source/Programs/UnrealHeaderTool/Private/UnrealTypeDefinitionInfo.h b/Engine/Source/Programs/UnrealHeaderTool/Private/UnrealTypeDefinitionInfo.h
--- a/Engine/Source/Programs/UnrealHeaderTool/Private/UnrealTypeDefinitionInfo.h
+++ b/Engine/Source/Programs/UnrealHeaderTool/Private/UnrealTypeDefinitionInfo.h
@@ -1277,13 +1277,13 @@ template <>
 struct FUHTCastImplTo<FUnrealTypeDefinitionInfo>
 {
 	template <typename From>
-	static FUnrealTypeDefinitionInfo* CastImpl(From& Src)
+	FUnrealTypeDefinitionInfo* CastImpl(From& Src)
 	{
 		return &Src;
 	}
 
 	template <>
-	static FUnrealTypeDefinitionInfo* CastImpl(FUnrealTypeDefinitionInfo& Src)
+	FUnrealTypeDefinitionInfo* CastImpl(FUnrealTypeDefinitionInfo& Src)
 	{
 		return &Src;
 	}
@@ -1294,12 +1294,12 @@ struct FUHTCastImplTo<FUnrealTypeDefinitionInfo>
 	struct FUHTCastImplTo<TypeName> \
 	{ \
 		template <typename From> \
-		static TypeName* CastImpl(From& Src) \
+		TypeName* CastImpl(From& Src) \
 		{ \
 			return Src.RoutineName(); \
 		} \
 		template <> \
-		static TypeName* CastImpl(TypeName& Src) \
+		TypeName* CastImpl(TypeName& Src) \
 		{ \
 			return &Src; \
 		} \
@@ -1320,37 +1320,37 @@ UHT_CAST_IMPL(FUnrealFunctionDefinitionInfo, AsFunction);
 template <typename To, typename From>
 To* UHTCast(TSharedRef<From>* Src)
 {
-	return Src ? FUHTCastImplTo<To>::template CastImpl<From>(**Src) : nullptr;
+	return Src ? FUHTCastImplTo<To>().template CastImpl<From>(**Src) : nullptr;
 }
 
 template <typename To, typename From>
 To* UHTCast(TSharedRef<From>& Src)
 {
-	return FUHTCastImplTo<To>::template CastImpl<From>(*Src);
+	return FUHTCastImplTo<To>().template CastImpl<From>(*Src);
 }
 
 template <typename To, typename From>
 To* UHTCast(From* Src)
 {
-	return Src ? FUHTCastImplTo<To>::template CastImpl<From>(*Src) : nullptr;
+	return Src ? FUHTCastImplTo<To>().template CastImpl<From>(*Src) : nullptr;
 }
 
 template <typename To, typename From>
 const To* UHTCast(const From* Src)
 {
-	return Src ? FUHTCastImplTo<To>::template CastImpl<From>(const_cast<From&>(*Src)) : nullptr;
+	return Src ? FUHTCastImplTo<To>().template CastImpl<From>(const_cast<From&>(*Src)) : nullptr;
 }
 
 template <typename To, typename From>
 To* UHTCast(From& Src)
 {
-	return FUHTCastImplTo<To>::template CastImpl<From>(Src);
+	return FUHTCastImplTo<To>().template CastImpl<From>(Src);
 }
 
 template <typename To, typename From>
 const To* UHTCast(const From& Src)
 {
-	return FUHTCastImplTo<To>::template CastImpl<From>(const_cast<From&>(Src));
+	return FUHTCastImplTo<To>().template CastImpl<From>(const_cast<From&>(Src));
 }
 
 template <typename To, typename From>
