Flip another path separator direction

diff --git a/Engine/Source/Programs/UnrealHeaderTool/Private/PropertyTypes.h b/Engine/Source/Programs/UnrealHeaderTool/Private/PropertyTypes.h
--- a/Engine/Source/Programs/UnrealHeaderTool/Private/PropertyTypes.h
+++ b/Engine/Source/Programs/UnrealHeaderTool/Private/PropertyTypes.h
@@ -4,7 +4,7 @@
 
 #include "CoreMinimal.h"
 #include "HeaderParser.h" // for EVariableCategory
-#include "UObject\Stack.h"
+#include "UObject/Stack.h"
 
 class FProperty;
 class FString;
