More UnrealHeaderTool refactoring
Remove FClass
Remove ClassDeclarationMetaData

#rnx

diff --git a/Engine/Source/Programs/UnrealHeaderTool/Private/ClassDeclarationMetaData.cpp b/Engine/Source/Programs/UnrealHeaderTool/Private/ClassDeclarationMetaData.cpp
deleted file mode 100644
diff --git a/Engine/Source/Programs/UnrealHeaderTool/Private/ClassDeclarationMetaData.h b/Engine/Source/Programs/UnrealHeaderTool/Private/ClassDeclarationMetaData.h
deleted file mode 100644
diff --git a/Engine/Source/Programs/UnrealHeaderTool/Private/Classes.cpp b/Engine/Source/Programs/UnrealHeaderTool/Private/Classes.cpp
--- a/Engine/Source/Programs/UnrealHeaderTool/Private/Classes.cpp
+++ b/Engine/Source/Programs/UnrealHeaderTool/Private/Classes.cpp
@@ -1,13 +1,13 @@
 // Copyright Epic Games, Inc. All Rights Reserved.
 
 #include "Classes.h"
-#include "ParserClass.h"
 #include "UnrealHeaderTool.h"
 #include "UObject/ErrorException.h"
 #include "UObject/Package.h"
 #include "Templates/Casts.h"
 #include "UObject/ObjectRedirector.h"
 #include "StringUtils.h"
+#include "UnrealTypeDefinitionInfo.h"
 
 namespace
 {
@@ -17,7 +17,7 @@ namespace
 	 * @param InNameToCheck - Name w/ potential prefix to check
 	 * @param OriginalClass - Class to check against
 	 */
-	bool ClassNameHasValidPrefix(const FString& InNameToCheck, const FClass* OriginalClass)
+	bool ClassNameHasValidPrefix(const FString& InNameToCheck, const UClass* OriginalClass)
 	{
 		bool bIsLabledDeprecated;
 		GetClassPrefix( InNameToCheck, bIsLabledDeprecated );
@@ -28,7 +28,7 @@ namespace
 			return true;
 		}
 
-		const FString OriginalClassName = OriginalClass->GetNameWithPrefix();
+		const FString OriginalClassName = FUnrealTypeDefinitionInfo::GetNameWithPrefix(OriginalClass);
 
 		bool bNamesMatch = (InNameToCheck == OriginalClassName);
 
@@ -45,7 +45,7 @@ namespace
 	}
 }
 
-FClass* FClasses::FindClass(const TCHAR* ClassName)
+UClass* FClasses::FindClass(const TCHAR* ClassName)
 {
 	check(ClassName);
 
@@ -53,21 +53,21 @@ FClass* FClasses::FindClass(const TCHAR* ClassName)
 
 	if (UClass* Result = FindObject<UClass>(ClassPackage, ClassName))
 	{
-		return (FClass*)Result;
+		return Result;
 	}
 
 	if (UObjectRedirector* RenamedClassRedirector = FindObject<UObjectRedirector>(ClassPackage, ClassName))
 	{
-		return (FClass*)CastChecked<UClass>(RenamedClassRedirector->DestinationObject);
+		return CastChecked<UClass>(RenamedClassRedirector->DestinationObject);
 	}
 
 	return nullptr;
 }
 
-FClass* FClasses::FindScriptClassOrThrow(const FString& InClassName)
+UClass* FClasses::FindScriptClassOrThrow(const FString& InClassName)
 {
 	FString ErrorMsg;
-	if (FClass* Result = FindScriptClass(InClassName, &ErrorMsg))
+	if (UClass* Result = FindScriptClass(InClassName, &ErrorMsg))
 	{
 		return Result;
 	}
@@ -78,33 +78,33 @@ FClass* FClasses::FindScriptClassOrThrow(const FString& InClassName)
 	return 0;
 }
 
-FClass* FClasses::FindScriptClass(const FString& InClassName, FString* OutErrorMsg)
+UClass* FClasses::FindScriptClass(const FString& InClassName, FString* OutErrorMsg)
 {
 	// Strip the class name of its prefix and then do a search for the class
 	FString ClassNameStripped = GetClassNameWithPrefixRemoved(InClassName);
-	if (FClass* FoundClass = FindClass(*ClassNameStripped))
+	if (UClass* FoundClass = FindClass(*ClassNameStripped))
 	{
 		// If the class was found with the stripped class name, verify that the correct prefix was used and throw an error otherwise
 		if (!ClassNameHasValidPrefix(InClassName, FoundClass))
 		{
 			if (OutErrorMsg)
 			{
-				*OutErrorMsg = FString::Printf(TEXT("Class '%s' has an incorrect prefix, expecting '%s'"), *InClassName, *FoundClass->GetNameWithPrefix());
+				*OutErrorMsg = FString::Printf(TEXT("Class '%s' has an incorrect prefix, expecting '%s'"), *InClassName, *FUnrealTypeDefinitionInfo::GetNameWithPrefix(FoundClass));
 			}
 			return nullptr;
 		}
 
-		return (FClass*)FoundClass;
+		return FoundClass;
 	}
 
 	// Couldn't find the class with a class name stripped of prefix (or a prefix was not found)
 	// See if the prefix was forgotten by trying to find the class with the given identifier
-	if (FClass* FoundClass = FindClass(*InClassName))
+	if (UClass* FoundClass = FindClass(*InClassName))
 	{
 		// If the class was found with the given identifier, the user forgot to use the correct Unreal prefix	
 		if (OutErrorMsg)
 		{
-			*OutErrorMsg = FString::Printf(TEXT("Class '%s' is missing a prefix, expecting '%s'"), *InClassName, *FoundClass->GetNameWithPrefix());
+			*OutErrorMsg = FString::Printf(TEXT("Class '%s' is missing a prefix, expecting '%s'"), *InClassName, *FUnrealTypeDefinitionInfo::GetNameWithPrefix(FoundClass));
 		}
 	}
 	else
diff --git a/Engine/Source/Programs/UnrealHeaderTool/Private/Classes.h b/Engine/Source/Programs/UnrealHeaderTool/Private/Classes.h
--- a/Engine/Source/Programs/UnrealHeaderTool/Private/Classes.h
+++ b/Engine/Source/Programs/UnrealHeaderTool/Private/Classes.h
@@ -7,12 +7,11 @@
 
 class UPackage;
 class FString;
-class FClass;
 
 class FClasses
 {
 public:
-	static FClass* FindClass(const TCHAR* ClassName);
+	static UClass* FindClass(const TCHAR* ClassName);
 
 	/** 
 	 * Attempts to find a script class based on the given name. Will attempt to strip
@@ -22,7 +21,7 @@ public:
 	 * @param   InClassName  Name w/ Unreal prefix to use when searching for a class
 	 * @return               The found class.
 	 */
-	static FClass* FindScriptClassOrThrow(const FString& InClassName);
+	static UClass* FindScriptClassOrThrow(const FString& InClassName);
 
 	/** 
 	 * Attempts to find a script class based on the given name. Will attempt to strip
@@ -32,7 +31,7 @@ public:
 	 * @param   OutErrorMsg  Error message (if any) giving the caller flexibility in how they present an error
 	 * @return               The found class, or NULL if the class was not found.
 	 */
-	static FClass* FindScriptClass(const FString& InClassName, FString* OutErrorMsg = nullptr);
+	static UClass* FindScriptClass(const FString& InClassName, FString* OutErrorMsg = nullptr);
 
 	/**
 	 * Find an optional object.
diff --git a/Engine/Source/Programs/UnrealHeaderTool/Private/CodeGenerator.cpp b/Engine/Source/Programs/UnrealHeaderTool/Private/CodeGenerator.cpp
--- a/Engine/Source/Programs/UnrealHeaderTool/Private/CodeGenerator.cpp
+++ b/Engine/Source/Programs/UnrealHeaderTool/Private/CodeGenerator.cpp
@@ -37,7 +37,6 @@
 #include "UnrealHeaderToolGlobals.h"
 
 #include "Exceptions.h"
-#include "ParserClass.h"
 #include "Scope.h"
 #include "HeaderProvider.h"
 #include "GeneratedCodeVersion.h"
@@ -917,7 +916,7 @@ void FNativeClassHeaderGenerator::PropertyNew(FOutputDevice& DeclOut, FOutputDev
 	FProperty*	   Prop                 = PropertyDef.GetProperty();
 	FString        PropName             = CreateUTF8LiteralString(FNativeClassHeaderGenerator::GetOverriddenName(Prop));
 	FString        PropNameDep          = Prop->HasAllPropertyFlags(CPF_Deprecated) ? Prop->GetName() + TEXT("_DEPRECATED") : Prop->GetName();
-	const TCHAR*   FPropertyObjectFlags = FClass::IsOwnedByDynamicType(Prop) ? TEXT("RF_Public|RF_Transient") : TEXT("RF_Public|RF_Transient|RF_MarkAsNative");
+	const TCHAR*   FPropertyObjectFlags = PropertyDef.IsOwnedByDynamicType() ? TEXT("RF_Public|RF_Transient") : TEXT("RF_Public|RF_Transient|RF_MarkAsNative");
 	EPropertyFlags PropFlags            = Prop->PropertyFlags & ~CPF_ComputedFlags;
 
 	FUHTStringBuilder PropTag;
@@ -1911,7 +1910,6 @@ static void FindNoExportStructsRecursive(TArray<FUnrealScriptStructDefinitionInf
 			}
 		}
 
-		//ETSTODO - Fixme
 		for (FUnrealPropertyDefinitionInfo* PropertyDef : StartDef->GetProperties())
 		{
 			FProperty* Prop = PropertyDef->GetProperty();
@@ -2042,7 +2040,7 @@ void FNativeClassHeaderGenerator::ExportNativeGeneratedInitCode(FOutputDevice& O
 	UE_CLOG(Class->ClassGeneratedBy, LogCompile, Fatal, TEXT("For intrinsic and compiled-in classes, ClassGeneratedBy should always be null"));
 
 	const bool   bIsNoExport  = Class->HasAnyClassFlags(CLASS_NoExport);
-	const bool   bIsDynamic   = FClass::IsDynamic(Class);
+	const bool   bIsDynamic   = ClassDef.IsDynamic();
 	const FString ClassNameCPP = FNameLookupCPP::GetNameCPP(Class);
 
 	const FString& ApiString = GetAPIString();
@@ -2301,7 +2299,7 @@ void FNativeClassHeaderGenerator::ExportNativeGeneratedInitCode(FOutputDevice& O
 		}
 		else
 		{
-			const FString& DynamicClassPackageName = FClass::GetTypePackageName(Class);
+			const FString& DynamicClassPackageName = ClassDef.GetTypePackageName();
 			GeneratedClassRegisterFunctionText.Logf(TEXT("\t\tUPackage* OuterPackage = FindOrConstructDynamicTypePackage(TEXT(\"%s\"));\r\n"), *DynamicClassPackageName);
 			GeneratedClassRegisterFunctionText.Logf(TEXT("\t\tUClass* OuterClass = Cast<UClass>(StaticFindObjectFast(UClass::StaticClass(), OuterPackage, TEXT(\"%s\")));\r\n"), *FNativeClassHeaderGenerator::GetOverriddenName(Class));
 			GeneratedClassRegisterFunctionText.Logf(TEXT("\t\tif (!OuterClass || !(OuterClass->ClassFlags & CLASS_Constructed))\r\n"));
@@ -2311,7 +2309,7 @@ void FNativeClassHeaderGenerator::ExportNativeGeneratedInitCode(FOutputDevice& O
 		GeneratedClassRegisterFunctionText.Logf(TEXT("\t\t\tUECodeGen_Private::ConstructUClass(OuterClass, %s::ClassParams);\r\n"), *StaticsStructName);
 
 		TArray<FString> SparseClassDataTypes;
-		((FClass*)Class)->GetSparseClassDataTypes(SparseClassDataTypes);
+		ClassDef.GetSparseClassDataTypes(SparseClassDataTypes);
 		
 		for (const FString& SparseClassDataString : SparseClassDataTypes)
 		{
@@ -2364,7 +2362,7 @@ void FNativeClassHeaderGenerator::ExportNativeGeneratedInitCode(FOutputDevice& O
 
 	// Append info for the sparse class data struct onto the text to be hashed
 	TArray<FString> SparseClassDataTypes;
-	((FClass*)Class)->GetSparseClassDataTypes(SparseClassDataTypes);
+	ClassDef.GetSparseClassDataTypes(SparseClassDataTypes);
 
 	for (const FString& SparseClassDataString : SparseClassDataTypes)
 	{
@@ -2403,7 +2401,7 @@ void FNativeClassHeaderGenerator::ExportNativeGeneratedInitCode(FOutputDevice& O
 
 	if (bIsDynamic)
 	{
-		const FString& ClassPackageName = FClass::GetTypePackageName(Class);
+		const FString& ClassPackageName = ClassDef.GetTypePackageName();
 		Out.Logf(TEXT("\tstatic FRegisterCompiledInInfo Z_CompiledInDefer_UClass_%s(%s, &%s::StaticClass, TEXT(\"%s\"), TEXT(\"%s\"), %s, %s, %s);\r\n"),
 			*ClassNameCPP,
 			*SingletonName,
@@ -2411,7 +2409,7 @@ void FNativeClassHeaderGenerator::ExportNativeGeneratedInitCode(FOutputDevice& O
 			*ClassNameCPP,
 			*OverriddenClassName,
 			*AsTEXT(ClassPackageName),
-			*AsTEXT(FNativeClassHeaderGenerator::GetOverriddenPathName(Class)),
+			*AsTEXT(FNativeClassHeaderGenerator::GetOverriddenPathName(ClassDef)),
 			*InitSearchableValuesFunctionParam);
 	}
 
@@ -2494,8 +2492,8 @@ void FNativeClassHeaderGenerator::ExportFunction(FOutputDevice& Out, FReferenceG
 		ExportEventParm(CurrentFunctionText, OutReferenceGatherers.ForwardDeclarations, FunctionDef, /*Indent=*/ 2, /*bOutputConstructor=*/ false, EExportingState::TypeEraseDelegates);
 	}
 
-	UField* FieldOuter = Cast<UField>(Function->GetOuter());
-	const bool bIsDynamic = (FieldOuter && FClass::IsDynamic(FieldOuter));
+	FUnrealFieldDefinitionInfo* FieldOuterDef = UHTCast<FUnrealFieldDefinitionInfo>(FunctionDef.GetOuter());
+	const bool bIsDynamic = (FieldOuterDef && FieldOuterDef->IsDynamic());
 
 	FString OuterFunc;
 	if (UObject* Outer = Function->GetOuter())
@@ -2542,7 +2540,7 @@ void FNativeClassHeaderGenerator::ExportFunction(FOutputDevice& Out, FReferenceG
 	}
 
 	USparseDelegateFunction* SparseDelegateFunction = Cast<USparseDelegateFunction>(Function);
-	const TCHAR* UFunctionObjectFlags = FClass::IsOwnedByDynamicType(Function) ? TEXT("RF_Public|RF_Transient") : TEXT("RF_Public|RF_Transient|RF_MarkAsNative");
+	const TCHAR* UFunctionObjectFlags = FunctionDef.IsOwnedByDynamicType() ? TEXT("RF_Public|RF_Transient") : TEXT("RF_Public|RF_Transient|RF_MarkAsNative");
 
 	TTuple<FString, FString> PropertyRange = OutputProperties(CurrentFunctionText, StaticDefinitions, OutReferenceGatherers, *FString::Printf(TEXT("%s::"), *StaticsStructName), Props, TEXT("\t\t"), TEXT("\t"));
 
@@ -2976,7 +2974,7 @@ void FNativeClassHeaderGenerator::ExportClassFromSourceFileInner(
 			Class->HasAnyClassFlags(CLASS_Abstract) ? TEXT("CLASS_Abstract") : TEXT("0"),
 			*GetClassFlagExportText(Class),
 			bCastedClass ? *FString::Printf(TEXT("CASTCLASS_%s"), *ClassCPPName) : TEXT("CASTCLASS_None"),
-			*FClass::GetTypePackageName(Class),
+			*ClassDef.GetTypePackageName(),
 			*APIArg);
 
 		Boilerplate.Logf(TEXT("\tDECLARE_SERIALIZER(%s)\r\n"), *ClassCPPName);
@@ -3246,7 +3244,6 @@ void ExportConstructorDefinition(FOutputDevice& Out, FUnrealClassDefinitionInfo&
 				// Since the SourceFile array provided to the ParallelFor is in dependency order and does not allow cyclic dependencies, 
 				// we can be certain that another thread has started processing the file containing our SuperClass before this
 				// file would have been assigned out,  so we just have to wait
-				//ETSTODO
 				FStructMetaData& SuperClassData = SuperClassDef->GetStructMetaData();
 				while (!SuperClassData.bConstructorDeclared)
 				{
@@ -3498,7 +3495,7 @@ void FNativeClassHeaderGenerator::ExportGeneratedStructBodyMacros(FOutputDevice&
 	FUnrealScriptStructDefinitionInfo& ScriptStructDef = GTypeDefinitionInfoMap.FindChecked<FUnrealScriptStructDefinitionInfo>(Struct);
 	ScriptStructDef.AddCrossModuleReference(OutReferenceGatherers.UniqueCrossModuleReferences, true);
 
-	const bool bIsDynamic = FClass::IsDynamic(Struct);
+	const bool bIsDynamic = ScriptStructDef.IsDynamic();
 	const FString ActualStructName = FNativeClassHeaderGenerator::GetOverriddenName(Struct);
 	const FString& FriendApiString  = GetAPIString();
 
@@ -3744,7 +3741,7 @@ void FNativeClassHeaderGenerator::ExportGeneratedStructBodyMacros(FOutputDevice&
 		}
 		else
 		{
-			Out.Logf(TEXT("\tclass UPackage* %s = FindOrConstructDynamicTypePackage(TEXT(\"%s\"));\r\n"), *OuterName, *FClass::GetTypePackageName(Struct));
+			Out.Logf(TEXT("\tclass UPackage* %s = FindOrConstructDynamicTypePackage(TEXT(\"%s\"));\r\n"), *OuterName, *ScriptStructDef.GetTypePackageName());
 			Out.Logf(TEXT("\tclass UScriptStruct* Singleton = Cast<UScriptStruct>(StaticFindObjectFast(UScriptStruct::StaticClass(), %s, TEXT(\"%s\")));\r\n"), *OuterName, *ActualStructName);
 		}
 
@@ -3780,14 +3777,14 @@ void FNativeClassHeaderGenerator::ExportGeneratedStructBodyMacros(FOutputDevice&
 
 		if (bIsDynamic)
 		{
-			const FString& StructPackageName = FClass::GetTypePackageName(Struct);
+			const FString& StructPackageName = ScriptStructDef.GetTypePackageName();
 			Out.Logf(TEXT("static FRegisterCompiledInInfo Z_CompiledInDeferStruct_UScriptStruct_%s(%s::StaticStruct, TEXT(\"%s\"), TEXT(\"%s\"),  %s, %s);\r\n"),
 				*StructNameCPP,
 				*StructNameCPP,
 				*StructPackageName,
 				*ActualStructName,
 				*AsTEXT(StructPackageName),
-				*AsTEXT(FNativeClassHeaderGenerator::GetOverriddenPathName(Struct)));
+				*AsTEXT(FNativeClassHeaderGenerator::GetOverriddenPathName(ScriptStructDef)));
 		}
 		else
 		{
@@ -3852,7 +3849,7 @@ void FNativeClassHeaderGenerator::ExportGeneratedStructBodyMacros(FOutputDevice&
 
 		StaticDefinitions.Logf(TEXT("\tUObject* %s::OuterFuncGetter()\r\n"), *StaticsStructName);
 		StaticDefinitions.Log (TEXT("\t{\r\n"));
-		StaticDefinitions.Logf(TEXT("\t\treturn FindOrConstructDynamicTypePackage(TEXT(\"%s\"));"), *FClass::GetTypePackageName(Struct));
+		StaticDefinitions.Logf(TEXT("\t\treturn FindOrConstructDynamicTypePackage(TEXT(\"%s\"));"), *ScriptStructDef.GetTypePackageName());
 		StaticDefinitions.Log (TEXT("\t}\r\n"));
 
 		OuterFunc = TEXT("&OuterFuncGetter");
@@ -3933,7 +3930,7 @@ void FNativeClassHeaderGenerator::ExportGeneratedStructBodyMacros(FOutputDevice&
 	}
 	else
 	{
-		GeneratedStructRegisterFunctionText.Logf(TEXT("\t\tUPackage* Outer = FindOrConstructDynamicTypePackage(TEXT(\"%s\"));\r\n"), *FClass::GetTypePackageName(Struct));
+		GeneratedStructRegisterFunctionText.Logf(TEXT("\t\tUPackage* Outer = FindOrConstructDynamicTypePackage(TEXT(\"%s\"));\r\n"), *ScriptStructDef.GetTypePackageName());
 		GeneratedStructRegisterFunctionText.Logf(TEXT("\t\tUScriptStruct* ReturnStruct = FindExistingStructIfHotReloadOrDynamic(Outer, TEXT(\"%s\"), sizeof(%s), %s(), true);\r\n"), *ActualStructName, *NoExportStructNameCPP, *HashFuncName);
 	}
 	GeneratedStructRegisterFunctionText.Logf(TEXT("\t\tif (!ReturnStruct)\r\n"));
@@ -4028,8 +4025,8 @@ void FNativeClassHeaderGenerator::ExportGeneratedStructBodyMacros(FOutputDevice&
 
 void FNativeClassHeaderGenerator::ExportGeneratedEnumInitCode(FOutputDevice& Out, FReferenceGatherers& OutReferenceGatherers, const FUnrealSourceFile& SourceFile, UEnum* Enum) const
 {
-	FUnrealEnumDefinitionInfo& EnumDef  = GTypeDefinitionInfoMap.FindChecked<FUnrealEnumDefinitionInfo>(Enum);
-	const bool bIsDynamic               = FClass::IsDynamic(static_cast<UField*>(Enum));
+	FUnrealEnumDefinitionInfo& EnumDef = GTypeDefinitionInfoMap.FindChecked<FUnrealEnumDefinitionInfo>(Enum);
+	const bool bIsDynamic               = EnumDef.IsDynamic();
 	const FString& SingletonName        = EnumDef.GetSingletonNameChopped(true);
 	const FString EnumNameCpp           = Enum->GetName(); //UserDefinedEnum should already have a valid cpp name.
 	const FString OverriddenEnumNameCpp = FNativeClassHeaderGenerator::GetOverriddenName(Enum);
@@ -4041,7 +4038,7 @@ void FNativeClassHeaderGenerator::ExportGeneratedEnumInitCode(FOutputDevice& Out
 	FMacroBlockEmitter EditorOnlyData(Out, TEXT("WITH_EDITORONLY_DATA"));
 	EditorOnlyData(bIsEditorOnlyDataType);
 
-	const FString& PackageSingletonName = (bIsDynamic ? FClass::GetTypePackageName(static_cast<UField*>(Enum)) : GetPackageSingletonName(CastChecked<UPackage>(Enum->GetOuter()), OutReferenceGatherers.UniqueCrossModuleReferences));
+	const FString& PackageSingletonName = (bIsDynamic ? EnumDef.GetTypePackageName() : GetPackageSingletonName(CastChecked<UPackage>(Enum->GetOuter()), OutReferenceGatherers.UniqueCrossModuleReferences));
 
 	if (!bIsDynamic)
 	{
@@ -4184,7 +4181,7 @@ void FNativeClassHeaderGenerator::ExportGeneratedEnumInitCode(FOutputDevice& Out
 
 	if (bIsDynamic)
 	{
-		const FString& EnumPackageName = FClass::GetTypePackageName(static_cast<UField*>(Enum));
+		const FString& EnumPackageName = EnumDef.GetTypePackageName();
 		Out.Logf(
 			TEXT("\tstatic FRegisterCompiledInInfo Z_CompiledInDeferEnum_UEnum_%s(%s_StaticEnum, TEXT(\"%s\"), TEXT(\"%s\"), %s, %s);\r\n"),
 			*EnumNameCpp,
@@ -4192,7 +4189,7 @@ void FNativeClassHeaderGenerator::ExportGeneratedEnumInitCode(FOutputDevice& Out
 			*EnumPackageName,
 			*OverriddenEnumNameCpp,
 			*AsTEXT(EnumPackageName),
-			*AsTEXT(FNativeClassHeaderGenerator::GetOverriddenPathName(static_cast<UField*>(Enum)))
+			*AsTEXT(FNativeClassHeaderGenerator::GetOverriddenPathName(EnumDef))
 		);
 	}
 	else
@@ -5199,8 +5196,8 @@ void FNativeClassHeaderGenerator::ExportNativeFunctions(FOutputDevice& OutGenera
 
 	// gather static class data
 	TArray<FString> SparseClassDataTypes;
-	((FClass*)Class)->GetSparseClassDataTypes(SparseClassDataTypes);
-	FString FullClassName = ((FClass*)Class)->GetNameWithPrefix();
+	ClassDef.GetSparseClassDataTypes(SparseClassDataTypes);
+	FString FullClassName = FUnrealTypeDefinitionInfo::GetNameWithPrefix(Class);
 	for (const FString& SparseClassDataString : SparseClassDataTypes)
 	{
 		RuntimeStringBuilders.AutogeneratedStaticData.Logf(TEXT("F%s* Get%s()\r\n"), *SparseClassDataString, *SparseClassDataString);
@@ -5568,8 +5565,8 @@ void FNativeClassHeaderGenerator::ApplyAlternatePropertyExportText(FProperty* Pr
 	FArrayProperty* ArrayProperty = CastField<FArrayProperty>(Prop);
 	FProperty* InnerProperty = ArrayProperty ? ArrayProperty->Inner : nullptr;
 	if (InnerProperty && (
-			(InnerProperty->IsA<FByteProperty>() && ((FByteProperty*)InnerProperty)->Enum && FClass::IsDynamic(static_cast<UField*>(((FByteProperty*)InnerProperty)->Enum))) ||
-			(InnerProperty->IsA<FEnumProperty>()                                          && FClass::IsDynamic(static_cast<UField*>(((FEnumProperty*)InnerProperty)->Enum)))
+			(InnerProperty->IsA<FByteProperty>() && ((FByteProperty*)InnerProperty)->Enum && FUnrealTypeDefinitionInfo::IsDynamic(static_cast<UField*>(((FByteProperty*)InnerProperty)->Enum))) ||
+			(InnerProperty->IsA<FEnumProperty>()                                          && FUnrealTypeDefinitionInfo::IsDynamic(static_cast<UField*>(((FEnumProperty*)InnerProperty)->Enum)))
 		)
 	)
 	{
@@ -5606,7 +5603,7 @@ void FNativeClassHeaderGenerator::ApplyAlternatePropertyExportText(FProperty* Pr
 static bool HasDynamicOuter(UField* Field)
 {
 	UField* FieldOuter = Cast<UField>(Field->GetOuter());
-	return FieldOuter && FClass::IsDynamic(FieldOuter);
+	return FieldOuter && FUnrealTypeDefinitionInfo::IsDynamic(FieldOuter);
 }
 
 static void RecordPackageSingletons(
@@ -6914,7 +6911,7 @@ void PrepareTypesForExport(TArray<FUnrealPackageDefinitionInfo*>& PackageDefs)
 				// Validate the sparse class data for all classes in the current package
 				for (const FUnrealClassDefinitionInfo* ClassDef : PackageClasses)
 				{
-					FHeaderParser::CheckSparseClassData(ClassDef->GetClass());
+					FHeaderParser::CheckSparseClassData(*ClassDef);
 				}
 
 				// Check to see if we need write the classes 
diff --git a/Engine/Source/Programs/UnrealHeaderTool/Private/HeaderParser.cpp b/Engine/Source/Programs/UnrealHeaderTool/Private/HeaderParser.cpp
--- a/Engine/Source/Programs/UnrealHeaderTool/Private/HeaderParser.cpp
+++ b/Engine/Source/Programs/UnrealHeaderTool/Private/HeaderParser.cpp
@@ -9,9 +9,7 @@
 #include "Misc/ConfigCacheIni.h"
 #include "Misc/FeedbackContext.h"
 #include "UObject/Interface.h"
-#include "ParserClass.h"
 #include "GeneratedCodeVersion.h"
-#include "ClassDeclarationMetaData.h"
 #include "ProfilingDebugging/ScopedTimers.h"
 #include "NativeClassExporter.h"
 #include "Classes.h"
@@ -86,10 +84,6 @@ namespace
 	static const FName NAME_CppFromBpEvent(TEXT("CppFromBpEvent"));
 	static const FName NAME_CustomThunk(TEXT("CustomThunk"));
 	static const FName NAME_ArraySizeEnum(TEXT("ArraySizeEnum"));
-	static const FName NAME_ClassGroupNames(TEXT("ClassGroupNames"));
-	static const FName NAME_AutoCollapseCategories(TEXT("AutoCollapseCategories"));
-	static const FName NAME_HideFunctions(TEXT("HideFunctions"));
-	static const FName NAME_AutoExpandCategories(TEXT("AutoExpandCategories"));
 	static const FName NAME_EditInline(TEXT("EditInline"));
 	static const FName NAME_IncludePath(TEXT("IncludePath"));
 	static const FName NAME_ModuleRelativePath(TEXT("ModuleRelativePath"));
@@ -104,6 +98,9 @@ const FName FHeaderParserNames::NAME_SparseClassDataTypes(TEXT("SparseClassDataT
 const FName FHeaderParserNames::NAME_IsConversionRoot(TEXT("IsConversionRoot"));
 const FName FHeaderParserNames::NAME_AdvancedClassDisplay(TEXT("AdvancedClassDisplay"));
 const FName FHeaderParserNames::NAME_BlueprintType(TEXT("BlueprintType"));
+const FName FHeaderParserNames::NAME_AutoCollapseCategories(TEXT("AutoCollapseCategories"));
+const FName FHeaderParserNames::NAME_HideFunctions(TEXT("HideFunctions"));
+const FName FHeaderParserNames::NAME_AutoExpandCategories(TEXT("AutoExpandCategories"));
 
 FRigVMStructMap FHeaderParser::StructRigVMMap;
 TArray<FString> FHeaderParser::PropertyCPPTypesRequiringUIRanges = { TEXT("float"), TEXT("double") };
@@ -1039,7 +1036,7 @@ FString FHeaderParser::GetContext()
 //
 // Get a qualified class.
 //
-FClass* FHeaderParser::GetQualifiedClass(const TCHAR* Thing)
+UClass* FHeaderParser::GetQualifiedClass(const TCHAR* Thing)
 {
 	TCHAR ClassName[256]=TEXT("");
 
@@ -2093,7 +2090,6 @@ UScriptStruct* FHeaderParser::CompileStructDeclaration()
 	ParseInheritance(TEXT("struct"), [](const TCHAR* StructName, bool bIsSuperClass) {}); // Eat the results, already been parsed
 
 	// if we have a base struct, propagate inherited struct flags now
-	//ETSTODO - Move base class resolution to the post declaration phase
 	FUnrealFieldDefinitionInfo* BaseStructDef = nullptr;
 	if (!StructDef.GetSuperClassInfo().Name.IsEmpty())
 	{
@@ -2379,7 +2375,7 @@ UScriptStruct* FHeaderParser::CompileStructDeclaration()
 	}
 
 	// Validate sparse class data
-	CheckSparseClassData(Struct);
+	CheckSparseClassData(StructDef);
 
 	// Link the properties within the struct
 	Struct->StaticLink(true);
@@ -2718,7 +2714,7 @@ void FHeaderParser::FixupDelegateProperties(FUnrealStructDefinitionInfo& StructD
 
 							if(FuncParam->HasAllPropertyFlags(CPF_OutParm) && !FuncParam->HasAllPropertyFlags(CPF_ConstParm)  )
 							{
-								const bool bClassGeneratedFromBP = FClass::IsDynamic(Struct);
+								const bool bClassGeneratedFromBP = StructDef.IsDynamic();
 								const bool bAllowedArrayRefFromBP = bClassGeneratedFromBP && FuncParam->IsA<FArrayProperty>();
 								if (!bAllowedArrayRefFromBP)
 								{
@@ -2748,27 +2744,28 @@ void FHeaderParser::FixupDelegateProperties(FUnrealStructDefinitionInfo& StructD
 	ParseRigVMMethodParameters(StructDef);
 }
 
-void FHeaderParser::CheckSparseClassData(const UStruct* StructToCheck)
+void FHeaderParser::CheckSparseClassData(const FUnrealStructDefinitionInfo& StructDef)
 {
 	// we're looking for classes that have sparse class data structures
-	const UClass* ClassToCheck = Cast<const UClass>(StructToCheck);
-	if (!ClassToCheck)
+	const UStruct* StructToCheck = StructDef.GetStruct();
+	const FUnrealClassDefinitionInfo* ClassDef = UHTCast<FUnrealClassDefinitionInfo>(StructDef);
+
+	// make sure we don't try to have sparse class data inside of a struct instead of a class
+	if (StructToCheck->HasMetaData(FHeaderParserNames::NAME_SparseClassDataTypes))
 	{
-		// make sure we don't try to have sparse class data inside of a struct instead of a class
-		if (StructToCheck->HasMetaData(FHeaderParserNames::NAME_SparseClassDataTypes))
+		if (ClassDef == nullptr)
 		{
 			FError::Throwf(TEXT("%s contains sparse class data but is not a class."), *StructToCheck->GetName());
 		}
-		return;
 	}
-
-	if (!ClassToCheck->HasMetaData(FHeaderParserNames::NAME_SparseClassDataTypes))
+	else
 	{
 		return;
 	}
+	const UClass* ClassToCheck = ClassDef->GetClass();
 
 	TArray<FString> SparseClassDataTypes;
-	((FClass*)ClassToCheck)->GetSparseClassDataTypes(SparseClassDataTypes);
+	ClassDef->GetSparseClassDataTypes(SparseClassDataTypes);
 
 	// for now we only support one sparse class data structure per class
 	if (SparseClassDataTypes.Num() > 1)
@@ -2819,7 +2816,7 @@ void FHeaderParser::CheckSparseClassData(const UStruct* StructToCheck)
 		// if the class's parent has a sparse class data struct then the current class must also use the same struct or one that inherits from it
 		const UClass* ParentClass = ClassToCheck->GetSuperClass();
 		TArray<FString> ParentSparseClassDataTypeNames;
-		((FClass*)ParentClass)->GetSparseClassDataTypes(ParentSparseClassDataTypeNames);
+		ClassDef->GetSuperClass()->GetSparseClassDataTypes(ParentSparseClassDataTypeNames);
 		for (FString& ParentSparseClassDataTypeName : ParentSparseClassDataTypeNames)
 		{
 			UScriptStruct* ParentSparseClassDataStruct = FClasses::FindObject<UScriptStruct>(ANY_PACKAGE, *ParentSparseClassDataTypeName);
@@ -2831,12 +2828,12 @@ void FHeaderParser::CheckSparseClassData(const UStruct* StructToCheck)
 	}
 }
 
-void FHeaderParser::ValidateClassFlags(const UClass* ToValidate)
+void FHeaderParser::ValidateClassFlags(const FUnrealClassDefinitionInfo& ToValidate)
 {
-	if (ToValidate->HasAnyClassFlags(CLASS_NeedsDeferredDependencyLoading) && !ToValidate->IsChildOf(UClass::StaticClass()))
+	if (ToValidate.HasAnyClassFlags(CLASS_NeedsDeferredDependencyLoading) && !ToValidate.GetClass()->IsChildOf(UClass::StaticClass()))
 	{
 		// CLASS_NeedsDeferredDependencyLoading can only be set on classes derived from UClass
-		FError::Throwf(TEXT("NeedsDeferredDependencyLoading is set on %s but the flag can only be used with classes derived from UClass."), *ToValidate->GetName());
+		FError::Throwf(TEXT("NeedsDeferredDependencyLoading is set on %s but the flag can only be used with classes derived from UClass."), *ToValidate.GetClass()->GetName());
 	}
 }
 
@@ -5637,7 +5634,7 @@ FUnrealClassDefinitionInfo& FHeaderParser::ParseClassNameDeclaration(FString& De
 {
 	ParseNameWithPotentialAPIMacroPrefix(/*out*/ DeclaredClassName, /*out*/ RequiredAPIMacroIfPresent, TEXT("class"));
 
-	FClass* FoundClass = FClasses::FindClass(*GetClassNameWithPrefixRemoved(*DeclaredClassName));
+	UClass* FoundClass = FClasses::FindClass(*GetClassNameWithPrefixRemoved(*DeclaredClassName));
 	check(FoundClass);
 	FoundClass->ClassCastFlags |= ClassCastFlagMap::Get().GetCastFlag(DeclaredClassName);
 
@@ -5750,12 +5747,12 @@ UClass* FHeaderParser::CompileClassDeclaration()
 	FString DeclaredClassName;
 	FString RequiredAPIMacroIfPresent;
 	
-	FUnrealClassDefinitionInfo& ClassDeclarationData = ParseClassNameDeclaration(/*out*/ DeclaredClassName, /*out*/ RequiredAPIMacroIfPresent);
-	FClass* Class = static_cast<FClass*>(ClassDeclarationData.GetClass());
+	FUnrealClassDefinitionInfo& ClassDef = ParseClassNameDeclaration(/*out*/ DeclaredClassName, /*out*/ RequiredAPIMacroIfPresent);
+	UClass* Class = ClassDef.GetClass();
 
 	ClassDefinitionRanges.Add(Class, ClassDefinitionRange(&Input[InputPos], nullptr));
 
-	check(Class->ClassFlags == 0 || (Class->ClassFlags & ClassDeclarationData.ClassFlags) != 0);
+	check(Class->ClassFlags == 0 || (Class->ClassFlags & ClassDef.ClassFlags) != 0);
 
 	Class->ClassFlags |= CLASS_Parsed;
 
@@ -5782,30 +5779,23 @@ UClass* FHeaderParser::CompileClassDeclaration()
 	}
 
 	// Merge with categories inherited from the parent.
-	ClassDeclarationData.MergeClassCategories(Class);
+	ClassDef.MergeClassCategories();
 
 	// Class attributes.
-	FStructMetaData& StructData = GTypeDefinitionInfoMap.FindChecked<FUnrealStructDefinitionInfo>(Class).GetStructMetaData();
+	FStructMetaData& StructData = ClassDef.GetStructMetaData();
 	StructData.SetPrologLine(PrologFinishLine);
 
-	ClassDeclarationData.MergeAndValidateClassFlags(DeclaredClassName, PrevClassFlags, Class);
+	ClassDef.MergeAndValidateClassFlags(DeclaredClassName, PrevClassFlags);
 	Class->SetInternalFlags(EInternalObjectFlags::Native);
 
 	// Class metadata
-	MetaData.Append(ClassDeclarationData.MetaData);
-	if (ClassDeclarationData.ClassGroupNames.Num()) { MetaData.Add(NAME_ClassGroupNames, FString::Join(ClassDeclarationData.ClassGroupNames, TEXT(" "))); }
-	if (ClassDeclarationData.AutoCollapseCategories.Num()) { MetaData.Add(NAME_AutoCollapseCategories, FString::Join(ClassDeclarationData.AutoCollapseCategories, TEXT(" "))); }
-	if (ClassDeclarationData.HideCategories.Num()) { MetaData.Add(FHeaderParserNames::NAME_HideCategories, FString::Join(ClassDeclarationData.HideCategories, TEXT(" "))); }
-	if (ClassDeclarationData.ShowSubCatgories.Num()) { MetaData.Add(FHeaderParserNames::NAME_ShowCategories, FString::Join(ClassDeclarationData.ShowSubCatgories, TEXT(" "))); }
-	if (ClassDeclarationData.SparseClassDataTypes.Num()) { MetaData.Add(FHeaderParserNames::NAME_SparseClassDataTypes, FString::Join(ClassDeclarationData.SparseClassDataTypes, TEXT(" "))); }
-	if (ClassDeclarationData.HideFunctions.Num()) { MetaData.Add(NAME_HideFunctions, FString::Join(ClassDeclarationData.HideFunctions, TEXT(" "))); }
-	if (ClassDeclarationData.AutoExpandCategories.Num()) { MetaData.Add(NAME_AutoExpandCategories, FString::Join(ClassDeclarationData.AutoExpandCategories, TEXT(" "))); }
-
+	MetaData.Append(ClassDef.MetaData);
+	ClassDef.MergeCategoryMetaData(MetaData);
 	AddIncludePathToMetadata(Class, MetaData);
-	AddModuleRelativePathToMetadata(ClassDeclarationData, MetaData);
+	AddModuleRelativePathToMetadata(ClassDef, MetaData);
 
 	// Register the metadata
-	AddMetaDataToClassData(ClassDeclarationData, MoveTemp(MetaData));
+	AddMetaDataToClassData(ClassDef, MoveTemp(MetaData));
 
 	// Handle the start of the rest of the class
 	RequireSymbol( TEXT('{'), TEXT("'Class'") );
@@ -5823,7 +5813,7 @@ UClass* FHeaderParser::CompileClassDeclaration()
 
 	// auto-create properties for all of the VFTables needed for the multiple inheritances
 	// get the inheritance parents
-	const TArray<FUnrealStructDefinitionInfo::FBaseClassInfo>& BaseClassInfos = ClassDeclarationData.GetBaseClassInfo();
+	const TArray<FUnrealStructDefinitionInfo::FBaseClassInfo>& BaseClassInfos = ClassDef.GetBaseClassInfo();
 
 	// for all base class types, make a VfTable property
 	for (int32 ParentIndex = BaseClassInfos.Num() - 1; ParentIndex >= 0; ParentIndex--)
@@ -5848,19 +5838,19 @@ UClass* FHeaderParser::CompileClassDeclaration()
 	}
 
 	// Validate sparse class data
-	CheckSparseClassData(Class);
+	CheckSparseClassData(ClassDef);
 
 	// Check that the class has appropriate class flags set
-	ValidateClassFlags(Class);
+	ValidateClassFlags(ClassDef);
 
 	return Class;
 }
 
-FClass* FHeaderParser::ParseInterfaceNameDeclaration(FString& DeclaredInterfaceName, FString& RequiredAPIMacroIfPresent)
+UClass* FHeaderParser::ParseInterfaceNameDeclaration(FString& DeclaredInterfaceName, FString& RequiredAPIMacroIfPresent)
 {
 	ParseNameWithPotentialAPIMacroPrefix(/*out*/ DeclaredInterfaceName, /*out*/ RequiredAPIMacroIfPresent, TEXT("interface"));
 
-	FClass* FoundClass = FClasses::FindClass(*GetClassNameWithPrefixRemoved(*DeclaredInterfaceName));
+	UClass* FoundClass = FClasses::FindClass(*GetClassNameWithPrefixRemoved(*DeclaredInterfaceName));
 	if (FoundClass == nullptr)
 	{
 		return nullptr;
@@ -5960,7 +5950,7 @@ void FHeaderParser::CompileInterfaceDeclaration()
 
 	// New style files have the interface name / extends afterwards
 	RequireIdentifier(TEXT("class"), ESearchCase::CaseSensitive, TEXT("Interface declaration"));
-	FClass* InterfaceClass = ParseInterfaceNameDeclaration(/*out*/ DeclaredInterfaceName, /*out*/ RequiredAPIMacroIfPresent);
+	UClass* InterfaceClass = ParseInterfaceNameDeclaration(/*out*/ DeclaredInterfaceName, /*out*/ RequiredAPIMacroIfPresent);
 	FUnrealClassDefinitionInfo& ClassDef = GTypeDefinitionInfoMap.FindChecked<FUnrealClassDefinitionInfo>(InterfaceClass);
 	ClassDefinitionRanges.Add(InterfaceClass, ClassDefinitionRange(&Input[InputPos], nullptr));
 
@@ -6018,7 +6008,7 @@ void FHeaderParser::CompileInterfaceDeclaration()
 	}
 
 	// All classes must start with a valid Unreal prefix
-	const FString ExpectedInterfaceName = InterfaceClass->GetNameWithPrefix(EEnforceInterfacePrefix::U);
+	const FString ExpectedInterfaceName = FUnrealTypeDefinitionInfo::GetNameWithPrefix(InterfaceClass, EEnforceInterfacePrefix::U);
 	if (DeclaredInterfaceName != ExpectedInterfaceName)
 	{
 		FError::Throwf(TEXT("Interface name '%s' is invalid, the first class should be identified as '%s'"), *DeclaredInterfaceName, *ExpectedInterfaceName );
@@ -6498,7 +6488,7 @@ FUnrealFunctionDefinitionInfo& FHeaderParser::CompileDelegateDeclaration(const T
 
 		//Workaround for UE-28897
 		const FStructScope* CurrentStructScope = TopNest->GetScope() ? TopNest->GetScope()->AsStructScope() : nullptr;
-		const bool bDynamicClassScope = CurrentStructScope && CurrentStructScope->GetStruct() && FClass::IsDynamic(CurrentStructScope->GetStruct());
+		const bool bDynamicClassScope = CurrentStructScope && CurrentStructScope->GetStruct() && CurrentStructScope->GetStructDef().IsDynamic();
 		CheckAllow(CurrentScopeName, bDynamicClassScope ? ENestAllowFlags::ImplicitDelegateDecl : ENestAllowFlags::TypeDecl);
 	}
 	else
@@ -6760,7 +6750,7 @@ void FHeaderParser::CompileFunctionDeclaration()
 
 	ProcessFunctionSpecifiers(FuncInfo, SpecifiersFound, MetaData);
 
-	const bool bClassGeneratedFromBP = FClass::IsDynamic(GetCurrentClass());
+	const bool bClassGeneratedFromBP = GetCurrentClassDef().IsDynamic();
 
 	if ((0 != (FuncInfo.FunctionExportFlags & FUNCEXPORT_CustomThunk)) && !MetaData.Contains(NAME_CustomThunk))
 	{
@@ -7498,7 +7488,6 @@ void FHeaderParser::CompileVariableDeclaration(FUnrealStructDefinitionInfo& Stru
 	TArray<FProperty*> NewProperties;
 	for (;;)
 	{
-		//ETSTODO
 		FPropertyBase Property = OriginalProperty;
 		FUnrealPropertyDefinitionInfo& NewPropDef = GetVarNameAndDim(StructDef, Property, EVariableCategory::Member, LayoutMacroType);
 		FProperty* NewProperty = NewPropDef.GetProperty();
@@ -9068,17 +9057,17 @@ void FHeaderParser::ResetClassData()
 		{
 			CurrentClass->SetMetaData(FHeaderParserNames::NAME_SparseClassDataTypes, *SuperClass->GetMetaData(FHeaderParserNames::NAME_SparseClassDataTypes));
 		}
-		if (SuperClass->HasMetaData(NAME_HideFunctions))
+		if (SuperClass->HasMetaData(FHeaderParserNames::NAME_HideFunctions))
 		{
-			CurrentClass->SetMetaData(NAME_HideFunctions, *SuperClass->GetMetaData(NAME_HideFunctions));
+			CurrentClass->SetMetaData(FHeaderParserNames::NAME_HideFunctions, *SuperClass->GetMetaData(FHeaderParserNames::NAME_HideFunctions));
 		}
-		if (SuperClass->HasMetaData(NAME_AutoExpandCategories))
+		if (SuperClass->HasMetaData(FHeaderParserNames::NAME_AutoExpandCategories))
 		{
-			CurrentClass->SetMetaData(NAME_AutoExpandCategories, *SuperClass->GetMetaData(NAME_AutoExpandCategories));
+			CurrentClass->SetMetaData(FHeaderParserNames::NAME_AutoExpandCategories, *SuperClass->GetMetaData(FHeaderParserNames::NAME_AutoExpandCategories));
 		}
-		if (SuperClass->HasMetaData(NAME_AutoCollapseCategories))
+		if (SuperClass->HasMetaData(FHeaderParserNames::NAME_AutoCollapseCategories))
 		{
-			CurrentClass->SetMetaData(NAME_AutoCollapseCategories, *SuperClass->GetMetaData(NAME_AutoCollapseCategories));
+			CurrentClass->SetMetaData(FHeaderParserNames::NAME_AutoCollapseCategories, *SuperClass->GetMetaData(FHeaderParserNames::NAME_AutoCollapseCategories));
 		}
 	}
 
diff --git a/Engine/Source/Programs/UnrealHeaderTool/Private/HeaderParser.h b/Engine/Source/Programs/UnrealHeaderTool/Private/HeaderParser.h
--- a/Engine/Source/Programs/UnrealHeaderTool/Private/HeaderParser.h
+++ b/Engine/Source/Programs/UnrealHeaderTool/Private/HeaderParser.h
@@ -25,7 +25,6 @@ class UEnum;
 class UScriptStruct;
 class UDelegateFunction;
 class UStruct;
-class FClass;
 class FClasses;
 class FScope;
 class FHeaderProvider;
@@ -663,10 +662,10 @@ public:
 	static FField* FindProperty(FUnrealStructDefinitionInfo& InScope, const TCHAR* InIdentifier, bool bIncludeParents = true, FFieldClass* FieldClass = FField::StaticClass(), const TCHAR* Thing = nullptr);
 
 	// Checks ToValidate to make sure that its associated sparse class data struct, if one exists, is a valid structure to use for storing sparse class data.
-	static void CheckSparseClassData(const UStruct* ToValidate);
+	static void CheckSparseClassData(const FUnrealStructDefinitionInfo& StructDef);
 
 	// Validates that ClassFlags are set appropriately 
-	static void ValidateClassFlags(const UClass* ToValidate);
+	static void ValidateClassFlags(const FUnrealClassDefinitionInfo& ToValidate);
 
 protected:
 
@@ -745,7 +744,7 @@ protected:
 	void CompileRigVMMethodDeclaration(UStruct* Struct);
 	void ParseRigVMMethodParameters(FUnrealStructDefinitionInfo& StructDef);
 
-	FClass* ParseInterfaceNameDeclaration(FString& DeclaredInterfaceName, FString& RequiredAPIMacroIfPresent);
+	UClass* ParseInterfaceNameDeclaration(FString& DeclaredInterfaceName, FString& RequiredAPIMacroIfPresent);
 	bool TryParseIInterfaceClass();
 
 	bool CompileStatement(TArray<FUnrealFunctionDefinitionInfo*>& DelegatesToFixup);
@@ -833,7 +832,7 @@ protected:
 
 	const TCHAR* NestTypeName( ENestType NestType );
 
-	FClass* GetQualifiedClass(const TCHAR* Thing);
+	UClass* GetQualifiedClass(const TCHAR* Thing);
 
 	/**
 	 * Increase the nesting level, setting the new top nesting level to
diff --git a/Engine/Source/Programs/UnrealHeaderTool/Private/NativeClassExporter.h b/Engine/Source/Programs/UnrealHeaderTool/Private/NativeClassExporter.h
--- a/Engine/Source/Programs/UnrealHeaderTool/Private/NativeClassExporter.h
+++ b/Engine/Source/Programs/UnrealHeaderTool/Private/NativeClassExporter.h
@@ -739,10 +739,14 @@ public:
 		return Item->GetFName();
 	}
 
-	template <typename T>
-	static FString GetOverriddenPathName(const T* Item)
+	static FString GetOverriddenPathName(const FUnrealPropertyDefinitionInfo& ItemDef)
+	{
+		return FString::Printf(TEXT("%s.%s"), *ItemDef.GetTypePackageName(), *GetOverriddenName(ItemDef.GetProperty()));
+	}
+
+	static FString GetOverriddenPathName(const FUnrealFieldDefinitionInfo& ItemDef)
 	{
-		return FString::Printf(TEXT("%s.%s"), *FClass::GetTypePackageName(Item), *GetOverriddenName(Item));
+		return FString::Printf(TEXT("%s.%s"), *ItemDef.GetTypePackageName(), *GetOverriddenName(ItemDef.GetField()));
 	}
 
 	/**
diff --git a/Engine/Source/Programs/UnrealHeaderTool/Private/ParserClass.cpp b/Engine/Source/Programs/UnrealHeaderTool/Private/ParserClass.cpp
deleted file mode 100644
diff --git a/Engine/Source/Programs/UnrealHeaderTool/Private/ParserClass.h b/Engine/Source/Programs/UnrealHeaderTool/Private/ParserClass.h
deleted file mode 100644
diff --git a/Engine/Source/Programs/UnrealHeaderTool/Private/PropertyTypes.cpp b/Engine/Source/Programs/UnrealHeaderTool/Private/PropertyTypes.cpp
--- a/Engine/Source/Programs/UnrealHeaderTool/Private/PropertyTypes.cpp
+++ b/Engine/Source/Programs/UnrealHeaderTool/Private/PropertyTypes.cpp
@@ -68,7 +68,6 @@ struct FPropertyTypeTraitsFloat;
 struct FPropertyTypeTraitsDouble;
 struct FPropertyTypeTraitsLargeWorldCoordinatesReal;
 struct FPropertyTypeTraitsObjectReference;
-struct FPropertyTypeTraitsObjectReference;
 struct FPropertyTypeTraitsWeakObjectReference;
 struct FPropertyTypeTraitsLazyObjectReference;
 struct FPropertyTypeTraitsObjectPtrReference;
@@ -723,7 +722,6 @@ struct FPropertyTypeTraitsObjectReference : public FPropertyTypeTraitsObjectBase
 		}
 		else
 		{
-			//ETSTODO - This prevents us from having PropDef const
 			if (FUnrealClassDefinitionInfo::HierarchyHasAnyClassFlags(VarProperty.PropertyClass, CLASS_DefaultToInstanced))
 			{
 				VarProperty.PropertyFlags |= CPF_InstancedReference;
@@ -801,7 +799,6 @@ struct FPropertyTypeTraitsObjectPtrReference : public FPropertyTypeTraitsObjectB
 		}
 		else
 		{
-			//ETSTODO - This prevents us from having PropDef const
 			if (FUnrealClassDefinitionInfo::HierarchyHasAnyClassFlags(VarProperty.PropertyClass, CLASS_DefaultToInstanced))
 			{
 				VarProperty.PropertyFlags |= CPF_InstancedReference;
@@ -1195,7 +1192,6 @@ struct FPropertyTypeTraitsStruct : public FPropertyTypeTraitsBase
 		PropDef.GetUnrealSourceFile().AddTypeDefIncludeIfNeeded(VarProperty.Struct);
 #endif
 
-		//ETSTODO - This is preventing PropDef from being const
 		if (VarProperty.Struct->StructFlags & STRUCT_HasInstancedReference)
 		{
 			VarProperty.PropertyFlags |= CPF_ContainsInstancedReference;
diff --git a/Engine/Source/Programs/UnrealHeaderTool/Private/Scope.h b/Engine/Source/Programs/UnrealHeaderTool/Private/Scope.h
--- a/Engine/Source/Programs/UnrealHeaderTool/Private/Scope.h
+++ b/Engine/Source/Programs/UnrealHeaderTool/Private/Scope.h
@@ -371,7 +371,7 @@ public:
 	/**
 	 * Get the structure definition associated with this scope
 	 */
-	FUnrealStructDefinitionInfo& GetStructDef()
+	FUnrealStructDefinitionInfo& GetStructDef() const
 	{
 		return StructDef;
 	}
diff --git a/Engine/Source/Programs/UnrealHeaderTool/Private/UnrealHeaderTool.h b/Engine/Source/Programs/UnrealHeaderTool/Private/UnrealHeaderTool.h
--- a/Engine/Source/Programs/UnrealHeaderTool/Private/UnrealHeaderTool.h
+++ b/Engine/Source/Programs/UnrealHeaderTool/Private/UnrealHeaderTool.h
@@ -18,4 +18,7 @@ struct FHeaderParserNames
 	static const FName NAME_SparseClassDataTypes;
 	static const FName NAME_AdvancedClassDisplay;
 	static const FName NAME_BlueprintType;
+	static const FName NAME_AutoCollapseCategories;
+	static const FName NAME_HideFunctions;
+	static const FName NAME_AutoExpandCategories;
 };
\ No newline at end of file
diff --git a/Engine/Source/Programs/UnrealHeaderTool/Private/UnrealTypeDefinitionInfo.cpp b/Engine/Source/Programs/UnrealHeaderTool/Private/UnrealTypeDefinitionInfo.cpp
--- a/Engine/Source/Programs/UnrealHeaderTool/Private/UnrealTypeDefinitionInfo.cpp
+++ b/Engine/Source/Programs/UnrealHeaderTool/Private/UnrealTypeDefinitionInfo.cpp
@@ -1,19 +1,30 @@
 // Copyright Epic Games, Inc. All Rights Reserved.
 
 #include "UnrealTypeDefinitionInfo.h"
+#include "BaseParser.h"
+#include "Classes.h"
 #include "ClassMaps.h"
+#include "HeaderParser.h"
 #include "NativeClassExporter.h"
 #include "Scope.h"
+#include "UnrealHeaderTool.h"
 #include "UnrealSourceFile.h"
 #include "Misc/PackageName.h"
 #include "UObject/ErrorException.h"
+#include "UObject/Interface.h"
+#include "Specifiers/ClassMetadataSpecifiers.h"
+#include "Algo/FindSortedStringCaseInsensitive.h"
 
 namespace
 {
+	const FName NAME_IgnoreCategoryKeywordsInSubclasses(TEXT("IgnoreCategoryKeywordsInSubclasses"));
+	const FName NAME_ReplaceConverted(TEXT("ReplaceConverted"));
+	const FName NAME_ClassGroupNames(TEXT("ClassGroupNames"));
+
 	/**
 	 * As part of the singleton name, collect the parent chain names
 	 */
-	static void AddOuterNames(FUHTStringBuilder& Out, UObject* Outer)
+	void AddOuterNames(FUHTStringBuilder& Out, UObject* Outer)
 	{
 		if (Outer == nullptr)
 		{
@@ -46,7 +57,7 @@ namespace
 	/**
 	 * Generates singleton name.
 	 */
-	static void GenerateSingletonName(FUHTStringBuilder& Out, UField* Item, bool bRequiresValidObject)
+	void GenerateSingletonName(FUHTStringBuilder& Out, UField* Item, bool bRequiresValidObject)
 	{
 		check(Item);
 
@@ -63,6 +74,39 @@ namespace
 		}
 		Out.Append(TEXT("()"));
 	}
+
+	bool IsActorClass(UClass* Class)
+	{
+		while (Class)
+		{
+			if (Class->GetFName() == NAME_Actor)
+			{
+				return true;
+			}
+			Class = Class->GetSuperClass();
+		}
+		return false;
+	}
+
+	template <typename T>
+	FString GetTypePackageName_Inner(const T* Field)
+	{
+		FString PackageName = Field->GetMetaData(NAME_ReplaceConverted);
+		if (PackageName.Len())
+		{
+			int32 ObjectDotIndex = INDEX_NONE;
+			// Strip the object name
+			if (PackageName.FindChar(TEXT('.'), ObjectDotIndex))
+			{
+				PackageName.MidInline(0, ObjectDotIndex, false);
+			}
+		}
+		else
+		{
+			PackageName = Field->GetOutermost()->GetName();
+		}
+		return PackageName;
+	}
 }
 
 FUnrealPropertyDefinitionInfo* FUnrealTypeDefinitionInfo::AsProperty()
@@ -149,6 +193,65 @@ void FUnrealTypeDefinitionInfo::GetHashTag(FUHTStringBuilder& Out) const
 	}
 }
 
+bool FUnrealTypeDefinitionInfo::IsDynamic(const UField* Field)
+{
+	return Field->HasMetaData(NAME_ReplaceConverted);
+}
+
+bool FUnrealTypeDefinitionInfo::IsDynamic(const FField* Field)
+{
+	return Field->HasMetaData(NAME_ReplaceConverted);
+}
+
+
+FString FUnrealTypeDefinitionInfo::GetNameWithPrefix(const UClass* InClass, EEnforceInterfacePrefix EnforceInterfacePrefix)
+{
+	const TCHAR* Prefix = 0;
+
+	if (InClass->HasAnyClassFlags(CLASS_Interface))
+	{
+		// Grab the expected prefix for interfaces (U on the first one, I on the second one)
+		switch (EnforceInterfacePrefix)
+		{
+		case EEnforceInterfacePrefix::None:
+			// For old-style files: "I" for interfaces, unless it's the actual "Interface" class, which gets "U"
+			if (InClass->GetFName() == NAME_Interface)
+			{
+				Prefix = TEXT("U");
+			}
+			else
+			{
+				Prefix = TEXT("I");
+			}
+			break;
+
+		case EEnforceInterfacePrefix::I:
+			Prefix = TEXT("I");
+			break;
+
+		case EEnforceInterfacePrefix::U:
+			Prefix = TEXT("U");
+			break;
+
+		default:
+			check(false);
+		}
+	}
+	else
+	{
+		// Get the expected class name with prefix
+		Prefix = InClass->GetPrefixCPP();
+	}
+
+	return FString::Printf(TEXT("%s%s"), Prefix, *InClass->GetName());
+}
+
+
+void FUnrealPropertyDefinitionInfo::PostParseFinalize()
+{
+	TypePackageName = GetTypePackageName_Inner(GetProperty());
+}
+
 void FUnrealPropertyDefinitionInfo::AddMetaData(TMap<FName, FString>&& InMetaData)
 {
 	// only add if we have some!
@@ -167,6 +270,27 @@ void FUnrealPropertyDefinitionInfo::AddMetaData(TMap<FName, FString>&& InMetaDat
 	}
 }
 
+bool FUnrealPropertyDefinitionInfo::IsDynamic() const
+{
+	return GetProperty()->HasMetaData(NAME_ReplaceConverted);
+}
+
+bool FUnrealPropertyDefinitionInfo::IsOwnedByDynamicType() const
+{
+	for (FFieldVariant Owner = GetProperty()->GetOwnerVariant(); Owner.IsValid(); Owner = Owner.GetOwnerVariant())
+	{
+		if (Owner.IsUObject())
+		{
+			return GTypeDefinitionInfoMap.FindChecked<FUnrealObjectDefinitionInfo>(Owner.ToUObject()).IsOwnedByDynamicType();
+		}
+		else
+		{
+			return GTypeDefinitionInfoMap.FindChecked(Owner.ToField()).IsOwnedByDynamicType();
+		}
+	}
+	return false;
+}
+
 FUnrealPackageDefinitionInfo::FUnrealPackageDefinitionInfo(const FManifestModule& InModule, UPackage* InPackage)
 	: FUnrealObjectDefinitionInfo(FString())
 	, Module(InModule)
@@ -213,6 +337,8 @@ void FUnrealFieldDefinitionInfo::PostParseFinalize()
 	ExternDecl[1].Appendf(TEXT("\t%s_API %s* %s;\r\n"), *PackageShortName, TypeStr, *Out);
 	SingletonName[1] = Out;
 	SingletonNameChopped[1] = SingletonName[1].LeftChop(2);
+
+	TypePackageName = GetTypePackageName_Inner(GetField());
 }
 
 void FUnrealFieldDefinitionInfo::AddCrossModuleReference(TSet<FString>* UniqueCrossModuleReferences, bool bRequiresValidObject) const
@@ -256,6 +382,23 @@ void FUnrealFieldDefinitionInfo::AddMetaData(TMap<FName, FString>&& InMetaData)
 	}
 }
 
+bool FUnrealFieldDefinitionInfo::IsDynamic() const
+{
+	return GetField()->HasMetaData(NAME_ReplaceConverted);
+}
+
+bool FUnrealFieldDefinitionInfo::IsOwnedByDynamicType() const
+{
+	for (const FUnrealObjectDefinitionInfo* OuterObject = GetOuter(); OuterObject; OuterObject = OuterObject->GetOuter())
+	{
+		if (OuterObject->IsDynamic())
+		{
+			return true;
+		}
+	}
+	return false;
+}
+
 FUnrealEnumDefinitionInfo::FUnrealEnumDefinitionInfo(FUnrealSourceFile& InSourceFile, int32 InLineNumber, FString&& InNameCPP)
 	: FUnrealFieldDefinitionInfo(InSourceFile, InLineNumber, MoveTemp(InNameCPP), InSourceFile.GetPackageDef())
 { }
@@ -397,6 +540,558 @@ bool FUnrealClassDefinitionInfo::HierarchyHasAllClassFlags(UClass* Class, EClass
 	return false;
 }
 
+void FUnrealClassDefinitionInfo::ParseClassProperties(TArray<FPropertySpecifier>&& InClassSpecifiers, const FString& InRequiredAPIMacroIfPresent)
+{
+	ClassFlags = CLASS_None;
+	// Record that this class is RequiredAPI if the CORE_API style macro was present
+	if (!InRequiredAPIMacroIfPresent.IsEmpty())
+	{
+		ClassFlags |= CLASS_RequiredAPI;
+	}
+	ClassFlags |= CLASS_Native;
+
+	// Process all of the class specifiers
+
+	for (FPropertySpecifier& PropSpecifier : InClassSpecifiers)
+	{
+		switch ((EClassMetadataSpecifier)Algo::FindSortedStringCaseInsensitive(*PropSpecifier.Key, GClassMetadataSpecifierStrings))
+		{
+		case EClassMetadataSpecifier::NoExport:
+
+			// Don't export to C++ header.
+			ClassFlags |= CLASS_NoExport;
+			break;
+
+		case EClassMetadataSpecifier::Intrinsic:
+
+			ClassFlags |= CLASS_Intrinsic;
+			break;
+
+		case EClassMetadataSpecifier::ComponentWrapperClass:
+
+			MetaData.Add(NAME_IgnoreCategoryKeywordsInSubclasses, TEXT("true"));
+			break;
+
+		case EClassMetadataSpecifier::Within:
+
+			ClassWithin = FHeaderParser::RequireExactlyOneSpecifierValue(PropSpecifier);
+			break;
+
+		case EClassMetadataSpecifier::EditInlineNew:
+
+			// Class can be constructed from the New button in editinline
+			ClassFlags |= CLASS_EditInlineNew;
+			break;
+
+		case EClassMetadataSpecifier::NotEditInlineNew:
+
+			// Class cannot be constructed from the New button in editinline
+			ClassFlags &= ~CLASS_EditInlineNew;
+			break;
+
+		case EClassMetadataSpecifier::Placeable:
+
+			bWantsToBePlaceable = true;
+			ClassFlags &= ~CLASS_NotPlaceable;
+			break;
+
+		case EClassMetadataSpecifier::DefaultToInstanced:
+
+			// these classed default to instanced.
+			ClassFlags |= CLASS_DefaultToInstanced;
+			break;
+
+		case EClassMetadataSpecifier::NotPlaceable:
+
+			// Don't allow the class to be placed in the editor.
+			ClassFlags |= CLASS_NotPlaceable;
+			break;
+
+		case EClassMetadataSpecifier::HideDropdown:
+
+			// Prevents class from appearing in class comboboxes in the property window
+			ClassFlags |= CLASS_HideDropDown;
+			break;
+
+		case EClassMetadataSpecifier::DependsOn:
+
+			FError::Throwf(TEXT("The dependsOn specifier is deprecated. Please use #include \"ClassHeaderFilename.h\" instead."));
+			break;
+
+		case EClassMetadataSpecifier::MinimalAPI:
+
+			ClassFlags |= CLASS_MinimalAPI;
+			break;
+
+		case EClassMetadataSpecifier::Const:
+
+			ClassFlags |= CLASS_Const;
+			break;
+
+		case EClassMetadataSpecifier::PerObjectConfig:
+
+			ClassFlags |= CLASS_PerObjectConfig;
+			break;
+
+		case EClassMetadataSpecifier::ConfigDoNotCheckDefaults:
+
+			ClassFlags |= CLASS_ConfigDoNotCheckDefaults;
+			break;
+
+		case EClassMetadataSpecifier::Abstract:
+
+			// Hide all editable properties.
+			ClassFlags |= CLASS_Abstract;
+			break;
+
+		case EClassMetadataSpecifier::Deprecated:
+
+			ClassFlags |= CLASS_Deprecated;
+
+			// Don't allow the class to be placed in the editor.
+			ClassFlags |= CLASS_NotPlaceable;
+
+			break;
+
+		case EClassMetadataSpecifier::Transient:
+
+			// Transient class.
+			ClassFlags |= CLASS_Transient;
+			break;
+
+		case EClassMetadataSpecifier::NonTransient:
+
+			// this child of a transient class is not transient - remove the transient flag
+			ClassFlags &= ~CLASS_Transient;
+			break;
+
+		case EClassMetadataSpecifier::CustomConstructor:
+
+			// we will not export a constructor for this class, assuming it is in the CPP block
+			ClassFlags |= CLASS_CustomConstructor;
+			break;
+
+		case EClassMetadataSpecifier::Config:
+
+			// Class containing config properties - parse the name of the config file to use
+			ConfigName = FHeaderParser::RequireExactlyOneSpecifierValue(PropSpecifier);
+			break;
+
+		case EClassMetadataSpecifier::DefaultConfig:
+
+			// Save object config only to Default INIs, never to local INIs.
+			ClassFlags |= CLASS_DefaultConfig;
+			break;
+
+		case EClassMetadataSpecifier::GlobalUserConfig:
+
+			// Save object config only to global user overrides, never to local INIs
+			ClassFlags |= CLASS_GlobalUserConfig;
+			break;
+
+		case EClassMetadataSpecifier::ProjectUserConfig:
+
+			// Save object config only to project user overrides, never to INIs that are checked in
+			ClassFlags |= CLASS_ProjectUserConfig;
+			break;
+
+		case EClassMetadataSpecifier::ShowCategories:
+
+			FHeaderParser::RequireSpecifierValue(PropSpecifier);
+
+			for (FString& Value : PropSpecifier.Values)
+			{
+				ShowCategories.AddUnique(MoveTemp(Value));
+			}
+			break;
+
+		case EClassMetadataSpecifier::HideCategories:
+
+			FHeaderParser::RequireSpecifierValue(PropSpecifier);
+
+			for (FString& Value : PropSpecifier.Values)
+			{
+				HideCategories.AddUnique(MoveTemp(Value));
+			}
+			break;
+
+		case EClassMetadataSpecifier::ShowFunctions:
+
+			FHeaderParser::RequireSpecifierValue(PropSpecifier);
+
+			for (const FString& Value : PropSpecifier.Values)
+			{
+				HideFunctions.RemoveSwap(Value);
+			}
+			break;
+
+		case EClassMetadataSpecifier::HideFunctions:
+
+			FHeaderParser::RequireSpecifierValue(PropSpecifier);
+
+			for (FString& Value : PropSpecifier.Values)
+			{
+				HideFunctions.AddUnique(MoveTemp(Value));
+			}
+			break;
+
+			// Currently some code only handles a single sidecar data structure so we enforce that here
+		case EClassMetadataSpecifier::SparseClassDataTypes:
+
+			SparseClassDataTypes.AddUnique(FHeaderParser::RequireExactlyOneSpecifierValue(PropSpecifier));
+			break;
+
+		case EClassMetadataSpecifier::ClassGroup:
+
+			FHeaderParser::RequireSpecifierValue(PropSpecifier);
+
+			for (FString& Value : PropSpecifier.Values)
+			{
+				ClassGroupNames.Add(MoveTemp(Value));
+			}
+			break;
+
+		case EClassMetadataSpecifier::AutoExpandCategories:
+
+			FHeaderParser::RequireSpecifierValue(PropSpecifier);
+
+			for (FString& Value : PropSpecifier.Values)
+			{
+				AutoCollapseCategories.RemoveSwap(Value);
+				AutoExpandCategories.AddUnique(MoveTemp(Value));
+			}
+			break;
+
+		case EClassMetadataSpecifier::AutoCollapseCategories:
+
+			FHeaderParser::RequireSpecifierValue(PropSpecifier);
+
+			for (FString& Value : PropSpecifier.Values)
+			{
+				AutoExpandCategories.RemoveSwap(Value);
+				AutoCollapseCategories.AddUnique(MoveTemp(Value));
+			}
+			break;
+
+		case EClassMetadataSpecifier::DontAutoCollapseCategories:
+
+			FHeaderParser::RequireSpecifierValue(PropSpecifier);
+
+			for (const FString& Value : PropSpecifier.Values)
+			{
+				AutoCollapseCategories.RemoveSwap(Value);
+			}
+			break;
+
+		case EClassMetadataSpecifier::CollapseCategories:
+
+			// Class' properties should not be shown categorized in the editor.
+			ClassFlags |= CLASS_CollapseCategories;
+			break;
+
+		case EClassMetadataSpecifier::DontCollapseCategories:
+
+			// Class' properties should be shown categorized in the editor.
+			ClassFlags &= ~CLASS_CollapseCategories;
+			break;
+
+		case EClassMetadataSpecifier::AdvancedClassDisplay:
+
+			// By default the class properties are shown in advanced sections in UI
+			MetaData.Add(FHeaderParserNames::NAME_AdvancedClassDisplay, TEXT("true"));
+			break;
+
+		case EClassMetadataSpecifier::ConversionRoot:
+
+			MetaData.Add(FHeaderParserNames::NAME_IsConversionRoot, TEXT("true"));
+			break;
+
+		case EClassMetadataSpecifier::NeedsDeferredDependencyLoading:
+
+			ClassFlags |= CLASS_NeedsDeferredDependencyLoading;
+			break;
+
+		default:
+			FError::Throwf(TEXT("Unknown class specifier '%s'"), *PropSpecifier.Key);
+		}
+	}
+}
+
+void FUnrealClassDefinitionInfo::MergeShowCategories()
+{
+	for (const FString& Value : ShowCategories)
+	{
+		// if we didn't find this specific category path in the HideCategories metadata
+		if (HideCategories.RemoveSwap(Value) == 0)
+		{
+			TArray<FString> SubCategoryList;
+			Value.ParseIntoArray(SubCategoryList, TEXT("|"), true);
+
+			FString SubCategoryPath;
+			// look to see if any of the parent paths are excluded in the HideCategories list
+			for (int32 CategoryPathIndex = 0; CategoryPathIndex < SubCategoryList.Num() - 1; ++CategoryPathIndex)
+			{
+				SubCategoryPath += SubCategoryList[CategoryPathIndex];
+				// if we're hiding a parent category, then we need to flag this sub category for show
+				if (HideCategories.Contains(SubCategoryPath))
+				{
+					ShowSubCatgories.AddUnique(Value);
+					break;
+				}
+				SubCategoryPath += "|";
+			}
+		}
+	}
+	// Once the categories have been merged, empty the array as we will no longer need it nor should we use it
+	ShowCategories.Empty();
+}
+
+void FUnrealClassDefinitionInfo::MergeClassCategories()
+{
+	UClass* Class = GetClass();
+	TArray<FString> ParentHideCategories;
+	TArray<FString> ParentShowSubCatgories;
+	TArray<FString> ParentHideFunctions;
+	TArray<FString> ParentAutoExpandCategories;
+	TArray<FString> ParentAutoCollapseCategories;
+	GetHideCategories(ParentHideCategories);
+	GetShowCategories(ParentShowSubCatgories);
+	Class->GetHideFunctions(ParentHideFunctions);
+	Class->GetAutoExpandCategories(ParentAutoExpandCategories);
+	Class->GetAutoCollapseCategories(ParentAutoCollapseCategories);
+
+	// Add parent categories. We store the opposite of HideCategories and HideFunctions in a separate array anyway.
+	HideCategories.Append(MoveTemp(ParentHideCategories));
+	ShowSubCatgories.Append(MoveTemp(ParentShowSubCatgories));
+	HideFunctions.Append(MoveTemp(ParentHideFunctions));
+
+	MergeShowCategories();
+
+	// Merge ShowFunctions and HideFunctions
+	for (const FString& Value : ShowFunctions)
+	{
+		HideFunctions.RemoveSwap(Value);
+	}
+	ShowFunctions.Empty();
+
+	// Merge DontAutoCollapseCategories and AutoCollapseCategories
+	for (const FString& Value : DontAutoCollapseCategories)
+	{
+		AutoCollapseCategories.RemoveSwap(Value);
+	}
+	DontAutoCollapseCategories.Empty();
+
+	// Merge ShowFunctions and HideFunctions
+	for (const FString& Value : ShowFunctions)
+	{
+		HideFunctions.RemoveSwap(Value);
+	}
+	ShowFunctions.Empty();
+
+	// Merge AutoExpandCategories and AutoCollapseCategories (we still want to keep AutoExpandCategories though!)
+	for (const FString& Value : AutoExpandCategories)
+	{
+		AutoCollapseCategories.RemoveSwap(Value);
+		ParentAutoCollapseCategories.RemoveSwap(Value);
+	}
+
+	// Do the same as above but the other way around
+	for (const FString& Value : AutoCollapseCategories)
+	{
+		AutoExpandCategories.RemoveSwap(Value);
+		ParentAutoExpandCategories.RemoveSwap(Value);
+	}
+
+	// Once AutoExpandCategories and AutoCollapseCategories for THIS class have been parsed, add the parent inherited categories
+	AutoCollapseCategories.Append(MoveTemp(ParentAutoCollapseCategories));
+	AutoExpandCategories.Append(MoveTemp(ParentAutoExpandCategories));
+}
+
+void FUnrealClassDefinitionInfo::MergeAndValidateClassFlags(const FString& DeclaredClassName, uint32 PreviousClassFlags)
+{
+	UClass* Class = GetClass();
+	if (bWantsToBePlaceable)
+	{
+		if (!(Class->ClassFlags & CLASS_NotPlaceable))
+		{
+			FError::Throwf(TEXT("The 'placeable' specifier is only allowed on classes which have a base class that's marked as not placeable. Classes are assumed to be placeable by default."));
+		}
+		Class->ClassFlags &= ~CLASS_NotPlaceable;
+		bWantsToBePlaceable = false; // Reset this flag after it's been merged
+	}
+
+	// Now merge all remaining flags/properties
+	Class->ClassFlags |= ClassFlags;
+	Class->ClassConfigName = FName(*ConfigName);
+
+	SetAndValidateWithinClass();
+	SetAndValidateConfigName();
+
+	if (!!(Class->ClassFlags & CLASS_EditInlineNew))
+	{
+		// don't allow actor classes to be declared editinlinenew
+		if (IsActorClass(Class))
+		{
+			FError::Throwf(TEXT("Invalid class attribute: Creating actor instances via the property window is not allowed"));
+		}
+	}
+
+	// Make sure both RequiredAPI and MinimalAPI aren't specified
+	if (Class->HasAllClassFlags(CLASS_MinimalAPI | CLASS_RequiredAPI))
+	{
+		FError::Throwf(TEXT("MinimalAPI cannot be specified when the class is fully exported using a MODULENAME_API macro"));
+	}
+
+	// All classes must start with a valid Unreal prefix
+	const FString ExpectedClassName = GetNameWithPrefix(GetClass());
+	if (DeclaredClassName != ExpectedClassName)
+	{
+		FError::Throwf(TEXT("Class name '%s' is invalid, should be identified as '%s'"), *DeclaredClassName, *ExpectedClassName);
+	}
+
+	if ((Class->ClassFlags & CLASS_NoExport))
+	{
+		// if the class's class flags didn't contain CLASS_NoExport before it was parsed, it means either:
+		// a) the DECLARE_CLASS macro for this native class doesn't contain the CLASS_NoExport flag (this is an error)
+		// b) this is a new native class, which isn't yet hooked up to static registration (this is OK)
+		if (!(Class->ClassFlags & CLASS_Intrinsic) && (PreviousClassFlags & CLASS_NoExport) == 0 &&
+			(PreviousClassFlags & CLASS_Native) != 0)	// a new native class (one that hasn't been compiled into C++ yet) won't have this set
+		{
+			FError::Throwf(TEXT("'noexport': Must include CLASS_NoExport in native class declaration"));
+		}
+	}
+
+	if (!Class->HasAnyClassFlags(CLASS_Abstract) && ((PreviousClassFlags & CLASS_Abstract) != 0))
+	{
+		if (Class->HasAnyClassFlags(CLASS_NoExport))
+		{
+			FError::Throwf(TEXT("'abstract': NoExport class missing abstract keyword from class declaration (must change C++ version first)"));
+			Class->ClassFlags |= CLASS_Abstract;
+		}
+		else if (Class->IsNative())
+		{
+			FError::Throwf(TEXT("'abstract': missing abstract keyword from class declaration - class will no longer be exported as abstract"));
+		}
+	}
+}
+
+void FUnrealClassDefinitionInfo::SetAndValidateConfigName()
+{
+	UClass* Class = GetClass();
+	if (ConfigName.IsEmpty() == false)
+	{
+		// if the user specified "inherit", we're just going to use the parent class's config filename
+		// this is not actually necessary but it can be useful for explicitly communicating config-ness
+		if (ConfigName == TEXT("inherit"))
+		{
+			UClass* SuperClass = Class->GetSuperClass();
+			if (!SuperClass)
+			{
+				FError::Throwf(TEXT("Cannot inherit config filename: %s has no super class"), *Class->GetName());
+			}
+
+			if (SuperClass->ClassConfigName == NAME_None)
+			{
+				FError::Throwf(TEXT("Cannot inherit config filename: parent class %s is not marked config."), *SuperClass->GetPathName());
+			}
+		}
+		else
+		{
+			// otherwise, set the config name to the parsed identifier
+			Class->ClassConfigName = FName(*ConfigName);
+		}
+	}
+	else
+	{
+		// Invalidate config name if not specifically declared.
+		Class->ClassConfigName = NAME_None;
+	}
+}
+
+void FUnrealClassDefinitionInfo::SetAndValidateWithinClass()
+{
+	UClass* Class = GetClass();
+
+	// Process all of the class specifiers
+	if (ClassWithin.IsEmpty() == false)
+	{
+		UClass* RequiredWithinClass = FClasses::FindClass(*ClassWithin);
+		if (!RequiredWithinClass)
+		{
+			FError::Throwf(TEXT("Within class '%s' not found."), *ClassWithin);
+		}
+		if (RequiredWithinClass->IsChildOf(UInterface::StaticClass()))
+		{
+			FError::Throwf(TEXT("Classes cannot be 'within' interfaces"));
+		}
+		else if (Class->ClassWithin == NULL || Class->ClassWithin == UObject::StaticClass() || RequiredWithinClass->IsChildOf(Class->ClassWithin))
+		{
+			Class->ClassWithin = RequiredWithinClass;
+		}
+		else if (Class->ClassWithin != RequiredWithinClass)
+		{
+			FError::Throwf(TEXT("%s must be within %s, not %s"), *Class->GetPathName(), *Class->ClassWithin->GetPathName(), *RequiredWithinClass->GetPathName());
+		}
+	}
+	else
+	{
+		// Make sure there is a valid within
+		Class->ClassWithin = Class->GetSuperClass()
+			? Class->GetSuperClass()->ClassWithin
+			: UObject::StaticClass();
+	}
+
+	UClass* ExpectedWithin = Class->GetSuperClass()
+		? Class->GetSuperClass()->ClassWithin
+		: UObject::StaticClass();
+
+	if (!Class->ClassWithin->IsChildOf(ExpectedWithin))
+	{
+		FError::Throwf(TEXT("Parent class declared within '%s'.  Cannot override within class with '%s' since it isn't a child"), *ExpectedWithin->GetName(), *Class->ClassWithin->GetName());
+	}
+}
+
+void FUnrealClassDefinitionInfo::GetHideCategories(TArray<FString>& OutHideCategories) const
+{
+	UClass* Class = GetClass();
+	if (Class->HasMetaData(FHeaderParserNames::NAME_HideCategories))
+	{
+		const FString& LocalHideCategories = Class->GetMetaData(FHeaderParserNames::NAME_HideCategories);
+		LocalHideCategories.ParseIntoArray(OutHideCategories, TEXT(" "), true);
+	}
+}
+
+void FUnrealClassDefinitionInfo::GetShowCategories(TArray<FString>& OutShowCategories) const
+{
+	UClass* Class = GetClass();
+	if (Class->HasMetaData(FHeaderParserNames::NAME_ShowCategories))
+	{
+		const FString& LocalShowCategories = Class->GetMetaData(FHeaderParserNames::NAME_ShowCategories);
+		LocalShowCategories.ParseIntoArray(OutShowCategories, TEXT(" "), true);
+	}
+}
+
+void FUnrealClassDefinitionInfo::MergeCategoryMetaData(TMap<FName, FString>& InMetaData) const
+{
+	if (ClassGroupNames.Num()) { InMetaData.Add(NAME_ClassGroupNames, FString::Join(ClassGroupNames, TEXT(" "))); }
+	if (AutoCollapseCategories.Num()) { InMetaData.Add(FHeaderParserNames::NAME_AutoCollapseCategories, FString::Join(AutoCollapseCategories, TEXT(" "))); }
+	if (HideCategories.Num()) { InMetaData.Add(FHeaderParserNames::NAME_HideCategories, FString::Join(HideCategories, TEXT(" "))); }
+	if (ShowSubCatgories.Num()) { InMetaData.Add(FHeaderParserNames::NAME_ShowCategories, FString::Join(ShowSubCatgories, TEXT(" "))); }
+	if (SparseClassDataTypes.Num()) { InMetaData.Add(FHeaderParserNames::NAME_SparseClassDataTypes, FString::Join(SparseClassDataTypes, TEXT(" "))); }
+	if (HideFunctions.Num()) { InMetaData.Add(FHeaderParserNames::NAME_HideFunctions, FString::Join(HideFunctions, TEXT(" "))); }
+	if (AutoExpandCategories.Num()) { InMetaData.Add(FHeaderParserNames::NAME_AutoExpandCategories, FString::Join(AutoExpandCategories, TEXT(" "))); }
+}
+
+void FUnrealClassDefinitionInfo::GetSparseClassDataTypes(TArray<FString>& OutSparseClassDataTypes) const
+{
+	UClass* Class = GetClass();
+	if (Class->HasMetaData(FHeaderParserNames::NAME_SparseClassDataTypes))
+	{
+		const FString& LocalSparseClassDataTypes = Class->GetMetaData(FHeaderParserNames::NAME_SparseClassDataTypes);
+		LocalSparseClassDataTypes.ParseIntoArray(OutSparseClassDataTypes, TEXT(" "), true);
+	}
+}
+
 void FUnrealFunctionDefinitionInfo::AddProperty(FUnrealPropertyDefinitionInfo& PropertyDef)
 {
 	const FProperty* Prop = PropertyDef.GetProperty();
diff --git a/Engine/Source/Programs/UnrealHeaderTool/Private/UnrealTypeDefinitionInfo.h b/Engine/Source/Programs/UnrealHeaderTool/Private/UnrealTypeDefinitionInfo.h
--- a/Engine/Source/Programs/UnrealHeaderTool/Private/UnrealTypeDefinitionInfo.h
+++ b/Engine/Source/Programs/UnrealHeaderTool/Private/UnrealTypeDefinitionInfo.h
@@ -3,12 +3,12 @@
 #pragma once
 
 #include "CoreMinimal.h"
-#include "ClassDeclarationMetaData.h"
 #include "Exceptions.h"
 #include "Templates/SharedPointer.h"
 #include <atomic>
 
-#include "ParserHelper.h" // TEMPORARY
+#include "BaseParser.h"
+#include "ParserHelper.h"
 
 // Forward declarations.
 class FClass;
@@ -36,6 +36,13 @@ class FUnrealTypeDefinitionInfo;							// Base for all types, provides virtual m
 				class FUnrealClassDefinitionInfo;			// Represents UClass
 				class FUnrealFunctionDefinitionInfo;		// Represents UFunction
 
+enum class EEnforceInterfacePrefix
+{
+	None,
+	I,
+	U
+};
+
 enum class EUnderlyingEnumType
 {
 	Unspecified,
@@ -307,6 +314,30 @@ public:
 		FUHTException::Throwf(*this, TEXT("Meta data can not be set for a definition of this type."));
 	}
 
+	/**
+	 * Helper function that checks if the field is a dynamic type (can be constructed post-startup)
+	 */
+	static bool IsDynamic(const UField* Field);
+	static bool IsDynamic(const FField* Field);
+
+	/**
+	 * Helper function that checks if the field is a dynamic type
+	 */
+	virtual bool IsDynamic() const
+	{
+		return false;
+	}
+
+	/**
+	 * Helper function that checks if the field is belongs to a dynamic type
+	 */
+	virtual bool IsOwnedByDynamicType() const
+	{
+		return false;
+	}
+
+	static FString GetNameWithPrefix(const UClass* Class, EEnforceInterfacePrefix EnforceInterfacePrefix = EEnforceInterfacePrefix::None);
+
 protected:
 	explicit FUnrealTypeDefinitionInfo(FString&& InNameCPP)
 		: NameCPP(MoveTemp(InNameCPP))
@@ -351,6 +382,11 @@ public:
 		return TEXT("UProperty");
 	}
 
+	/**
+	 * Perform any post parsing finalization and validation
+	 */
+	virtual void PostParseFinalize() override;
+
 	/**
 	 * Return the Engine instance associated with the compiler instance
 	 */
@@ -501,9 +537,28 @@ public:
 		return GetPropertyBase().MatchesType(Other.GetPropertyBase(), bDisallowGeneralization, bIgnoreImplementedInterfaces);
 	}
 
+	/**
+	 * Return the type package name
+	 */
+	const FString& GetTypePackageName() const
+	{
+		return TypePackageName;
+	}
+
+	/**
+	 * Helper function that checks if the field is a dynamic type
+	 */
+	virtual bool IsDynamic() const override;
+
+	/**
+	 * Helper function that checks if the field is belongs to a dynamic type
+	 */
+	virtual bool IsOwnedByDynamicType() const override;
+
 private:
 	FPropertyBase PropertyBase;
 	FString ArrayDimensions;
+	FString TypePackageName;
 	FUnrealPropertyDefinitionInfo* KeyPropDef = nullptr;
 	FUnrealPropertyDefinitionInfo* ValuePropDef = nullptr;
 	FProperty* Property = nullptr;
@@ -759,15 +814,34 @@ public:
 		return ExternDecl[bRequiresValidObject];
 	}
 
+	/** 
+	 * Return the type package name
+	 */
+	const FString& GetTypePackageName() const
+	{
+		return TypePackageName;
+	}
+
 	/**
 	 * Add meta data for the given definition
 	 */
 	virtual void AddMetaData(TMap<FName, FString>&& InMetaData) override;
 
+	/**
+	 * Helper function that checks if the field is a dynamic type
+	 */
+	virtual bool IsDynamic() const override;
+
+	/**
+	 * Helper function that checks if the field is belongs to a dynamic type
+	 */
+	virtual bool IsOwnedByDynamicType() const override;
+
 private:
 	FString SingletonName[2];
 	FString SingletonNameChopped[2];
 	FString ExternDecl[2];
+	FString TypePackageName;
 };
 
 /**
@@ -1089,7 +1163,6 @@ private:
  */
 class FUnrealClassDefinitionInfo
 	: public FUnrealStructDefinitionInfo
-	, public FClassDeclarationMetaData
 {
 public:
 	FUnrealClassDefinitionInfo(FUnrealSourceFile& InSourceFile, int32 InLineNumber, FString&& InNameCPP, bool bInIsInterface);
@@ -1263,10 +1336,68 @@ public:
 	 */
 	static bool HierarchyHasAllClassFlags(UClass* Class, EClassFlags FlagsToCheck);
 
+	/**
+	* Parse Class's properties to generate its declaration data.
+	*
+	* @param	InClassSpecifiers Class properties collected from its UCLASS macro
+	* @param	InRequiredAPIMacroIfPresent *_API macro if present (empty otherwise)
+	* @param	OutClassData Parsed class meta data
+	*/
+	void ParseClassProperties(TArray<FPropertySpecifier>&& InClassSpecifiers, const FString& InRequiredAPIMacroIfPresent);
+
+	/**
+	* Merges all category properties with the class which at this point only has its parent propagated categories
+	*/
+	void MergeClassCategories();
+
+	/**
+	* Merges all class flags and validates them
+	*
+	* @param	DeclaredClassName Name this class was declared with (for validation)
+	* @param	PreviousClassFlags Class flags before resetting the class (for validation)
+	*/
+	void MergeAndValidateClassFlags(const FString& DeclaredClassName, uint32 PreviousClassFlags);
+
+	/**
+	 * Add the category meta data
+	 */
+	void MergeCategoryMetaData(TMap<FName, FString>& InMetaData) const;
+
+
+	void GetSparseClassDataTypes(TArray<FString>& OutSparseClassDataTypes) const;
+
+public:
+	EClassFlags ClassFlags = CLASS_None;
+	FString ClassWithin;
+	FString ConfigName;
+	TMap<FName, FString> MetaData;
+
 private:
+	/** Merges all 'show' categories */
+	void MergeShowCategories();
+	/** Sets and validates 'within' property */
+	void SetAndValidateWithinClass();
+	/** Sets and validates 'ConfigName' property */
+	void SetAndValidateConfigName();
+
+	void GetHideCategories(TArray<FString>& OutHideCategories) const;
+	void GetShowCategories(TArray<FString>& OutShowCategories) const;
+
 	FString EnclosingDefine;
 	ESerializerArchiveType ArchiveType = ESerializerArchiveType::None;
+	TArray<FString> ShowCategories;
+	TArray<FString> ShowFunctions;
+	TArray<FString> DontAutoCollapseCategories;
+	TArray<FString> HideCategories;
+	TArray<FString> ShowSubCatgories;
+	TArray<FString> HideFunctions;
+	TArray<FString> AutoExpandCategories;
+	TArray<FString> AutoCollapseCategories;
+	TArray<FString> DependsOn;
+	TArray<FString> ClassGroupNames;
+	TArray<FString> SparseClassDataTypes;
 	bool bIsInterface = false;
+	bool bWantsToBePlaceable = false;
 };
 
 template <typename To>
@@ -1482,12 +1613,12 @@ public:
 		return StructDef;
 	}
 protected:
-	inline T* const * UpdateRange(const FUnrealStructDefinitionInfo* InStructDef)
+	inline T* const* UpdateRange(const FUnrealStructDefinitionInfo* InStructDef)
 	{
 		if (InStructDef)
 		{
 			auto& Array = GetFieldsFromDef<T>(*InStructDef);
-			T* const * Out = Array.GetData();
+			T* const* Out = Array.GetData();
 			End = Out + Array.Num();
 			return Out;
 		}
@@ -1500,7 +1631,7 @@ protected:
 
 	inline void IterateToNext()
 	{
-		T* const * CurrentField = Field;
+		T* const* CurrentField = Field;
 		const FUnrealStructDefinitionInfo* CurrentStructDef = StructDef;
 
 		while (CurrentStructDef)
