Change PropertyBase to reference new types.

#rnx

diff --git a/Engine/Source/Programs/UnrealHeaderTool/Private/CodeGenerator.cpp b/Engine/Source/Programs/UnrealHeaderTool/Private/CodeGenerator.cpp
--- a/Engine/Source/Programs/UnrealHeaderTool/Private/CodeGenerator.cpp
+++ b/Engine/Source/Programs/UnrealHeaderTool/Private/CodeGenerator.cpp
@@ -766,26 +766,29 @@ void FNativeClassHeaderGenerator::ExportProperties(FOutputDevice& Out, FUnrealSt
 
 static FString GNullPtr(TEXT("nullptr"));
 
-const FString& FNativeClassHeaderGenerator::GetPackageSingletonName(UPackage* Item, TSet<FString>* UniqueCrossModuleReferences)
+const FString& FNativeClassHeaderGenerator::GetPackageSingletonName(FUnrealPackageDefinitionInfo& PackageDef, TSet<FString>* UniqueCrossModuleReferences)
 {
-	check(Item);
-	FUnrealPackageDefinitionInfo& Package = GTypeDefinitionInfoMap.FindChecked<FUnrealPackageDefinitionInfo>(Item);
-	Package.AddCrossModuleReference(UniqueCrossModuleReferences);
-	return Package.GetSingletonName();
+	PackageDef.AddCrossModuleReference(UniqueCrossModuleReferences);
+	return PackageDef.GetSingletonName();
 }
 
-const FString& FNativeClassHeaderGenerator::GetPackageSingletonNameFuncAddr(UPackage* Item, TSet<FString>* UniqueCrossModuleReferences)
+const FString& FNativeClassHeaderGenerator::GetPackageSingletonNameFuncAddr(FUnrealPackageDefinitionInfo& PackageDef, TSet<FString>* UniqueCrossModuleReferences)
 {
-	check(Item);
-	FUnrealPackageDefinitionInfo& Package = GTypeDefinitionInfoMap.FindChecked<FUnrealPackageDefinitionInfo>(Item);
-	Package.AddCrossModuleReference(UniqueCrossModuleReferences);
-	return Package.GetSingletonNameChopped();
+	PackageDef.AddCrossModuleReference(UniqueCrossModuleReferences);
+	return PackageDef.GetSingletonNameChopped();
 }
 
-const FString& FNativeClassHeaderGenerator::GetSingletonNameFuncAddr(FUnrealFieldDefinitionInfo& Field, TSet<FString>* UniqueCrossModuleReferences, bool bRequiresValidObject)
+const FString& FNativeClassHeaderGenerator::GetSingletonNameFuncAddr(FUnrealFieldDefinitionInfo* FieldDef, TSet<FString>* UniqueCrossModuleReferences, bool bRequiresValidObject)
 {
-	Field.AddCrossModuleReference(UniqueCrossModuleReferences, bRequiresValidObject);
-	return Field.GetSingletonNameChopped(bRequiresValidObject);
+	if (!FieldDef)
+	{
+		return GNullPtr;
+	}
+	else
+	{
+		FieldDef->AddCrossModuleReference(UniqueCrossModuleReferences, bRequiresValidObject);
+		return FieldDef->GetSingletonNameChopped(bRequiresValidObject);
+	}
 }
 
 const FString& FNativeClassHeaderGenerator::GetSingletonNameFuncAddr(UField* Item, TSet<FString>* UniqueCrossModuleReferences, bool bRequiresValidObject)
@@ -797,59 +800,53 @@ const FString& FNativeClassHeaderGenerator::GetSingletonNameFuncAddr(UField* Ite
 	else
 	{
 		FUnrealFieldDefinitionInfo& Field = GTypeDefinitionInfoMap.FindChecked<FUnrealFieldDefinitionInfo>(Item);
-		return GetSingletonNameFuncAddr(Field, UniqueCrossModuleReferences, bRequiresValidObject);
+		return GetSingletonNameFuncAddr(&Field, UniqueCrossModuleReferences, bRequiresValidObject);
 	}
 }
 
-void FNativeClassHeaderGenerator::GetPropertyTag(FUHTStringBuilder& Out, UObject* Object)
+void FNativeClassHeaderGenerator::GetPropertyTag(FUHTStringBuilder& Out, FUnrealPropertyDefinitionInfo& PropDef)
 {
-	if (Object != nullptr)
-	{
-		GTypeDefinitionInfoMap.FindChecked(Object).GetHashTag(Out);
-	}
-}
-
-void FNativeClassHeaderGenerator::GetPropertyTag(FUHTStringBuilder& Out, FProperty* Prop)
-{
-	if (Prop == nullptr)
-	{
-		return;
-	}
+	const FPropertyBase& PropertyBase = PropDef.GetPropertyBase();
+	FProperty* Prop = PropDef.GetProperty();
 
 #if UHT_ENABLE_VALUE_PROPERTY_TAG
 	if (FNumericProperty* TypedProp = CastField<FNumericProperty>(Prop))
 	{
-		GetPropertyTag(Out, TypedProp->GetIntPropertyEnum());
+		if (PropertyBase.EnumDef != nullptr)
+		{
+			PropertyBase.EnumDef->GetHashTag(Out);
+		}
 		return;
 	}
 
 	if (FEnumProperty* TypedProp = CastField<FEnumProperty>(Prop))
 	{
-		GetPropertyTag(Out, TypedProp->Enum);
+		PropertyBase.EnumDef->GetHashTag(Out);
+		return;
 	}
 
 	if (FStructProperty* TypedProp = CastField<FStructProperty>(Prop))
 	{
-		GetPropertyTag(Out, TypedProp->Struct);
+		PropertyBase.ScriptStructDef->GetHashTag(Out);
 		return;
 	}
 
 	if (FArrayProperty* TypedProp = CastField<FArrayProperty>(Prop))
 	{
-		GetPropertyTag(Out, TypedProp->Inner);
+		GetPropertyTag(Out, PropDef.GetValuePropDef());
 		return;
 	}
 
 	if (FMapProperty* TypedProp = CastField<FMapProperty>(Prop))
 	{
-		GetPropertyTag(Out, TypedProp->KeyProp);
-		GetPropertyTag(Out, TypedProp->ValueProp);
+		GetPropertyTag(Out, PropDef.GetKeyPropDef());
+		GetPropertyTag(Out, PropDef.GetValuePropDef());
 		return;
 	}
 
 	if (FSetProperty* TypedProp = CastField<FSetProperty>(Prop))
 	{
-		GetPropertyTag(Out, TypedProp->ElementProp);
+		GetPropertyTag(Out, PropDef.GetValuePropDef());
 		return;
 	}
 #endif
@@ -857,55 +854,55 @@ void FNativeClassHeaderGenerator::GetPropertyTag(FUHTStringBuilder& Out, FProper
 #if UHT_ENABLE_PTR_PROPERTY_TAG
 	if (FClassProperty* TypedProp = CastField<FClassProperty>(Prop))
 	{
-		GetPropertyTag(Out, TypedProp->PropertyClass);
+		PropertyBase.ClassDef->GetHashTag(Out);
 		return;
 	}
 
 	if (FObjectProperty* TypedProp = CastField<FObjectProperty>(Prop))
 	{
-		GetPropertyTag(Out, TypedProp->PropertyClass);
+		PropertyBase.ClassDef->GetHashTag(Out);
 		return;
 	}
 
 	if (FSoftClassProperty* TypedProp = CastField<FSoftClassProperty>(Prop))
 	{
-		GetPropertyTag(Out, TypedProp->PropertyClass);
+		PropertyBase.ClassDef->GetHashTag(Out);
 		return;
 	}
 
 	if (FSoftObjectProperty* TypedProp = CastField<FSoftObjectProperty>(Prop))
 	{
-		GetPropertyTag(Out, TypedProp->PropertyClass);
+		PropertyBase.ClassDef->GetHashTag(Out);
 		return;
 	}
 
 	if (FClassPtrProperty* TypedProp = CastField<FClassPtrProperty>(Prop))
 	{
-		GetPropertyTag(Out, TypedProp->PropertyClass);
+		PropertyBase.ClassDef->GetHashTag(Out);
 		return;
 	}
 
 	if (FObjectPtrProperty* TypedProp = CastField<FObjectPtrProperty>(Prop))
 	{
-		GetPropertyTag(Out, TypedProp->PropertyClass);
+		PropertyBase.ClassDef->GetHashTag(Out);
 		return;
 	}
 
 	if (FInterfaceProperty* TypedProp = CastField<FInterfaceProperty>(Prop))
 	{
-		GetPropertyTag(Out, TypedProp->InterfaceClass);
+		PropertyBase.ClassDef->GetHashTag(Out);
 		return;
 	}
 
 	if (FWeakObjectProperty* TypedProp = CastField<FWeakObjectProperty>(Prop))
 	{
-		GetPropertyTag(Out, TypedProp->PropertyClass);
+		PropertyBase.ClassDef->GetHashTag(Out);
 		return;
 	}
 
 	if (FLazyObjectProperty* TypedProp = CastField<FLazyObjectProperty>(Prop))
 	{
-		GetPropertyTag(Out, TypedProp->PropertyClass);
+		PropertyBase.ClassDef->GetHashTag(Out);
 		return;
 	}
 #endif
@@ -913,6 +910,7 @@ void FNativeClassHeaderGenerator::GetPropertyTag(FUHTStringBuilder& Out, FProper
 
 void FNativeClassHeaderGenerator::PropertyNew(FOutputDevice& DeclOut, FOutputDevice& Out, FReferenceGatherers& OutReferenceGatherers, FUnrealPropertyDefinitionInfo& PropertyDef, const TCHAR* OffsetStr, const TCHAR* Name, const TCHAR* DeclSpaces, const TCHAR* Spaces, const TCHAR* SourceStruct) const
 {
+	const FPropertyBase& PropertyBase   = PropertyDef.GetPropertyBase();
 	FProperty*	   Prop                 = PropertyDef.GetProperty();
 	FString        PropName             = CreateUTF8LiteralString(FNativeClassHeaderGenerator::GetOverriddenName(Prop));
 	FString        PropNameDep          = Prop->HasAllPropertyFlags(CPF_Deprecated) ? Prop->GetName() + TEXT("_DEPRECATED") : Prop->GetName();
@@ -920,7 +918,7 @@ void FNativeClassHeaderGenerator::PropertyNew(FOutputDevice& DeclOut, FOutputDev
 	EPropertyFlags PropFlags            = Prop->PropertyFlags & ~CPF_ComputedFlags;
 
 	FUHTStringBuilder PropTag;
-	GetPropertyTag(PropTag, Prop);
+	GetPropertyTag(PropTag, PropertyDef);
 
 	FString PropNotifyFunc = (Prop->RepNotifyFunc != NAME_None) ? CreateUTF8LiteralString(*Prop->RepNotifyFunc.ToString()) : TEXT("nullptr");
 
@@ -951,7 +949,7 @@ void FNativeClassHeaderGenerator::PropertyNew(FOutputDevice& DeclOut, FOutputDev
 			FPropertyObjectFlags,
 			*ArrayDim,
 			OffsetStr,
-			*GetSingletonNameFuncAddr(TypedProp->Enum, OutReferenceGatherers.UniqueCrossModuleReferences),
+			*GetSingletonNameFuncAddr(PropertyBase.EnumDef, OutReferenceGatherers.UniqueCrossModuleReferences),
 			*MetaDataParams,
 			*PropTag
 		);
@@ -1234,7 +1232,7 @@ void FNativeClassHeaderGenerator::PropertyNew(FOutputDevice& DeclOut, FOutputDev
 			FPropertyObjectFlags,
 			*ArrayDim,
 			OffsetStr,
-			*GetSingletonNameFuncAddr(TypedProp->MetaClass, OutReferenceGatherers.UniqueCrossModuleReferences, false),
+			*GetSingletonNameFuncAddr(PropertyBase.MetaClassDef, OutReferenceGatherers.UniqueCrossModuleReferences, false),
 			*MetaDataParams,
 			*PropTag
 		);
@@ -1256,7 +1254,7 @@ void FNativeClassHeaderGenerator::PropertyNew(FOutputDevice& DeclOut, FOutputDev
 			FPropertyObjectFlags,
 			*ArrayDim,
 			OffsetStr,
-			*GetSingletonNameFuncAddr(TypedProp->PropertyClass, OutReferenceGatherers.UniqueCrossModuleReferences, false),
+			*GetSingletonNameFuncAddr(PropertyBase.ClassDef, OutReferenceGatherers.UniqueCrossModuleReferences, false),
 			*MetaDataParams,
 			*PropTag
 		);
@@ -1278,7 +1276,7 @@ void FNativeClassHeaderGenerator::PropertyNew(FOutputDevice& DeclOut, FOutputDev
 			FPropertyObjectFlags,
 			*ArrayDim,
 			OffsetStr,
-			*GetSingletonNameFuncAddr(TypedProp->PropertyClass, OutReferenceGatherers.UniqueCrossModuleReferences, false),
+			*GetSingletonNameFuncAddr(PropertyBase.ClassDef, OutReferenceGatherers.UniqueCrossModuleReferences, false),
 			*MetaDataParams,
 			*PropTag
 		);
@@ -1300,7 +1298,7 @@ void FNativeClassHeaderGenerator::PropertyNew(FOutputDevice& DeclOut, FOutputDev
 			FPropertyObjectFlags,
 			*ArrayDim,
 			OffsetStr,
-			*GetSingletonNameFuncAddr(TypedProp->PropertyClass, OutReferenceGatherers.UniqueCrossModuleReferences, false),
+			*GetSingletonNameFuncAddr(PropertyBase.ClassDef, OutReferenceGatherers.UniqueCrossModuleReferences, false),
 			*MetaDataParams,
 			*PropTag
 		);
@@ -1322,8 +1320,8 @@ void FNativeClassHeaderGenerator::PropertyNew(FOutputDevice& DeclOut, FOutputDev
 			FPropertyObjectFlags,
 			*ArrayDim,
 			OffsetStr,
-			*GetSingletonNameFuncAddr(TypedProp->MetaClass, OutReferenceGatherers.UniqueCrossModuleReferences, false),
-			*GetSingletonNameFuncAddr(TypedProp->PropertyClass, OutReferenceGatherers.UniqueCrossModuleReferences, false),
+			*GetSingletonNameFuncAddr(PropertyBase.MetaClassDef, OutReferenceGatherers.UniqueCrossModuleReferences, false),
+			*GetSingletonNameFuncAddr(PropertyBase.ClassDef, OutReferenceGatherers.UniqueCrossModuleReferences, false),
 			*MetaDataParams,
 			*PropTag
 		);
@@ -1345,7 +1343,7 @@ void FNativeClassHeaderGenerator::PropertyNew(FOutputDevice& DeclOut, FOutputDev
 			FPropertyObjectFlags,
 			*ArrayDim,
 			OffsetStr,
-			*GetSingletonNameFuncAddr(TypedProp->PropertyClass, OutReferenceGatherers.UniqueCrossModuleReferences, false),
+			*GetSingletonNameFuncAddr(PropertyBase.ClassDef, OutReferenceGatherers.UniqueCrossModuleReferences, false),
 			*MetaDataParams,
 			*PropTag
 		);
@@ -1367,8 +1365,8 @@ void FNativeClassHeaderGenerator::PropertyNew(FOutputDevice& DeclOut, FOutputDev
 			FPropertyObjectFlags,
 			*ArrayDim,
 			OffsetStr,
-			*GetSingletonNameFuncAddr(TypedProp->MetaClass, OutReferenceGatherers.UniqueCrossModuleReferences, false),
-			*GetSingletonNameFuncAddr(TypedProp->PropertyClass, OutReferenceGatherers.UniqueCrossModuleReferences, false),
+			*GetSingletonNameFuncAddr(PropertyBase.MetaClassDef, OutReferenceGatherers.UniqueCrossModuleReferences, false),
+			*GetSingletonNameFuncAddr(PropertyBase.ClassDef, OutReferenceGatherers.UniqueCrossModuleReferences, false),
 			*MetaDataParams,
 			*PropTag
 		);
@@ -1390,7 +1388,7 @@ void FNativeClassHeaderGenerator::PropertyNew(FOutputDevice& DeclOut, FOutputDev
 			FPropertyObjectFlags,
 			*ArrayDim,
 			OffsetStr,
-			*GetSingletonNameFuncAddr(TypedProp->PropertyClass, OutReferenceGatherers.UniqueCrossModuleReferences, false),
+			*GetSingletonNameFuncAddr(PropertyBase.ClassDef, OutReferenceGatherers.UniqueCrossModuleReferences, false),
 			*MetaDataParams,
 			*PropTag
 		);
@@ -1412,7 +1410,7 @@ void FNativeClassHeaderGenerator::PropertyNew(FOutputDevice& DeclOut, FOutputDev
 			FPropertyObjectFlags,
 			*ArrayDim,
 			OffsetStr,
-			*GetSingletonNameFuncAddr(TypedProp->InterfaceClass, OutReferenceGatherers.UniqueCrossModuleReferences, false),
+			*GetSingletonNameFuncAddr(PropertyBase.ClassDef, OutReferenceGatherers.UniqueCrossModuleReferences, false),
 			*MetaDataParams,
 			*PropTag
 		);
@@ -1541,7 +1539,7 @@ void FNativeClassHeaderGenerator::PropertyNew(FOutputDevice& DeclOut, FOutputDev
 			FPropertyObjectFlags,
 			*ArrayDim,
 			OffsetStr,
-			*GetSingletonNameFuncAddr(TypedProp->Struct, OutReferenceGatherers.UniqueCrossModuleReferences),
+			*GetSingletonNameFuncAddr(PropertyBase.ScriptStructDef, OutReferenceGatherers.UniqueCrossModuleReferences),
 			*MetaDataParams,
 			*PropTag
 		);
@@ -1563,7 +1561,7 @@ void FNativeClassHeaderGenerator::PropertyNew(FOutputDevice& DeclOut, FOutputDev
 			FPropertyObjectFlags,
 			*ArrayDim,
 			OffsetStr,
-			*GetSingletonNameFuncAddr(TypedProp->SignatureFunction, OutReferenceGatherers.UniqueCrossModuleReferences),
+			*GetSingletonNameFuncAddr(PropertyBase.FunctionDef, OutReferenceGatherers.UniqueCrossModuleReferences),
 			*MetaDataParams,
 			*PropTag
 		);
@@ -1586,7 +1584,7 @@ void FNativeClassHeaderGenerator::PropertyNew(FOutputDevice& DeclOut, FOutputDev
 			FPropertyObjectFlags,
 			*ArrayDim,
 			OffsetStr,
-			*GetSingletonNameFuncAddr(TypedProp->SignatureFunction, OutReferenceGatherers.UniqueCrossModuleReferences),
+			*GetSingletonNameFuncAddr(PropertyBase.FunctionDef, OutReferenceGatherers.UniqueCrossModuleReferences),
 			*MetaDataParams,
 			*PropTag
 		);
@@ -1629,7 +1627,7 @@ void FNativeClassHeaderGenerator::PropertyNew(FOutputDevice& DeclOut, FOutputDev
 			FPropertyObjectFlags,
 			*ArrayDim,
 			OffsetStr,
-			*GetSingletonNameFuncAddr(TypedProp->Enum, OutReferenceGatherers.UniqueCrossModuleReferences),
+			*GetSingletonNameFuncAddr(PropertyBase.EnumDef, OutReferenceGatherers.UniqueCrossModuleReferences),
 			*MetaDataParams,
 			*PropTag
 		);
@@ -1678,9 +1676,9 @@ bool IsEditorOnlyDataProperty(FProperty* Prop)
 	return false;
 }
 
-TTuple<FString, FString> FNativeClassHeaderGenerator::OutputProperties(FOutputDevice& DeclOut, FOutputDevice& Out, FReferenceGatherers& OutReferenceGatherers, const TCHAR* Scope, const TArray<FProperty*>& Properties, const TCHAR* DeclSpaces, const TCHAR* Spaces) const
+TTuple<FString, FString> FNativeClassHeaderGenerator::OutputProperties(FOutputDevice& DeclOut, FOutputDevice& Out, FReferenceGatherers& OutReferenceGatherers, const TCHAR* Scope, const TArray<FUnrealPropertyDefinitionInfo*>& PropertyDefs, const TCHAR* DeclSpaces, const TCHAR* Spaces) const
 {
-	if (Properties.Num() == 0)
+	if (PropertyDefs.Num() == 0)
 	{
 		return TTuple<FString, FString>(TEXT("nullptr"), TEXT("0"));
 	}
@@ -1692,9 +1690,9 @@ TTuple<FString, FString> FNativeClassHeaderGenerator::OutputProperties(FOutputDe
 		FMacroBlockEmitter WithEditorOnlyMacroEmitter(Out, TEXT("WITH_EDITORONLY_DATA"));
 		FMacroBlockEmitter WithEditorOnlyMacroEmitterDecl(DeclOut, TEXT("WITH_EDITORONLY_DATA"));
 
-		for (FProperty* Prop : Properties)
+		for (FUnrealPropertyDefinitionInfo* PropDef : PropertyDefs)
 		{
-			bool bRequiresHasEditorOnlyMacro = IsEditorOnlyDataProperty(Prop);
+			bool bRequiresHasEditorOnlyMacro = IsEditorOnlyDataProperty(PropDef->GetProperty());
 			if (!bRequiresHasEditorOnlyMacro)
 			{
 				bHasAllEditorOnlyDataProperties = false;
@@ -1702,7 +1700,7 @@ TTuple<FString, FString> FNativeClassHeaderGenerator::OutputProperties(FOutputDe
 
 			WithEditorOnlyMacroEmitter(bRequiresHasEditorOnlyMacro);
 			WithEditorOnlyMacroEmitterDecl(bRequiresHasEditorOnlyMacro);
-			OutputProperty(DeclOut, Out, OutReferenceGatherers, Scope, PropertyNamesAndPointers, GTypeDefinitionInfoMap.FindChecked<FUnrealPropertyDefinitionInfo>(Prop), DeclSpaces, Spaces);
+			OutputProperty(DeclOut, Out, OutReferenceGatherers, Scope, PropertyNamesAndPointers, *PropDef, DeclSpaces, Spaces);
 		}
 
 		WithEditorOnlyMacroEmitter(bHasAllEditorOnlyDataProperties);
@@ -1912,35 +1910,14 @@ static void FindNoExportStructsRecursive(TArray<FUnrealScriptStructDefinitionInf
 
 		for (FUnrealPropertyDefinitionInfo* PropertyDef : StartDef->GetProperties())
 		{
-			FProperty* Prop = PropertyDef->GetProperty();
-			if (FStructProperty* StructProp = CastField<FStructProperty>(Prop))
-			{
-				FindNoExportStructsRecursive(StructDefs, &GTypeDefinitionInfoMap.FindChecked<FUnrealStructDefinitionInfo>(StructProp->Struct));
-			}
-			else if (FArrayProperty* ArrayProp = CastField<FArrayProperty>(Prop))
+			const FPropertyBase& PropertyBase = PropertyDef->GetPropertyBase();
+			if (UHTCast<FUnrealScriptStructDefinitionInfo>(PropertyBase.TypeDef) != nullptr)
 			{
-				if (FStructProperty* InnerStructProp = CastField<FStructProperty>(ArrayProp->Inner))
-				{
-					FindNoExportStructsRecursive(StructDefs, &GTypeDefinitionInfoMap.FindChecked<FUnrealStructDefinitionInfo>(InnerStructProp->Struct));
-				}
-			}
-			else if (FMapProperty* MapProp = CastField<FMapProperty>(Prop))
-			{
-				if (FStructProperty* KeyStructProp = CastField<FStructProperty>(MapProp->KeyProp))
-				{
-					FindNoExportStructsRecursive(StructDefs, &GTypeDefinitionInfoMap.FindChecked<FUnrealStructDefinitionInfo>(KeyStructProp->Struct));
-				}
-				if (FStructProperty* ValueStructProp = CastField<FStructProperty>(MapProp->ValueProp))
-				{
-					FindNoExportStructsRecursive(StructDefs, &GTypeDefinitionInfoMap.FindChecked<FUnrealStructDefinitionInfo>(ValueStructProp->Struct));
-				}
+				FindNoExportStructsRecursive(StructDefs, PropertyBase.ScriptStructDef);
 			}
-			else if (FSetProperty* SetProp = CastField<FSetProperty>(Prop))
+			if (PropertyBase.MapKeyProp != nullptr && UHTCast<FUnrealScriptStructDefinitionInfo>(PropertyBase.MapKeyProp->TypeDef) != nullptr)
 			{
-				if (FStructProperty* ElementStructProp = CastField<FStructProperty>(SetProp->ElementProp))
-				{
-					FindNoExportStructsRecursive(StructDefs, &GTypeDefinitionInfoMap.FindChecked<FUnrealStructDefinitionInfo>(ElementStructProp->Struct));
-				}
+				FindNoExportStructsRecursive(StructDefs, PropertyBase.MapKeyProp->ScriptStructDef);
 			}
 		}
 	}
@@ -1960,30 +1937,30 @@ static TArray<FUnrealScriptStructDefinitionInfo*> FindNoExportStructs(FUnrealStr
 void FNativeClassHeaderGenerator::ExportGeneratedPackageInitCode(FOutputDevice& Out, const TCHAR* InDeclarations, uint32 Hash)
 {
 	UPackage* Package = PackageDef.GetPackage();
-	const FString& SingletonName = GetPackageSingletonName(Package, nullptr);
+	const FString& SingletonName = GetPackageSingletonName(PackageDef, nullptr);
 
-	TArray<UField*> Singletons;
+	TArray<FUnrealFieldDefinitionInfo*> Singletons;
 	for (TSharedRef<FUnrealSourceFile>& SourceFile : PackageDef.GetAllSourceFiles())
 	{
 		Singletons.Append(SourceFile->GetSingletons());
 	}
 
-	Algo::Sort(Singletons, [](UField* A, UField* B)
+	Algo::Sort(Singletons, [](FUnrealFieldDefinitionInfo* A, FUnrealFieldDefinitionInfo* B)
 	{
-		bool bADel = A->IsA<UDelegateFunction>();
-		bool bBDel = B->IsA<UDelegateFunction>();
+		bool bADel = A->GetField()->IsA<UDelegateFunction>();
+		bool bBDel = B->GetField()->IsA<UDelegateFunction>();
 		if (bADel != bBDel)
 		{
 			return !bADel;
 		}
-		const FString& AName = GTypeDefinitionInfoMap.FindChecked<FUnrealFieldDefinitionInfo>(A).GetSingletonName(true);
-		const FString& BName = GTypeDefinitionInfoMap.FindChecked<FUnrealFieldDefinitionInfo>(B).GetSingletonName(true);
+		const FString& AName = A->GetSingletonName(true);
+		const FString& BName = B->GetSingletonName(true);
 		return AName < BName;
 	});
 
-	for (UField* ScriptType : Singletons)
+	for (FUnrealFieldDefinitionInfo* FieldDef : Singletons)
 	{
-		Out.Log(GTypeDefinitionInfoMap.FindChecked<FUnrealFieldDefinitionInfo>(ScriptType).GetExternDecl(true));
+		Out.Log(FieldDef->GetExternDecl(true));
 	}
 
 	FOutputDeviceNull OutputDeviceNull;
@@ -2000,10 +1977,9 @@ void FNativeClassHeaderGenerator::ExportGeneratedPackageInitCode(FOutputDevice&
 	if (!Singletons.IsEmpty())
 	{
 		Out.Logf(TEXT("\t\t\tstatic UObject* (*const SingletonFuncArray[])() = {\r\n"));
-		for (UField* ScriptType : Singletons)
+		for (FUnrealFieldDefinitionInfo* FieldDef : Singletons)
 		{
-			const FString& Name = GTypeDefinitionInfoMap.FindChecked<FUnrealFieldDefinitionInfo>(ScriptType).GetSingletonNameChopped(true);
-
+			const FString& Name = FieldDef->GetSingletonNameChopped(true);
 			Out.Logf(TEXT("\t\t\t\t(UObject* (*)())%s,\r\n"), *Name);
 		}
 		Out.Logf(TEXT("\t\t\t};\r\n"));
@@ -2097,7 +2073,7 @@ void FNativeClassHeaderGenerator::ExportNativeGeneratedInitCode(FOutputDevice& O
 		FString StaticsStructName = ClassDef.GetSingletonNameChopped(true) + TEXT("_Statics");
 
 		OutFriendText.Logf(TEXT("\tfriend struct %s;\r\n"), *StaticsStructName);
-		OutDeclarations.Log(GTypeDefinitionInfoMap.FindChecked<FUnrealClassDefinitionInfo>(Class).GetExternDecl(true));
+		OutDeclarations.Log(ClassDef.GetExternDecl(true));
 
 		GeneratedClassRegisterFunctionText.Logf(TEXT("\tstruct %s\r\n"), *StaticsStructName);
 		GeneratedClassRegisterFunctionText.Logf(TEXT("\t{\r\n"));
@@ -2105,16 +2081,16 @@ void FNativeClassHeaderGenerator::ExportNativeGeneratedInitCode(FOutputDevice& O
 		FUHTStringBuilder StaticDefinitions;
 
 		FUHTStringBuilder Singletons;
-		UClass* SuperClass = Class->GetSuperClass();
-		if (SuperClass && SuperClass != Class)
+		FUnrealClassDefinitionInfo* SuperClassDef = ClassDef.GetSuperClass();
+		if (SuperClassDef && SuperClassDef != &ClassDef)
 		{
-			FUnrealClassDefinitionInfo& SuperClassDef = GTypeDefinitionInfoMap.FindChecked<FUnrealClassDefinitionInfo>(SuperClass);
-			SuperClassDef.AddCrossModuleReference(OutReferenceGatherers.UniqueCrossModuleReferences, true);
-			OutDeclarations.Log(SuperClassDef.GetExternDecl(true));
-			Singletons.Logf(TEXT("\t\t(UObject* (*)())%s,\r\n"), *SuperClassDef.GetSingletonNameChopped(true));
+			SuperClassDef->AddCrossModuleReference(OutReferenceGatherers.UniqueCrossModuleReferences, true);
+			OutDeclarations.Log(SuperClassDef->GetExternDecl(true));
+			Singletons.Logf(TEXT("\t\t(UObject* (*)())%s,\r\n"), *SuperClassDef->GetSingletonNameChopped(true));
 		}
 		if (!bIsDynamic)
 		{
+			check(ClassDef.HasSource());
 			FUnrealPackageDefinitionInfo& ClassPackageDef = GTypeDefinitionInfoMap.FindChecked<FUnrealPackageDefinitionInfo>(Class->GetOutermost());
 			PackageDef.AddCrossModuleReference(OutReferenceGatherers.UniqueCrossModuleReferences);
 			OutDeclarations.Logf(TEXT("\t%s_API UPackage* %s;\r\n"), *ApiString, *ClassPackageDef.GetSingletonName());
@@ -2167,7 +2143,7 @@ void FNativeClassHeaderGenerator::ExportNativeGeneratedInitCode(FOutputDevice& O
 				StaticDefinitions.Logf(
 					TEXT("%s\t\t{ &%s, %s },%s\r\n%s"),
 					BEGIN_WRAP_EDITOR_ONLY(bIsEditorOnlyFunction),
-					*GetSingletonNameFuncAddr(*FunctionDef, OutReferenceGatherers.UniqueCrossModuleReferences),
+					*GetSingletonNameFuncAddr(FunctionDef, OutReferenceGatherers.UniqueCrossModuleReferences),
 					*FNativeClassHeaderGenerator::GetUTF8OverriddenNameForLiteral(Function),
 					*FuncHashTag,
 					END_WRAP_EDITOR_ONLY(bIsEditorOnlyFunction)
@@ -2205,13 +2181,7 @@ void FNativeClassHeaderGenerator::ExportNativeGeneratedInitCode(FOutputDevice& O
 
 		FString MetaDataParams = OutputMetaDataCodeForObject(GeneratedClassRegisterFunctionText, StaticDefinitions, Class, *FString::Printf(TEXT("%s::Class_MetaDataParams"), *StaticsStructName), TEXT("\t\t"), TEXT("\t"));
 
-		TArray<FProperty*> Props;
-		for (FUnrealPropertyDefinitionInfo* PropertyDef : ClassDef.GetProperties())
-		{
-			Props.Add(PropertyDef->GetProperty());
-		}
-
-		TTuple<FString, FString> PropertyRange = OutputProperties(GeneratedClassRegisterFunctionText, StaticDefinitions, OutReferenceGatherers, *FString::Printf(TEXT("%s::"), *StaticsStructName), Props, TEXT("\t\t"), TEXT("\t"));
+		TTuple<FString, FString> PropertyRange = OutputProperties(GeneratedClassRegisterFunctionText, StaticDefinitions, OutReferenceGatherers, *FString::Printf(TEXT("%s::"), *StaticsStructName), ClassDef.GetProperties(), TEXT("\t\t"), TEXT("\t"));
 
 		const TCHAR* InterfaceArray;
 		const TCHAR* InterfaceCount;
@@ -2353,10 +2323,10 @@ void FNativeClassHeaderGenerator::ExportNativeGeneratedInitCode(FOutputDevice& O
 	// Append base class' hash at the end of the generated code, this will force update derived classes
 	// when base class changes during hot-reload.
 	uint32 BaseClassHash = 0;
-	UClass* SuperClass = Class->GetSuperClass();
-	if (SuperClass && !SuperClass->HasAnyClassFlags(CLASS_Intrinsic))
+	FUnrealClassDefinitionInfo* SuperClassDef = ClassDef.GetSuperClass();
+	if (SuperClassDef && !SuperClassDef->GetClass()->HasAnyClassFlags(CLASS_Intrinsic))
 	{
-		BaseClassHash = GTypeDefinitionInfoMap.FindChecked(SuperClass).GetHash();
+		BaseClassHash = SuperClassDef->GetHash();
 	}
 	GeneratedClassRegisterFunctionText.Logf(TEXT("\r\n// %u\r\n"), BaseClassHash);
 
@@ -2496,23 +2466,18 @@ void FNativeClassHeaderGenerator::ExportFunction(FOutputDevice& Out, FReferenceG
 	const bool bIsDynamic = (FieldOuterDef && FieldOuterDef->IsDynamic());
 
 	FString OuterFunc;
-	if (UObject* Outer = Function->GetOuter())
+	if (FUnrealObjectDefinitionInfo* OuterDef = FunctionDef.GetOuter())
 	{
-		OuterFunc = Outer->IsA<UPackage>() ? GetPackageSingletonNameFuncAddr((UPackage*)Outer, OutReferenceGatherers.UniqueCrossModuleReferences) : GetSingletonNameFuncAddr(Function->GetOwnerClass(), OutReferenceGatherers.UniqueCrossModuleReferences);
+		FUnrealPackageDefinitionInfo* OuterPackageDef = UHTCast<FUnrealPackageDefinitionInfo>(OuterDef);
+		OuterFunc = OuterPackageDef ? GetPackageSingletonNameFuncAddr(*OuterPackageDef, OutReferenceGatherers.UniqueCrossModuleReferences) : GetSingletonNameFuncAddr(Function->GetOwnerClass(), OutReferenceGatherers.UniqueCrossModuleReferences);
 	}
 	else
 	{
 		OuterFunc = TEXT("nullptr");
 	}
 	
-	TArray<FProperty*> Props;
-	for (FUnrealPropertyDefinitionInfo* PropertyDef : FunctionDef.GetProperties())
-	{
-		Props.Add(PropertyDef->GetProperty());
-	}
-
 	FString StructureSize;
-	if (Props.Num())
+	if (FunctionDef.GetProperties().Num())
 	{
 		UFunction* TempFunction = Function;
 		while (TempFunction->GetSuperFunction())
@@ -2542,7 +2507,7 @@ void FNativeClassHeaderGenerator::ExportFunction(FOutputDevice& Out, FReferenceG
 	USparseDelegateFunction* SparseDelegateFunction = Cast<USparseDelegateFunction>(Function);
 	const TCHAR* UFunctionObjectFlags = FunctionDef.IsOwnedByDynamicType() ? TEXT("RF_Public|RF_Transient") : TEXT("RF_Public|RF_Transient|RF_MarkAsNative");
 
-	TTuple<FString, FString> PropertyRange = OutputProperties(CurrentFunctionText, StaticDefinitions, OutReferenceGatherers, *FString::Printf(TEXT("%s::"), *StaticsStructName), Props, TEXT("\t\t"), TEXT("\t"));
+	TTuple<FString, FString> PropertyRange = OutputProperties(CurrentFunctionText, StaticDefinitions, OutReferenceGatherers, *FString::Printf(TEXT("%s::"), *StaticsStructName), FunctionDef.GetProperties(), TEXT("\t\t"), TEXT("\t"));
 
 	const FFuncInfo&     FunctionData = FunctionDef.GetFunctionData();
 	const bool           bIsNet       = !!(FunctionData.FunctionFlags & (FUNC_NetRequest | FUNC_NetResponse));
@@ -3330,7 +3295,7 @@ FString GetBuildPath(FUnrealSourceFile& SourceFile)
 {
 	FString Out = SourceFile.GetFilename();
 
-	ConvertToBuildIncludePath(GTypeDefinitionInfoMap.FindChecked<FUnrealPackageDefinitionInfo>(SourceFile.GetPackage()).GetModule(), Out);
+	ConvertToBuildIncludePath(SourceFile.GetPackageDef().GetModule(), Out);
 
 	return Out;
 }
@@ -3436,8 +3401,10 @@ FString FNativeClassHeaderGenerator::GetClassFlagExportText( UClass* Class )
 *
 * @param	Enums	the enums to export
 */
-void FNativeClassHeaderGenerator::ExportEnum(FOutputDevice& Out, UEnum* Enum) const
+void FNativeClassHeaderGenerator::ExportEnum(FOutputDevice& Out, FUnrealEnumDefinitionInfo& EnumDef) const
 {
+	UEnum* Enum = EnumDef.GetEnum();
+
 	// Export FOREACH macro
 	Out.Logf( TEXT("#define FOREACH_ENUM_%s(op) "), *Enum->GetName().ToUpper() );
 	bool bHasExistingMax = Enum->ContainsExistingMax();
@@ -3454,8 +3421,6 @@ void FNativeClassHeaderGenerator::ExportEnum(FOutputDevice& Out, UEnum* Enum) co
 	}
 	Out.Logf( TEXT("\r\n") );
 
-	FUnrealEnumDefinitionInfo& EnumDef = GTypeDefinitionInfoMap.FindChecked<FUnrealEnumDefinitionInfo>(Enum);
-
 	// Forward declare the StaticEnum<> specialisation for enum classes
 	if (Enum->GetCppForm() == UEnum::ECppForm::EnumClass)
 	{
@@ -3490,15 +3455,16 @@ void FNativeClassHeaderGenerator::ExportEnum(FOutputDevice& Out, UEnum* Enum) co
 }
 
 // Exports the header text for the list of structs specified (GENERATED_BODY impls)
-void FNativeClassHeaderGenerator::ExportGeneratedStructBodyMacros(FOutputDevice& OutGeneratedHeaderText, FOutputDevice& Out, FReferenceGatherers& OutReferenceGatherers, const FUnrealSourceFile& SourceFile, UScriptStruct* Struct) const
+void FNativeClassHeaderGenerator::ExportGeneratedStructBodyMacros(FOutputDevice& OutGeneratedHeaderText, FOutputDevice& Out, FReferenceGatherers& OutReferenceGatherers, const FUnrealSourceFile& SourceFile, FUnrealScriptStructDefinitionInfo& ScriptStructDef) const
 {
-	FUnrealScriptStructDefinitionInfo& ScriptStructDef = GTypeDefinitionInfoMap.FindChecked<FUnrealScriptStructDefinitionInfo>(Struct);
+	UScriptStruct* Struct = ScriptStructDef.GetScriptStruct();
 	ScriptStructDef.AddCrossModuleReference(OutReferenceGatherers.UniqueCrossModuleReferences, true);
 
 	const bool bIsDynamic = ScriptStructDef.IsDynamic();
 	const FString ActualStructName = FNativeClassHeaderGenerator::GetOverriddenName(Struct);
 	const FString& FriendApiString  = GetAPIString();
 
+	FUnrealStructDefinitionInfo* BaseStructDef = ScriptStructDef.GetSuperClassInfo().Struct;
 	UStruct* BaseStruct = Struct->GetSuperStruct();
 
 	const FString StructNameCPP = FNameLookupCPP::GetNameCPP(Struct);
@@ -3734,7 +3700,7 @@ void FNativeClassHeaderGenerator::ExportGeneratedStructBodyMacros(FOutputDevice&
 		Out.Logf(TEXT("{\r\n"));
 
 		// UStructs can have UClass or UPackage outer (if declared in non-UClass headers).
-		const FString& OuterName (bIsDynamic ? STRING_StructPackage : GetPackageSingletonName(CastChecked<UPackage>(Struct->GetOuter()), OutReferenceGatherers.UniqueCrossModuleReferences));
+		const FString& OuterName (bIsDynamic ? STRING_StructPackage : GetPackageSingletonName(UHTCastChecked<FUnrealPackageDefinitionInfo>(ScriptStructDef.GetOuter()), OutReferenceGatherers.UniqueCrossModuleReferences));
 		if (!bIsDynamic)
 		{
 			Out.Logf(TEXT("\tstatic class UScriptStruct*& Singleton = Z_Registration_Info_UScriptStruct_%s().OuterSingleton;\r\n"), *Struct->GetName());
@@ -3833,7 +3799,7 @@ void FNativeClassHeaderGenerator::ExportGeneratedStructBodyMacros(FOutputDevice&
 	if (BaseStruct)
 	{
 		CastChecked<UScriptStruct>(BaseStruct); // this better actually be a script struct
-		GTypeDefinitionInfoMap.FindChecked<FUnrealFieldDefinitionInfo>(BaseStruct).AddCrossModuleReference(OutReferenceGatherers.UniqueCrossModuleReferences, true);
+		BaseStructDef->AddCrossModuleReference(OutReferenceGatherers.UniqueCrossModuleReferences, true);
 	}
 
 	EStructFlags UncomputedFlags = (EStructFlags)(Struct->StructFlags & ~STRUCT_ComputedFlags);
@@ -3841,7 +3807,7 @@ void FNativeClassHeaderGenerator::ExportGeneratedStructBodyMacros(FOutputDevice&
 	FString OuterFunc;
 	if (!bIsDynamic)
 	{
-		OuterFunc = GetPackageSingletonNameFuncAddr(CastChecked<UPackage>(Struct->GetOuter()), OutReferenceGatherers.UniqueCrossModuleReferences);
+		OuterFunc = GetPackageSingletonNameFuncAddr(UHTCastChecked<FUnrealPackageDefinitionInfo>(ScriptStructDef.GetOuter()), OutReferenceGatherers.UniqueCrossModuleReferences);
 	}
 	else
 	{
@@ -3857,12 +3823,6 @@ void FNativeClassHeaderGenerator::ExportGeneratedStructBodyMacros(FOutputDevice&
 
 	FString MetaDataParams = OutputMetaDataCodeForObject(GeneratedStructRegisterFunctionText, StaticDefinitions, Struct, *FString::Printf(TEXT("%s::Struct_MetaDataParams"), *StaticsStructName), TEXT("\t\t"), TEXT("\t"));
 
-	TArray<FProperty*> Props;
-	for (FUnrealPropertyDefinitionInfo* PropertyDef : ScriptStructDef.GetProperties())
-	{
-		Props.Add(PropertyDef->GetProperty());
-	}
-
 	FString NewStructOps;
 	if (Struct->StructFlags & STRUCT_Native)
 	{
@@ -3880,13 +3840,13 @@ void FNativeClassHeaderGenerator::ExportGeneratedStructBodyMacros(FOutputDevice&
 		NewStructOps = TEXT("nullptr");
 	}
 
-	TTuple<FString, FString> PropertyRange = OutputProperties(GeneratedStructRegisterFunctionText, StaticDefinitions, OutReferenceGatherers, *FString::Printf(TEXT("%s::"), *StaticsStructName), Props, TEXT("\t\t"), TEXT("\t"));
+	TTuple<FString, FString> PropertyRange = OutputProperties(GeneratedStructRegisterFunctionText, StaticDefinitions, OutReferenceGatherers, *FString::Printf(TEXT("%s::"), *StaticsStructName), ScriptStructDef.GetProperties(), TEXT("\t\t"), TEXT("\t"));
 
 	GeneratedStructRegisterFunctionText.Log (TEXT("\t\tstatic const UECodeGen_Private::FStructParams ReturnStructParams;\r\n"));
 
 	StaticDefinitions.Logf(TEXT("\tconst UECodeGen_Private::FStructParams %s::ReturnStructParams = {\r\n"), *StaticsStructName);
 	StaticDefinitions.Logf(TEXT("\t\t(UObject* (*)())%s,\r\n"), *OuterFunc);
-	StaticDefinitions.Logf(TEXT("\t\t%s,\r\n"), *GetSingletonNameFuncAddr(BaseStruct, OutReferenceGatherers.UniqueCrossModuleReferences));
+	StaticDefinitions.Logf(TEXT("\t\t%s,\r\n"), *GetSingletonNameFuncAddr(BaseStructDef, OutReferenceGatherers.UniqueCrossModuleReferences));
 	StaticDefinitions.Logf(TEXT("\t\t%s,\r\n"), *NewStructOps);
 	StaticDefinitions.Logf(TEXT("\t\t%s,\r\n"), *CreateUTF8LiteralString(ActualStructName));
 	StaticDefinitions.Logf(TEXT("\t\tsizeof(%s),\r\n"), *StructNameCPP);
@@ -4023,9 +3983,9 @@ void FNativeClassHeaderGenerator::ExportGeneratedStructBodyMacros(FOutputDevice&
 	}
 }
 
-void FNativeClassHeaderGenerator::ExportGeneratedEnumInitCode(FOutputDevice& Out, FReferenceGatherers& OutReferenceGatherers, const FUnrealSourceFile& SourceFile, UEnum* Enum) const
+void FNativeClassHeaderGenerator::ExportGeneratedEnumInitCode(FOutputDevice& Out, FReferenceGatherers& OutReferenceGatherers, const FUnrealSourceFile& SourceFile, FUnrealEnumDefinitionInfo& EnumDef) const
 {
-	FUnrealEnumDefinitionInfo& EnumDef = GTypeDefinitionInfoMap.FindChecked<FUnrealEnumDefinitionInfo>(Enum);
+	UEnum* Enum                         = EnumDef.GetEnum();
 	const bool bIsDynamic               = EnumDef.IsDynamic();
 	const FString& SingletonName        = EnumDef.GetSingletonNameChopped(true);
 	const FString EnumNameCpp           = Enum->GetName(); //UserDefinedEnum should already have a valid cpp name.
@@ -4038,7 +3998,7 @@ void FNativeClassHeaderGenerator::ExportGeneratedEnumInitCode(FOutputDevice& Out
 	FMacroBlockEmitter EditorOnlyData(Out, TEXT("WITH_EDITORONLY_DATA"));
 	EditorOnlyData(bIsEditorOnlyDataType);
 
-	const FString& PackageSingletonName = (bIsDynamic ? EnumDef.GetTypePackageName() : GetPackageSingletonName(CastChecked<UPackage>(Enum->GetOuter()), OutReferenceGatherers.UniqueCrossModuleReferences));
+	const FString& PackageSingletonName = (bIsDynamic ? EnumDef.GetTypePackageName() : GetPackageSingletonName(UHTCastChecked<FUnrealPackageDefinitionInfo>(EnumDef.GetOuter()), OutReferenceGatherers.UniqueCrossModuleReferences));
 
 	if (!bIsDynamic)
 	{
@@ -5279,8 +5239,6 @@ void FNativeClassHeaderGenerator::ExportNativeFunctions(FOutputDevice& OutGenera
 		const bool bEditorOnlyFunc = Function->HasAnyFunctionFlags(FUNC_EditorOnly);
 		FNativeFunctionStringBuilder& FuncStringBuilders = bEditorOnlyFunc ? EditorStringBuilders : RuntimeStringBuilders;
 
-		FUnrealFunctionDefinitionInfo& FuncDef = GTypeDefinitionInfoMap.FindChecked<FUnrealFunctionDefinitionInfo>(Function);
-
 		// Custom thunks don't get any C++ stub function generated
 		if (FunctionData.FunctionExportFlags & FUNCEXPORT_CustomThunk)
 		{
@@ -5315,13 +5273,13 @@ void FNativeClassHeaderGenerator::ExportNativeFunctions(FOutputDevice& OutGenera
 
 			//Emit warning here if necessary
 			FUHTStringBuilder FunctionDeclaration;
-			ExportNativeFunctionHeader(FunctionDeclaration, OutReferenceGatherers.ForwardDeclarations, FuncDef, FunctionData, EExportFunctionType::Function, EExportFunctionHeaderStyle::Declaration, nullptr, *GetAPIString());
+			ExportNativeFunctionHeader(FunctionDeclaration, OutReferenceGatherers.ForwardDeclarations, *FunctionDef, FunctionData, EExportFunctionType::Function, EExportFunctionHeaderStyle::Declaration, nullptr, *GetAPIString());
 			FunctionDeclaration.Log(TEXT(";\r\n"));
 
 			// Declare validation function if needed
 			if (FunctionData.FunctionFlags & FUNC_NetValidate)
 			{
-				FString ParameterList = GetFunctionParameterString(FuncDef, OutReferenceGatherers);
+				FString ParameterList = GetFunctionParameterString(*FunctionDef, OutReferenceGatherers);
 
 				const TCHAR* Virtual = (!FunctionData.FunctionReference->HasAnyFunctionFlags(FUNC_Static) && !(FunctionData.FunctionExportFlags & FUNCEXPORT_Final)) ? TEXT("virtual") : TEXT("");
 				FStringOutputDevice ValidDecl;
@@ -5342,7 +5300,7 @@ void FNativeClassHeaderGenerator::ExportNativeFunctions(FOutputDevice& OutGenera
 			// Versions that skip function autodeclaration throw an error when a function is missing.
 			if (ClassRange.bHasGeneratedBody && (SourceFile.GetGeneratedCodeVersionForStruct(Class) > EGeneratedCodeVersion::V1))
 			{
-				CheckRPCFunctions(OutReferenceGatherers, FuncDef, ClassCPPName, ImplementationPosition, ValidatePosition, SourceFile);
+				CheckRPCFunctions(OutReferenceGatherers, *FunctionDef, ClassCPPName, ImplementationPosition, ValidatePosition, SourceFile);
 			}
 		}
 
@@ -5361,7 +5319,7 @@ void FNativeClassHeaderGenerator::ExportNativeFunctions(FOutputDevice& OutGenera
 		FuncStringBuilders.RPCImplementations += LINE_TERMINATOR TEXT("\t{") LINE_TERMINATOR;
 
 		FParmsAndReturnProperties Parameters = GetFunctionParmsAndReturn(*FunctionDef);
-		ExportFunctionThunk(FuncStringBuilders.RPCImplementations, OutReferenceGatherers, FuncDef, Parameters.Parms, Parameters.Return);
+		ExportFunctionThunk(FuncStringBuilders.RPCImplementations, OutReferenceGatherers, *FunctionDef, Parameters.Parms, Parameters.Return);
 
 		FuncStringBuilders.RPCImplementations += TEXT("\t}") LINE_TERMINATOR;
 	}
@@ -5608,24 +5566,26 @@ static bool HasDynamicOuter(UField* Field)
 
 static void RecordPackageSingletons(
 	FUnrealSourceFile& SourceFile,
-	const TArray<UScriptStruct*>& Structs,
-	const TArray<UDelegateFunction*>& Delegates)
+	const TArray<FUnrealScriptStructDefinitionInfo*>& StructDefs,
+	const TArray<FUnrealFunctionDefinitionInfo*>& DelegateDefs)
 {
-	TArray<UField*> Singletons;
-	Singletons.Reserve(Structs.Num() + Delegates.Num());
-	for (UScriptStruct* Struct : Structs)
+	TArray<FUnrealFieldDefinitionInfo*> Singletons;
+	Singletons.Reserve(StructDefs.Num() + DelegateDefs.Num());
+	for (FUnrealScriptStructDefinitionInfo* StructDef : StructDefs)
 	{
+		UScriptStruct* Struct = StructDef->GetScriptStruct();
 		if (Struct->StructFlags & STRUCT_NoExport && !HasDynamicOuter(Struct))
 		{
-			Singletons.Add(Struct);
+			Singletons.Add(StructDef);
 		}
 	}
 
-	for (UDelegateFunction* Delegate : Delegates)
+	for (FUnrealFunctionDefinitionInfo* DelegateDef : DelegateDefs)
 	{
+		UFunction* Delegate = DelegateDef->GetFunction();
 		if (!HasDynamicOuter(Delegate))
 		{
-			Singletons.Add(Delegate);
+			Singletons.Add(DelegateDef);
 		}
 	}
 
@@ -5769,23 +5729,23 @@ void FNativeClassHeaderGenerator::GenerateSourceFiles(
 
 					FReferenceGatherers ReferenceGatherers(&GeneratedCPP.CrossModuleReferences, GeneratedCPP.ForwardDeclarations);
 
-					TArray<UEnum*> Enums;
-					TArray<UScriptStruct*> Structs;
-					TArray<UDelegateFunction*> DelegateFunctions;
-					SourceFile.GetScope()->SplitTypesIntoArrays(Enums, Structs, DelegateFunctions);
+					TArray<FUnrealEnumDefinitionInfo*> EnumDefs;
+					TArray<FUnrealScriptStructDefinitionInfo*> StructDefs;
+					TArray<FUnrealFunctionDefinitionInfo*> DelegateFunctionDefs;
+					SourceFile.GetScope()->SplitTypesIntoArrays(EnumDefs, StructDefs, DelegateFunctionDefs);
 
-					RecordPackageSingletons(SourceFile, Structs, DelegateFunctions);
+					RecordPackageSingletons(SourceFile, StructDefs, DelegateFunctionDefs);
 
 					// Sort by line number to help make sure things are in order
-					auto Predicate = [](const UField& Lhs, const UField& Rhs)
+					auto Predicate = [](const FUnrealFieldDefinitionInfo& Lhs, const FUnrealFieldDefinitionInfo& Rhs)
 					{
-						return GTypeDefinitionInfoMap.FindChecked(&Lhs).GetLineNumber() < GTypeDefinitionInfoMap.FindChecked(&Rhs).GetLineNumber();
+						return Lhs.GetLineNumber() < Rhs.GetLineNumber();
 					};
-					Enums.StableSort(Predicate);
-					Structs.StableSort(Predicate);
+					EnumDefs.StableSort(Predicate);
+					StructDefs.StableSort(Predicate);
 
 					// Reverse the containers as they come out in the reverse order of declaration
-					Algo::Reverse(DelegateFunctions);
+					Algo::Reverse(DelegateFunctionDefs);
 
 					const FString FileDefineName = SourceFile.GetFileDefineName();
 					const FString& StrippedFilename = SourceFile.GetStrippedFilename();
@@ -5799,37 +5759,35 @@ void FNativeClassHeaderGenerator::GenerateSourceFiles(
 						*FileDefineName, *StrippedFilename, *StrippedFilename, *FileDefineName);
 
 					// Export enums declared in non-UClass headers.
-					for (UEnum* Enum : Enums)
+					for (FUnrealEnumDefinitionInfo* EnumDef : EnumDefs)
 					{
 						// Is this ever not the case?
-						if (Enum->GetOuter()->IsA(UPackage::StaticClass()))
+						if (EnumDef->GetEnum()->GetOuter()->IsA(UPackage::StaticClass()))
 						{
-							GeneratedFunctionDeclarations.Log(GTypeDefinitionInfoMap.FindChecked<FUnrealFieldDefinitionInfo>(Enum).GetExternDecl(true));
-							Generator.ExportGeneratedEnumInitCode(GeneratedCPPText, ReferenceGatherers, SourceFile, Enum);
+							GeneratedFunctionDeclarations.Log(EnumDef->GetExternDecl(true));
+							Generator.ExportGeneratedEnumInitCode(GeneratedCPPText, ReferenceGatherers, SourceFile, *EnumDef);
 						}
 					}
 
 					// export boilerplate macros for structs
 					// reverse the order.
-					for (UScriptStruct* Struct : Structs)
+					for (FUnrealScriptStructDefinitionInfo* StructDef : StructDefs)
 					{
-						GeneratedFunctionDeclarations.Log(GTypeDefinitionInfoMap.FindChecked<FUnrealFieldDefinitionInfo>(Struct).GetExternDecl(true));
-						Generator.ExportGeneratedStructBodyMacros(GeneratedHeaderText, GeneratedCPPText, ReferenceGatherers, SourceFile, Struct);
+						GeneratedFunctionDeclarations.Log(StructDef->GetExternDecl(true));
+						Generator.ExportGeneratedStructBodyMacros(GeneratedHeaderText, GeneratedCPPText, ReferenceGatherers, SourceFile, *StructDef);
 					}
 
 					// export delegate definitions
-					for (UDelegateFunction* Func : DelegateFunctions)
+					for (FUnrealFunctionDefinitionInfo* FunctionDef : DelegateFunctionDefs)
 					{
-						FUnrealFunctionDefinitionInfo& FunctionDef = GTypeDefinitionInfoMap.FindChecked<FUnrealFunctionDefinitionInfo>(Func);
-						GeneratedFunctionDeclarations.Log(FunctionDef.GetExternDecl(true));
-						Generator.ExportDelegateDeclaration(GeneratedCPPText, ReferenceGatherers, SourceFile, FunctionDef);
+						GeneratedFunctionDeclarations.Log(FunctionDef->GetExternDecl(true));
+						Generator.ExportDelegateDeclaration(GeneratedCPPText, ReferenceGatherers, SourceFile, *FunctionDef);
 					}
 
 					// export delegate wrapper function implementations
-					for (UDelegateFunction* Func : DelegateFunctions)
+					for (FUnrealFunctionDefinitionInfo* FunctionDef : DelegateFunctionDefs)
 					{
-						FUnrealFunctionDefinitionInfo& FunctionDef = GTypeDefinitionInfoMap.FindChecked<FUnrealFunctionDefinitionInfo>(Func);
-						Generator.ExportDelegateDefinition(GeneratedHeaderText, ReferenceGatherers, SourceFile, FunctionDef);
+						Generator.ExportDelegateDefinition(GeneratedHeaderText, ReferenceGatherers, SourceFile, *FunctionDef);
 					}
 
 					for (const TSharedRef<FUnrealTypeDefinitionInfo>& TypeDef : SourceFile.GetDefinedClasses())
@@ -5845,9 +5803,9 @@ void FNativeClassHeaderGenerator::GenerateSourceFiles(
 					GeneratedHeaderText.Log(TEXT("#undef CURRENT_FILE_ID\r\n"));
 					GeneratedHeaderText.Logf(TEXT("#define CURRENT_FILE_ID %s\r\n\r\n\r\n"), *SourceFile.GetFileId());
 
-					for (UEnum* Enum : Enums)
+					for (FUnrealEnumDefinitionInfo* EnumDef : EnumDefs)
 					{
-						Generator.ExportEnum(GeneratedHeaderText, Enum);
+						Generator.ExportEnum(GeneratedHeaderText, *EnumDef);
 					}
 				}
 			);
@@ -6776,23 +6734,11 @@ void PrepareTypesForParsing(TArray<FUnrealPackageDefinitionInfo*>& PackageDefs)
 	for (UClass* ScriptClass : TObjectRange<UClass>())
 	{
 		const UClass* ScriptSuperClass = ScriptClass->GetSuperClass();
-		if (ScriptSuperClass)
-		{
-			TSharedRef<FUnrealTypeDefinitionInfo>* ScriptClassTypeDef = GTypeDefinitionInfoMap.Find(ScriptClass);
-			TSharedRef<FUnrealTypeDefinitionInfo>* SuperTypeDef = GTypeDefinitionInfoMap.Find(ScriptSuperClass);
-			if (ScriptClassTypeDef && (*ScriptClassTypeDef)->HasSource() && SuperTypeDef && !(*SuperTypeDef)->HasSource())
-			{
-				UE_LOG(LogCompile, Log, TEXT("%s has been used"), *(*SuperTypeDef)->GetNameCPP());
-				Uniques.FindOrAdd(ScriptSuperClass).Add(ScriptClass);
-			}
-		}
-
 		if (ScriptSuperClass && !ScriptSuperClass->HasAnyClassFlags(CLASS_Intrinsic) && GTypeDefinitionInfoMap.Contains(ScriptClass) && !GTypeDefinitionInfoMap.Contains(ScriptSuperClass))
 		{
 			FResults::LogError(*ScriptClass, *FString::Printf(TEXT("Superclass %s of class %s not found"), *ScriptSuperClass->GetName(), *ScriptClass->GetName()));
 		}
 	}
-	UE_LOG(LogCompile, Log, TEXT("BOB"));
 }
 
 void TopologicalRecursion(FUnrealSourceFile& First, FUnrealSourceFile& Visit)
@@ -6878,8 +6824,7 @@ void ParseSourceFiles(TArray<FUnrealSourceFile*>& OrderedSourceFiles)
 {
 	for (FUnrealSourceFile* SourceFile : OrderedSourceFiles)
 	{
-		UPackage* Package = SourceFile->GetPackage();
-		FUnrealPackageDefinitionInfo& PackageDef = GTypeDefinitionInfoMap.FindChecked<FUnrealPackageDefinitionInfo>(Package);
+		FUnrealPackageDefinitionInfo& PackageDef = SourceFile->GetPackageDef();
 		FScopedDurationTimer SourceTimer(SourceFile->GetTime(ESourceFileTime::Parse));
 
 		ECompilationResult::Type Result = FHeaderParser::Parse(PackageDef, GWarn, *SourceFile);
@@ -6962,8 +6907,7 @@ void Export(TArray<FUnrealPackageDefinitionInfo*>& PackageDefs, TArray<FUnrealSo
 	GeneratedCPPs.Reserve(OrderedSourceFiles.Num());
 	for (FUnrealSourceFile* SourceFile : OrderedSourceFiles)
 	{
-		FUnrealPackageDefinitionInfo& PackageDef = GTypeDefinitionInfoMap.FindChecked<FUnrealPackageDefinitionInfo>(SourceFile->GetPackage());
-		GeneratedCPPs.Emplace(PackageDef, *SourceFile);
+		GeneratedCPPs.Emplace(SourceFile->GetPackageDef(), *SourceFile);
 	}
 
 	FResults::Try([&GeneratedCPPs]() { FNativeClassHeaderGenerator::GenerateSourceFiles(GeneratedCPPs); });
diff --git a/Engine/Source/Programs/UnrealHeaderTool/Private/HeaderParser.cpp b/Engine/Source/Programs/UnrealHeaderTool/Private/HeaderParser.cpp
--- a/Engine/Source/Programs/UnrealHeaderTool/Private/HeaderParser.cpp
+++ b/Engine/Source/Programs/UnrealHeaderTool/Private/HeaderParser.cpp
@@ -833,15 +833,6 @@ namespace
 		}
 	}
 
-	bool IsPropertySupportedByBlueprint(const FProperty* Property, bool bMemberVariable)
-	{
-		if (FUnrealPropertyDefinitionInfo* PropDef = GTypeDefinitionInfoMap.Find<FUnrealPropertyDefinitionInfo>(Property))
-		{
-			return FPropertyTraits::IsSupportedByBlueprint(*PropDef, bMemberVariable);
-		}
-		return false;
-	}
-
 	void SkipAlignasIfNecessary(FBaseParser& Parser)
 	{
 		if (Parser.MatchIdentifier(TEXT("alignas"), ESearchCase::CaseSensitive))
@@ -1179,14 +1170,10 @@ FField* FHeaderParser::FindProperty(FUnrealStructDefinitionInfo& InScope, const
  * @param Type Type for which to add include path.
  * @param MetaData Meta data to fill the information.
  */
-void AddIncludePathToMetadata(UField* Type, TMap<FName, FString> &MetaData)
+void AddIncludePathToMetadata(FUnrealFieldDefinitionInfo& FieldDef, TMap<FName, FString> &MetaData)
 {
 	// Add metadata for the include path.
-	TSharedRef<FUnrealTypeDefinitionInfo>* TypeDefinitionPtr = GTypeDefinitionInfoMap.Find(Type);
-	if (TypeDefinitionPtr != nullptr)
-	{
-		MetaData.Add(NAME_IncludePath, (*TypeDefinitionPtr)->GetUnrealSourceFile().GetIncludePath());
-	}
+	MetaData.Add(NAME_IncludePath, FieldDef.GetUnrealSourceFile().GetIncludePath());
 }
 
 /**
@@ -2602,8 +2589,7 @@ void FHeaderParser::FixupDelegateProperties(FUnrealStructDefinitionInfo& StructD
 		{
 			// this FDelegateProperty corresponds to an actual delegate variable (i.e. delegate<SomeDelegate> Foo); we need to lookup the token data for
 			// this property and verify that the delegate property's "type" is an actual delegate function
-			FUnrealPropertyDefinitionInfo& PropDef = GTypeDefinitionInfoMap.FindChecked<FUnrealPropertyDefinitionInfo>(Property);
-			FPropertyBase& DelegatePropertyToken = PropDef.GetPropertyBase();
+			FPropertyBase& DelegatePropertyToken = PropertyDef->GetPropertyBase();
 
 			// attempt to find the delegate function in the map of functions we've already found
 			UFunction* SourceDelegateFunction = DelegateCache.FindRef(DelegatePropertyToken.DelegateName);
@@ -2643,14 +2629,14 @@ void FHeaderParser::FixupDelegateProperties(FUnrealStructDefinitionInfo& StructD
 					// verify that we got a valid string for the class name
 					if ( DelegateClassName.Len() == 0 )
 					{
-						UngetToken(PropDef.GetLineNumber(), PropDef.GetParsePosition());
+						UngetToken(PropertyDef->GetLineNumber(), PropertyDef->GetParsePosition());
 						FError::Throwf(TEXT("Invalid scope specified in delegate property function reference: '%s'"), *NameOfDelegateFunction);
 					}
 
 					// verify that we got a valid string for the name of the function
 					if ( DelegateName.Len() == 0 )
 					{
-						UngetToken(PropDef.GetLineNumber(), PropDef.GetParsePosition());
+						UngetToken(PropertyDef->GetLineNumber(), PropertyDef->GetParsePosition());
 						FError::Throwf(TEXT("Invalid delegate name specified in delegate property function reference '%s'"), *NameOfDelegateFunction);
 					}
 
@@ -2665,12 +2651,12 @@ void FHeaderParser::FixupDelegateProperties(FUnrealStructDefinitionInfo& StructD
 
 				if ( SourceDelegateFunction == NULL )
 				{
-					UngetToken(PropDef.GetLineNumber(), PropDef.GetParsePosition());
+					UngetToken(PropertyDef->GetLineNumber(), PropertyDef->GetParsePosition());
 					FError::Throwf(TEXT("Failed to find delegate function '%s'"), *NameOfDelegateFunction);
 				}
 				else if ( (SourceDelegateFunction->FunctionFlags&FUNC_Delegate) == 0 )
 				{
-					UngetToken(PropDef.GetLineNumber(), PropDef.GetParsePosition());
+					UngetToken(PropertyDef->GetLineNumber(), PropertyDef->GetParsePosition());
 					FError::Throwf(TEXT("Only delegate functions can be used as the type for a delegate property; '%s' is not a delegate."), *NameOfDelegateFunction);
 				}
 			}
@@ -2680,12 +2666,15 @@ void FHeaderParser::FixupDelegateProperties(FUnrealStructDefinitionInfo& StructD
 			// save this into the delegate cache for faster lookup later
 			DelegateCache.Add(DelegatePropertyToken.DelegateName, SourceDelegateFunction);
 
+			FUnrealFunctionDefinitionInfo& SourceDelegateFunctionDef = GTypeDefinitionInfoMap.FindChecked<FUnrealFunctionDefinitionInfo>(SourceDelegateFunction);
+
 			// bind it to the delegate property
 			if( DelegateProperty != NULL )
 			{
 				if( !SourceDelegateFunction->HasAnyFunctionFlags( FUNC_MulticastDelegate ) )
 				{
-					DelegateProperty->SignatureFunction = DelegatePropertyToken.Function = SourceDelegateFunction;
+					DelegateProperty->SignatureFunction = SourceDelegateFunction;
+					DelegatePropertyToken.FunctionDef = &SourceDelegateFunctionDef;
 				}
 				else
 				{
@@ -2696,16 +2685,16 @@ void FHeaderParser::FixupDelegateProperties(FUnrealStructDefinitionInfo& StructD
 			{
 				if( SourceDelegateFunction->HasAnyFunctionFlags( FUNC_MulticastDelegate ) )
 				{
-					MulticastDelegateProperty->SignatureFunction = DelegatePropertyToken.Function = SourceDelegateFunction;
+					MulticastDelegateProperty->SignatureFunction = SourceDelegateFunction;
+					DelegatePropertyToken.FunctionDef = &SourceDelegateFunctionDef;
 
 					if(MulticastDelegateProperty->HasAnyPropertyFlags(CPF_BlueprintAssignable | CPF_BlueprintCallable))
 					{
-						FUnrealFunctionDefinitionInfo& SourceDelegateFunctionDef = GTypeDefinitionInfoMap.FindChecked<FUnrealFunctionDefinitionInfo>(SourceDelegateFunction);
 						for (FUnrealPropertyDefinitionInfo* FuncParamDef : SourceDelegateFunctionDef.GetProperties())
 						{
 							FProperty* FuncParam = FuncParamDef->GetProperty();
 
-							if (!IsPropertySupportedByBlueprint(FuncParam, false))
+							if (!FPropertyTraits::IsSupportedByBlueprint(*FuncParamDef, false))
 							{
 								FString ExtendedCPPType;
 								FString CPPType = FuncParam->GetCPPType(&ExtendedCPPType);
@@ -4175,7 +4164,7 @@ void FHeaderParser::GetVarType(
 			if (UEnum* Enum = FClasses::FindObject<UEnum>(ANY_PACKAGE, InnerEnumType.Identifier))
 			{
 				// In-scope enumeration.
-				VarProperty = FPropertyBase(Enum, CPT_Byte);
+				VarProperty = FPropertyBase(GTypeDefinitionInfoMap.FindChecked<FUnrealEnumDefinitionInfo>(Enum), CPT_Byte);
 				bFoundEnum  = true;
 			}
 		}
@@ -4249,7 +4238,7 @@ void FHeaderParser::GetVarType(
 		}
 
 		// In-scope enumeration.
-		VarProperty            = FPropertyBase(Enum, UnderlyingType);
+		VarProperty            = FPropertyBase(GTypeDefinitionInfoMap.FindChecked<FUnrealEnumDefinitionInfo>(Enum), UnderlyingType);
 		bUnconsumedEnumKeyword = false;
 	}
 	else
@@ -4265,13 +4254,13 @@ void FHeaderParser::GetVarType(
 			bStripped = true;
 		}
 
-		auto SetDelegateType = [&](UFunction* InFunction, const FString& InIdentifierStripped)
+		auto SetDelegateType = [&](FUnrealFunctionDefinitionInfo& InFunction, const FString& InIdentifierStripped)
 		{
 			bHandledType = true;
 
-			VarProperty = FPropertyBase(InFunction->HasAnyFunctionFlags(FUNC_MulticastDelegate) ? CPT_MulticastDelegate : CPT_Delegate);
+			VarProperty = FPropertyBase(InFunction.GetFunction()->HasAnyFunctionFlags(FUNC_MulticastDelegate) ? CPT_MulticastDelegate : CPT_Delegate);
 			VarProperty.DelegateName = *InIdentifierStripped;
-			VarProperty.Function = InFunction;
+			VarProperty.FunctionDef = &InFunction;
 
 			if (!(Disallow & CPF_InstancedReference))
 			{
@@ -4293,7 +4282,7 @@ void FHeaderParser::GetVarType(
 					{
 						if (FUnrealFunctionDefinitionInfo* DelegateFuncDef = FoundDef->AsFunction())
 						{
-							SetDelegateType(DelegateFuncDef->GetFunction(), DelegateIdentifierStripped);
+							SetDelegateType(*DelegateFuncDef, DelegateIdentifierStripped);
 							VarProperty.DelegateSignatureOwnerClass = LocalOwnerClass;
 						}
 					}
@@ -4327,7 +4316,7 @@ void FHeaderParser::GetVarType(
 
 			bHandledType = true;
 
-			VarProperty = FPropertyBase(Struct);
+			VarProperty = FPropertyBase(GTypeDefinitionInfoMap.FindChecked<FUnrealScriptStructDefinitionInfo>(Struct));
 			if ((Struct->StructFlags & STRUCT_HasInstancedReference) && !(Disallow & CPF_ContainsInstancedReference))
 			{
 				Flags |= CPF_ContainsInstancedReference;
@@ -4344,18 +4333,11 @@ void FHeaderParser::GetVarType(
 		}
 		else
 		{
-			UFunction* DelegateFunc = nullptr;
-			if (FUnrealFieldDefinitionInfo* FieldDef = Scope->FindTypeByName(*(IdentifierStripped + HEADER_GENERATED_DELEGATE_SIGNATURE_SUFFIX)))
-			{
-				if (FUnrealFunctionDefinitionInfo* FunctionDef = FieldDef->AsFunction())
-				{
-					DelegateFunc = FunctionDef->GetFunction();
-				}
-			}
+			FUnrealFunctionDefinitionInfo* DelegateFuncDef = UHTCast<FUnrealFunctionDefinitionInfo>(Scope->FindTypeByName(*(IdentifierStripped + HEADER_GENERATED_DELEGATE_SIGNATURE_SUFFIX)));
 
-			if (DelegateFunc)
+			if (DelegateFuncDef)
 			{
-				SetDelegateType(DelegateFunc, IdentifierStripped);
+				SetDelegateType(*DelegateFuncDef, IdentifierStripped);
 			}
 			else
 			{
@@ -4466,7 +4448,8 @@ void FHeaderParser::GetVarType(
 						PropertyType = CPT_Interface;
 					}
 
-					VarProperty = FPropertyBase(TempClass, PropertyType, bWeakIsAuto);
+					FUnrealClassDefinitionInfo& TempClassDef = GTypeDefinitionInfoMap.FindChecked<FUnrealClassDefinitionInfo>(TempClass);
+					VarProperty = FPropertyBase(TempClassDef, PropertyType, bWeakIsAuto);
 					if (TempClass->IsChildOf(UClass::StaticClass()))
 					{
 						if (MatchSymbol(TEXT('<')))
@@ -4485,13 +4468,13 @@ void FHeaderParser::GetVarType(
 
 							RedirectTypeIdentifier(Limitor);
 
-							VarProperty.MetaClass = FClasses::FindScriptClassOrThrow(Limitor.Identifier);
+							VarProperty.MetaClassDef = &GTypeDefinitionInfoMap.FindChecked<FUnrealClassDefinitionInfo>(FClasses::FindScriptClassOrThrow(Limitor.Identifier));
 
 							RequireSymbol(TEXT('>'), TEXT("'class limitor'"), ESymbolParseOption::CloseTemplateBracket);
 						}
 						else
 						{
-							VarProperty.MetaClass = UObject::StaticClass();
+							VarProperty.MetaClassDef = &GTypeDefinitionInfoMap.FindChecked<FUnrealClassDefinitionInfo>(UObject::StaticClass());
 						}
 
 						if (PropertyType == CPT_WeakObjectReference)
@@ -4510,7 +4493,7 @@ void FHeaderParser::GetVarType(
 					}
 
 					// Inherit instancing flags
-					if (FUnrealClassDefinitionInfo::HierarchyHasAnyClassFlags(TempClass, CLASS_DefaultToInstanced))
+					if (TempClassDef.HierarchyHasAnyClassFlags(CLASS_DefaultToInstanced))
 					{
 						Flags |= ((CPF_InstancedReference | CPF_ExportObject) & (~Disallow));
 					}
@@ -4555,7 +4538,7 @@ void FHeaderParser::GetVarType(
 			{
 				if (UFunction* ExternDelegateFunc = FClasses::FindObject<UFunction>(ANY_PACKAGE, *(IdentifierStripped + HEADER_GENERATED_DELEGATE_SIGNATURE_SUFFIX)))
 				{
-					SetDelegateType(ExternDelegateFunc, IdentifierStripped);
+					SetDelegateType(GTypeDefinitionInfoMap.FindChecked<FUnrealFunctionDefinitionInfo>(ExternDelegateFunc), IdentifierStripped);
 				}
 
 				if (!bHandledType)
@@ -4684,7 +4667,7 @@ void FHeaderParser::GetVarType(
 	{
 		if ((VarProperty.Type == CPT_ObjectReference) || (VarProperty.Type == CPT_ObjectPtrReference))
 		{
-			if (VarProperty.PropertyClass->IsChildOf<UClass>())
+			if (VarProperty.ClassDef->GetClass()->IsChildOf<UClass>())
 			{
 				FError::Throwf(TEXT("'Instanced' cannot be applied to class properties (UClass* or TSubclassOf<>)"));
 			}
@@ -4695,7 +4678,7 @@ void FHeaderParser::GetVarType(
 		}
 	}
 
-	if ( VarProperty.IsObject() && VarProperty.Type != CPT_SoftObjectReference && VarProperty.MetaClass == nullptr && (VarProperty.PropertyFlags&CPF_Config) != 0 )
+	if ( VarProperty.IsObject() && VarProperty.Type != CPT_SoftObjectReference && VarProperty.MetaClassDef == nullptr && (VarProperty.PropertyFlags&CPF_Config) != 0 )
 	{
 		FError::Throwf(TEXT("Not allowed to use 'config' with object variables"));
 	}
@@ -5791,7 +5774,7 @@ UClass* FHeaderParser::CompileClassDeclaration()
 	// Class metadata
 	MetaData.Append(ClassDef.MetaData);
 	ClassDef.MergeCategoryMetaData(MetaData);
-	AddIncludePathToMetadata(Class, MetaData);
+	AddIncludePathToMetadata(ClassDef, MetaData);
 	AddModuleRelativePathToMetadata(ClassDef, MetaData);
 
 	// Register the metadata
@@ -6015,7 +5998,7 @@ void FHeaderParser::CompileInterfaceDeclaration()
 	}
 
 	// Try parsing metadata for the interface
-	FStructMetaData& StructData = GTypeDefinitionInfoMap.FindChecked<FUnrealStructDefinitionInfo>(InterfaceClass).GetStructMetaData();
+	FStructMetaData& StructData = ClassDef.GetStructMetaData();
 	StructData.SetPrologLine(PrologFinishLine);
 
 	// Register the metadata
@@ -6324,9 +6307,12 @@ void FHeaderParser::ParseParameterList(FUnrealFunctionDefinitionInfo& FunctionDe
 				}
 			}
 
-			if (Property.Type == CPT_Struct && Property.Struct && !(Function->FunctionFlags & FUNC_NetRequest || Function->FunctionFlags & FUNC_NetResponse))
+			if (Property.Type == CPT_Struct)
 			{
-				ValidateScriptStructOkForNet(Property.Struct->GetName(), Property.Struct);
+				if (!(Function->FunctionFlags & FUNC_NetRequest || Function->FunctionFlags & FUNC_NetResponse))
+				{
+					ValidateScriptStructOkForNet(Property.ScriptStructDef->GetStruct()->GetName(), *Property.ScriptStructDef);
+				}
 			}
 
 			if (!(Function->FunctionFlags & FUNC_NetRequest))
@@ -7133,15 +7119,15 @@ void FHeaderParser::CompileFunctionDeclaration()
 
 	if (TopFunction->FunctionFlags & (FUNC_BlueprintCallable | FUNC_BlueprintEvent))
 	{
-		for (FUnrealPropertyDefinitionInfo* ProperyDef : FuncDef.GetProperties())
+		for (FUnrealPropertyDefinitionInfo* PropertyDef : FuncDef.GetProperties())
 		{
-			FProperty const* const Param = ProperyDef->GetProperty();
+			FProperty const* const Param = PropertyDef->GetProperty();
 			if (Param->ArrayDim > 1)
 			{
 				FError::Throwf(TEXT("Static array cannot be exposed to blueprint. Function: %s Parameter %s\n"), *TopFunction->GetName(), *Param->GetName());
 			}
 
-			if (!IsPropertySupportedByBlueprint(Param, false))
+			if (!FPropertyTraits::IsSupportedByBlueprint(*PropertyDef, false))
 			{
 				FString ExtendedCPPType;
 				FString CPPType = Param->GetCPPType(&ExtendedCPPType);
@@ -7293,48 +7279,55 @@ bool FHeaderParser::IsBitfieldProperty(ELayoutMacroType LayoutMacroType)
 
 void FHeaderParser::ValidatePropertyIsDeprecatedIfNecessary(const FPropertyBase& VarProperty, const EPropertyFlags* OuterPropertyFlags)
 {
+	// If the property is in an array that has been deprecated, then we don't need to do any further tests 
+	if (OuterPropertyFlags != nullptr && (*OuterPropertyFlags & CPF_Deprecated) != 0)
+	{
+		return;
+	}
+
+	// If the property is already marked deprecated, then we don't need to do any further tests 
+	if ((VarProperty.PropertyFlags & CPF_Deprecated) != 0)
+	{
+		return;
+	}
+
 	// check to see if we have a FClassProperty using a deprecated class
-	if ( VarProperty.MetaClass != NULL && VarProperty.MetaClass->HasAnyClassFlags(CLASS_Deprecated) && !(VarProperty.PropertyFlags & CPF_Deprecated) &&
-		(OuterPropertyFlags == NULL || !(*OuterPropertyFlags & CPF_Deprecated)) )
+	if (VarProperty.MetaClassDef != nullptr && VarProperty.MetaClassDef->GetClass()->HasAnyClassFlags(CLASS_Deprecated))
 	{
-		UE_LOG_ERROR_UHT(TEXT("Property is using a deprecated class: %s.  Property should be marked deprecated as well."), *VarProperty.MetaClass->GetPathName());
+		UE_LOG_ERROR_UHT(TEXT("Property is using a deprecated class: %s.  Property should be marked deprecated as well."), *VarProperty.MetaClassDef->GetClass()->GetPathName());
 	}
 
 	// check to see if we have a FObjectProperty using a deprecated class.
 	// PropertyClass is part of a union, so only check PropertyClass if this token represents an object property
-	if ( (VarProperty.Type == CPT_ObjectReference || VarProperty.Type == CPT_WeakObjectReference || VarProperty.Type == CPT_LazyObjectReference || VarProperty.Type == CPT_ObjectPtrReference || VarProperty.Type == CPT_SoftObjectReference) && VarProperty.PropertyClass != NULL
-		&&	VarProperty.PropertyClass->HasAnyClassFlags(CLASS_Deprecated)	// and the object class being used has been deprecated
-		&& (VarProperty.PropertyFlags&CPF_Deprecated) == 0					// and this property isn't marked deprecated as well
-		&& (OuterPropertyFlags == NULL || !(*OuterPropertyFlags & CPF_Deprecated)) ) // and this property isn't in an array that was marked deprecated either
+	if ((VarProperty.Type == CPT_ObjectReference || VarProperty.Type == CPT_WeakObjectReference || VarProperty.Type == CPT_LazyObjectReference || VarProperty.Type == CPT_ObjectPtrReference || VarProperty.Type == CPT_SoftObjectReference)
+		&& VarProperty.ClassDef != nullptr)
 	{
-		UE_LOG_ERROR_UHT(TEXT("Property is using a deprecated class: %s.  Property should be marked deprecated as well."), *VarProperty.PropertyClass->GetPathName());
+		UClass* PropertyClass = VarProperty.ClassDef->GetClass();
+		if (PropertyClass->HasAnyClassFlags(CLASS_Deprecated))	// and the object class being used has been deprecated
+		{
+			UE_LOG_ERROR_UHT(TEXT("Property is using a deprecated class: %s.  Property should be marked deprecated as well."), *PropertyClass->GetPathName());
+		}
 	}
 }
 
-bool FHeaderParser::ValidateScriptStructOkForNet(const FString& OriginStructName, UScriptStruct* InStruct)
+bool FHeaderParser::ValidateScriptStructOkForNet(const FString& OriginStructName, FUnrealScriptStructDefinitionInfo& InStructDef)
 {
-	if (!InStruct)
-	{
-		return false;
-	}
-
-	if (ScriptStructsValidForNet.Contains(InStruct))
+	if (ScriptStructsValidForNet.Contains(&InStructDef))
 	{
 		return true;
 	}
 
 	bool bIsStructValid = true;
 
-	if (UScriptStruct* SuperScriptStruct = Cast<UScriptStruct>(InStruct->GetSuperStruct()))
+	if (FUnrealScriptStructDefinitionInfo* SuperScriptStructDef = UHTCast<FUnrealScriptStructDefinitionInfo>(InStructDef.GetSuperClassInfo().Struct))
 	{
-		if (!ValidateScriptStructOkForNet(OriginStructName, SuperScriptStruct))
+		if (!ValidateScriptStructOkForNet(OriginStructName, *SuperScriptStructDef))
 		{
 			bIsStructValid = false;
 		}
 	}
 
-	FUnrealScriptStructDefinitionInfo& StructDef = GTypeDefinitionInfoMap.FindChecked<FUnrealScriptStructDefinitionInfo>(InStruct);
-	for (FUnrealPropertyDefinitionInfo* PropertyDef : StructDef.GetProperties())
+	for (FUnrealPropertyDefinitionInfo* PropertyDef : InStructDef.GetProperties())
 	{
 		FProperty* ChildProp = PropertyDef->GetProperty();
 		if (const FSetProperty* const SetProp = CastField<FSetProperty>(ChildProp))
@@ -7355,7 +7348,7 @@ bool FHeaderParser::ValidateScriptStructOkForNet(const FString& OriginStructName
 		}
 		else if (const FStructProperty* const StructProperty = CastField<FStructProperty>(ChildProp))
 		{
-			if (!ValidateScriptStructOkForNet(OriginStructName, StructProperty->Struct))
+			if (!ValidateScriptStructOkForNet(OriginStructName, *PropertyDef->GetPropertyBase().ScriptStructDef))
 			{
 				bIsStructValid = false;
 			}
@@ -7364,7 +7357,7 @@ bool FHeaderParser::ValidateScriptStructOkForNet(const FString& OriginStructName
 
 	if (bIsStructValid)
 	{
-		ScriptStructsValidForNet.Add(InStruct);
+		ScriptStructsValidForNet.Add(&InStructDef);
 	}
 
 	return bIsStructValid;
@@ -7394,9 +7387,10 @@ struct FExposeOnSpawnValidator
 			ProperNativeType = true;
 		}
 
-		if (!ProperNativeType && (CPT_Struct == Property.Type) && Property.Struct)
+		if (!ProperNativeType && (CPT_Struct == Property.Type) && Property.ScriptStructDef)
 		{
-			ProperNativeType |= Property.Struct->GetBoolMetaData(FHeaderParserNames::NAME_BlueprintType);
+			UScriptStruct* Struct = Property.ScriptStructDef->GetScriptStruct();
+			ProperNativeType |= Struct->GetBoolMetaData(FHeaderParserNames::NAME_BlueprintType);
 		}
 
 		return ProperNativeType;
@@ -7432,7 +7426,7 @@ void FHeaderParser::CompileVariableDeclaration(FUnrealStructDefinitionInfo& Stru
 	}
 
 	// Validate that pointer properties are not interfaces (which are not GC'd and so will cause runtime errors)
-	if (OriginalProperty.PointerType == EPointerType::Native && OriginalProperty.Struct->IsChildOf(UInterface::StaticClass()))
+	if (OriginalProperty.PointerType == EPointerType::Native && OriginalProperty.ClassDef->AsClass() != nullptr && OriginalProperty.ClassDef->IsInterface())
 	{
 		// Get the name of the type, removing the asterisk representing the pointer
 		FString TypeName = FString(TypeRange.Count, Input + TypeRange.StartIndex).TrimStartAndEnd().LeftChop(1).TrimEnd();
@@ -7479,9 +7473,9 @@ void FHeaderParser::CompileVariableDeclaration(FUnrealStructDefinitionInfo& Stru
 	}
 
 	// If Property is a Replicated Struct check to make sure there are no Properties that are not allowed to be Replicated in the Struct 
-	if (OriginalProperty.Type == CPT_Struct && OriginalProperty.PropertyFlags & CPF_Net && OriginalProperty.Struct)
+	if (OriginalProperty.Type == CPT_Struct && OriginalProperty.PropertyFlags & CPF_Net && OriginalProperty.ScriptStructDef)
 	{
-		ValidateScriptStructOkForNet(OriginalProperty.Struct->GetName(), OriginalProperty.Struct);
+		ValidateScriptStructOkForNet(OriginalProperty.ScriptStructDef->GetScriptStruct()->GetName(), *OriginalProperty.ScriptStructDef);
 	}
 
 	// Process all variables of this type.
@@ -7545,7 +7539,7 @@ void FHeaderParser::CompileVariableDeclaration(FUnrealStructDefinitionInfo& Stru
 				UE_LOG_ERROR_UHT(TEXT("Static array cannot be exposed to blueprint %s.%s"), *Struct->GetName(), *NewProperty->GetName());
 			}
 
-			if (!IsPropertySupportedByBlueprint(NewProperty, true))
+			if (!FPropertyTraits::IsSupportedByBlueprint(NewPropDef, true))
 			{
 				FString ExtendedCPPType;
 				FString CPPType = NewProperty->GetCPPType(&ExtendedCPPType);
@@ -8051,10 +8045,14 @@ ECompilationResult::Type FHeaderParser::Parse(
 
 		for (const TSharedRef<FUnrealTypeDefinitionInfo>& ClassDataPair : SourceFile.GetDefinedClasses())
 		{
-			UClass* Class = ClassDataPair->AsClassChecked().GetClass();
-			for (UClass* ParentClass = Class->GetSuperClass(); ParentClass && !ParentClass->HasAnyClassFlags(CLASS_Parsed | CLASS_Intrinsic); ParentClass = ParentClass->GetSuperClass())
+			FUnrealClassDefinitionInfo& ClassDef = ClassDataPair->AsClassChecked();
+			for (FUnrealClassDefinitionInfo* ParentClassDef = ClassDef.GetSuperClass(); ParentClassDef; ParentClassDef = ParentClassDef->GetSuperClass())
 			{
-				SourceFilesRequired.Add(&GTypeDefinitionInfoMap[ParentClass]->GetUnrealSourceFile());
+				if (ParentClassDef->GetClass()->HasAnyClassFlags(CLASS_Parsed | CLASS_Intrinsic))
+				{
+					break;
+				}
+				SourceFilesRequired.Add(&ParentClassDef->GetUnrealSourceFile());
 			}
 		}
 
diff --git a/Engine/Source/Programs/UnrealHeaderTool/Private/HeaderParser.h b/Engine/Source/Programs/UnrealHeaderTool/Private/HeaderParser.h
--- a/Engine/Source/Programs/UnrealHeaderTool/Private/HeaderParser.h
+++ b/Engine/Source/Programs/UnrealHeaderTool/Private/HeaderParser.h
@@ -888,7 +888,7 @@ protected:
 	static void ValidatePropertyIsDeprecatedIfNecessary(const FPropertyBase& VarProperty, const EPropertyFlags* OuterPropertyFlags);
 
 	// Cache of ScriptStructs that have been validated for Net Replication and RPC
-	TSet<UScriptStruct*> ScriptStructsValidForNet;
+	TSet<FUnrealScriptStructDefinitionInfo*> ScriptStructsValidForNet;
 
 	/**
 	 * Validate that a ScriptStruct is ok to be Replicated or Sent in an RPC.
@@ -896,7 +896,7 @@ protected:
 	 * @param OriginStructName  The Name of the ScriptStruct to check
 	 * @param InStruct          The ScriptStruct to check 
 	 */
-	bool ValidateScriptStructOkForNet(const FString& OriginStructName, UScriptStruct* InStruct);
+	bool ValidateScriptStructOkForNet(const FString& OriginStructName, FUnrealScriptStructDefinitionInfo& InStructDef);
 
 private:
 	// Definition of package being parsed
diff --git a/Engine/Source/Programs/UnrealHeaderTool/Private/NativeClassExporter.h b/Engine/Source/Programs/UnrealHeaderTool/Private/NativeClassExporter.h
--- a/Engine/Source/Programs/UnrealHeaderTool/Private/NativeClassExporter.h
+++ b/Engine/Source/Programs/UnrealHeaderTool/Private/NativeClassExporter.h
@@ -331,13 +331,13 @@ private:
 	static void ExportProperties(FOutputDevice& Out, FUnrealStructDefinitionInfo& StructDef, int32 TextIndent);
 
 	/** Return the name of the singleton function */
-	static const FString& GetPackageSingletonName(UPackage* Item, TSet<FString>* UniqueCrossModuleReferences);
+	static const FString& GetPackageSingletonName(FUnrealPackageDefinitionInfo& PackageDef, TSet<FString>* UniqueCrossModuleReferences);
 
 	/** Return the address of the singleton function */
-	static const FString& GetPackageSingletonNameFuncAddr(UPackage* Item, TSet<FString>* UniqueCrossModuleReferences);
+	static const FString& GetPackageSingletonNameFuncAddr(FUnrealPackageDefinitionInfo& PackageDef, TSet<FString>* UniqueCrossModuleReferences);
 
 	/** Return the address of the singleton function - handles nullptr */
-	static const FString& GetSingletonNameFuncAddr(FUnrealFieldDefinitionInfo& Item, TSet<FString>* UniqueCrossModuleReferences, bool bRequiresValidObject = true);
+	static const FString& GetSingletonNameFuncAddr(FUnrealFieldDefinitionInfo* FieldDef, TSet<FString>* UniqueCrossModuleReferences, bool bRequiresValidObject = true);
 
 	/** Return the address of the singleton function - handles nullptr */
 	static const FString& GetSingletonNameFuncAddr(UField* Item, TSet<FString>* UniqueCrossModuleReferences, bool bRequiresValidObject = true);
@@ -407,7 +407,7 @@ private:
 	 * @param	Out		the output device for the mirror struct
 	 * @param	Enums	the enum to export
 	 */
-	void ExportEnum(FOutputDevice& Out, UEnum* Enum) const;
+	void ExportEnum(FOutputDevice& Out, FUnrealEnumDefinitionInfo& EnumDef) const;
 
 	/**
 	 * Exports the inl text for enums declared in non-UClass headers.
@@ -415,7 +415,7 @@ private:
 	 * @param	OutputGetter	The function to call to get the output.
 	 * @param	Enum			the enum to export
 	 */
-	void ExportGeneratedEnumInitCode(FOutputDevice& Out, FReferenceGatherers& OutReferenceGatherers, const FUnrealSourceFile& SourceFile, UEnum* Enum) const;
+	void ExportGeneratedEnumInitCode(FOutputDevice& Out, FReferenceGatherers& OutReferenceGatherers, const FUnrealSourceFile& SourceFile, FUnrealEnumDefinitionInfo& EnumDef) const;
 
 	/**
 	 * Exports the macro declarations for GENERATED_BODY() for each Foo in the struct specified
@@ -423,7 +423,7 @@ private:
 	 * @param	Out				output device
 	 * @param	Struct			The struct to export
 	 */
-	void ExportGeneratedStructBodyMacros(FOutputDevice& OutGeneratedHeaderText, FOutputDevice& Out, FReferenceGatherers& OutReferenceGatherers, const FUnrealSourceFile& SourceFile, UScriptStruct* Struct) const;
+	void ExportGeneratedStructBodyMacros(FOutputDevice& OutGeneratedHeaderText, FOutputDevice& Out, FReferenceGatherers& OutReferenceGatherers, const FUnrealSourceFile& SourceFile, FUnrealScriptStructDefinitionInfo& ScriptStructDef) const;
 
 	/**
 	 * Exports a local mirror of the specified struct; used to get offsets
@@ -582,7 +582,7 @@ private:
 	 *
 	 * @return      A pair of strings which represents the pointer and a count of the emitted properties.
 	 */
-	TTuple<FString, FString> OutputProperties(FOutputDevice& DeclOut, FOutputDevice& Out, FReferenceGatherers& OutReferenceGatherers, const TCHAR* Scope, const TArray<FProperty*>& Properties, const TCHAR* DeclSpaces, const TCHAR* Spaces) const;
+	TTuple<FString, FString> OutputProperties(FOutputDevice& DeclOut, FOutputDevice& Out, FReferenceGatherers& OutReferenceGatherers, const TCHAR* Scope, const TArray<FUnrealPropertyDefinitionInfo*>& PropertyDefs, const TCHAR* DeclSpaces, const TCHAR* Spaces) const;
 
 	/**
 	 * Function to output the C++ code necessary to set up a property
@@ -610,21 +610,13 @@ private:
 	**/
 	void PropertyNew(FOutputDevice& DeclOut, FOutputDevice& Out, FReferenceGatherers& OutReferenceGatherers, FUnrealPropertyDefinitionInfo& PropertyDef, const TCHAR* OffsetStr, const TCHAR* Name, const TCHAR* DeclSpaces, const TCHAR* Spaces, const TCHAR* SourceStruct = nullptr) const;
 
-	/**
-	 * Function to generate the property tag
-	 *
-	 * @param	Out				Destination string builder.
-	 * @param	Object			UObject in question
-	 */
-	static void GetPropertyTag(FUHTStringBuilder& Out, UObject* Object);
-
 	/**
 	 * Function to generate the property tag
 	 * 
 	 * @param	Out				Destination string builder.
-	 * @param	Prop			Property in question.
+	 * @param	PropDef			Property in question.
 	 */
-	static void GetPropertyTag(FUHTStringBuilder& Out, FProperty* Prop);
+	static void GetPropertyTag(FUHTStringBuilder& Out, FUnrealPropertyDefinitionInfo& PropDef);
 
 	/**
 	 * Exports the proxy definitions for the list of enums specified
diff --git a/Engine/Source/Programs/UnrealHeaderTool/Private/ParserHelper.cpp b/Engine/Source/Programs/UnrealHeaderTool/Private/ParserHelper.cpp
--- a/Engine/Source/Programs/UnrealHeaderTool/Private/ParserHelper.cpp
+++ b/Engine/Source/Programs/UnrealHeaderTool/Private/ParserHelper.cpp
@@ -11,6 +11,58 @@
 /////////////////////////////////////////////////////
 // FPropertyBase
 
+FPropertyBase::FPropertyBase(FUnrealEnumDefinitionInfo& InEnumDef, EPropertyType InType)
+	: Type(InType)
+	, ArrayType(EArrayType::None)
+	, PropertyFlags(CPF_None)
+	, ImpliedPropertyFlags(CPF_None)
+	, RefQualifier(ERefQualifier::None)
+	, PropertyExportFlags(PROPEXPORT_Public)
+	, EnumDef(&InEnumDef)
+	, DelegateName(NAME_None)
+	, DelegateSignatureOwnerClass(nullptr)
+	, RepNotifyName(NAME_None)
+	, PointerType(EPointerType::None)
+	, IntType(GetSizedIntTypeFromPropertyType(InType))
+{
+}
+
+FPropertyBase::FPropertyBase(FUnrealClassDefinitionInfo& InClassDef, EPropertyType InType, bool bWeakIsAuto/* = false*/)
+	: Type(InType)
+	, ArrayType(EArrayType::None)
+	, PropertyFlags(CPF_None)
+	, ImpliedPropertyFlags(CPF_None)
+	, RefQualifier(ERefQualifier::None)
+	, PropertyExportFlags(PROPEXPORT_Public)
+	, ClassDef(&InClassDef)
+	, DelegateName(NAME_None)
+	, DelegateSignatureOwnerClass(nullptr)
+	, RepNotifyName(NAME_None)
+	, PointerType(EPointerType::None)
+	, IntType(EIntType::None)
+{
+	if ((Type == CPT_WeakObjectReference) && bWeakIsAuto)
+	{
+		PropertyFlags |= CPF_AutoWeak;
+	}
+}
+
+FPropertyBase::FPropertyBase(FUnrealScriptStructDefinitionInfo& InStructDef)
+	: Type(CPT_Struct)
+	, ArrayType(EArrayType::None)
+	, PropertyFlags(CPF_None)
+	, ImpliedPropertyFlags(CPF_None)
+	, RefQualifier(ERefQualifier::None)
+	, PropertyExportFlags(PROPEXPORT_Public)
+	, ScriptStructDef(&InStructDef)
+	, DelegateName(NAME_None)
+	, DelegateSignatureOwnerClass(nullptr)
+	, RepNotifyName(NAME_None)
+	, PointerType(EPointerType::None)
+	, IntType(EIntType::None)
+{
+}
+
 const TCHAR* FPropertyBase::GetPropertyTypeText( EPropertyType Type )
 {
 	switch ( Type )
@@ -52,6 +104,184 @@ const TCHAR* FPropertyBase::GetPropertyTypeText( EPropertyType Type )
 	return TEXT("");
 }
 
+bool FPropertyBase::MatchesType(const FPropertyBase& Other, bool bDisallowGeneralization, bool bIgnoreImplementedInterfaces/* = false*/) const
+{
+	check(Type != CPT_None || !bDisallowGeneralization);
+
+	bool bIsObjectType = IsObject();
+	bool bOtherIsObjectType = Other.IsObject();
+	bool bIsObjectComparison = bIsObjectType && bOtherIsObjectType;
+	bool bReverseClassChainCheck = true;
+
+	// If converting to an l-value, we require an exact match with an l-value.
+	if ((PropertyFlags & CPF_OutParm) != 0)
+	{
+		// if the other type is not an l-value, disallow
+		if ((Other.PropertyFlags & CPF_OutParm) == 0)
+		{
+			return false;
+		}
+
+		// if the other type is const and we are not const, disallow
+		if ((Other.PropertyFlags & CPF_ConstParm) != 0 && (PropertyFlags & CPF_ConstParm) == 0)
+		{
+			return false;
+		}
+
+		if (Type == CPT_Struct)
+		{
+			// Allow derived structs to be passed by reference, unless this is a dynamic array of structs
+			bDisallowGeneralization = bDisallowGeneralization || ArrayType == EArrayType::Dynamic || Other.ArrayType == EArrayType::Dynamic;
+		}
+
+		// if Type == CPT_ObjectReference, out object function parm; allow derived classes to be passed in
+		// if Type == CPT_Interface, out interface function parm; allow derived classes to be passed in
+		else if ((PropertyFlags & CPF_ConstParm) == 0 || !IsObject())
+		{
+			// all other variable types must match exactly when passed as the value to an 'out' parameter
+			bDisallowGeneralization = true;
+		}
+
+		// both types are objects, but one is an interface and one is an object reference
+		else if (bIsObjectComparison && Type != Other.Type)
+		{
+			return false;
+		}
+	}
+	else if ((Type == CPT_ObjectReference || Type == CPT_WeakObjectReference || Type == CPT_LazyObjectReference || Type == CPT_ObjectPtrReference || Type == CPT_SoftObjectReference) && Other.Type != CPT_Interface && (PropertyFlags & CPF_ReturnParm))
+	{
+		bReverseClassChainCheck = false;
+	}
+
+	// Check everything.
+	if (Type == CPT_None && (Other.Type == CPT_None || !bDisallowGeneralization))
+	{
+		// If Other has no type, accept anything.
+		return true;
+	}
+	else if (Type != Other.Type && !bIsObjectComparison)
+	{
+		// Mismatched base types.
+		return false;
+	}
+	else if (ArrayType != Other.ArrayType)
+	{
+		// Mismatched array types.
+		return false;
+	}
+	else if (Type == CPT_Byte)
+	{
+		// Make sure enums match, or we're generalizing.
+		return EnumDef == Other.EnumDef || (EnumDef == NULL && !bDisallowGeneralization);
+	}
+	else if (bIsObjectType)
+	{
+		check(ClassDef != NULL);
+
+		// Make sure object types match, or we're generalizing.
+		if (bDisallowGeneralization)
+		{
+			// Exact match required.
+			return ClassDef == Other.ClassDef && MetaClassDef == Other.MetaClassDef;
+		}
+		else if (Other.ClassDef == NULL)
+		{
+			// Cannonical 'None' matches all object classes.
+			return true;
+		}
+		else
+		{
+			UClass* PropertyClass = ClassDef->GetClass();
+			UClass* OtherPropertyClass = Other.ClassDef->GetClass();
+			UClass* MetaClass = MetaClassDef ? MetaClassDef->GetClass() : nullptr;
+			UClass* OtherMetaClass = Other.MetaClassDef ? Other.MetaClassDef->GetClass() : nullptr;
+			// Generalization is ok (typical example of this check would look like: VarA = VarB;, where this is VarB and Other is VarA)
+			if (OtherPropertyClass->IsChildOf(PropertyClass))
+			{
+				if (!bIgnoreImplementedInterfaces || ((Type == CPT_Interface) == (Other.Type == CPT_Interface)))
+				{
+					if (!PropertyClass->IsChildOf(UClass::StaticClass()) || MetaClass == NULL || OtherMetaClass->IsChildOf(MetaClass) ||
+						(bReverseClassChainCheck && (OtherMetaClass == NULL || MetaClass->IsChildOf(OtherMetaClass))))
+					{
+						return true;
+					}
+				}
+			}
+			// check the opposite class chain for object types
+			else if (bReverseClassChainCheck && Type != CPT_Interface && bIsObjectComparison && PropertyClass != NULL && PropertyClass->IsChildOf(OtherPropertyClass))
+			{
+				if (!OtherPropertyClass->IsChildOf(UClass::StaticClass()) || MetaClass == NULL || OtherMetaClass == NULL || MetaClass->IsChildOf(OtherMetaClass) || OtherMetaClass->IsChildOf(MetaClass))
+				{
+					return true;
+				}
+			}
+
+			if (PropertyClass->HasAnyClassFlags(CLASS_Interface) && !bIgnoreImplementedInterfaces)
+			{
+				if (OtherPropertyClass->ImplementsInterface(PropertyClass))
+				{
+					return true;
+				}
+			}
+
+			return false;
+		}
+	}
+	else if (Type == CPT_Struct)
+	{
+		check(ScriptStructDef != NULL);
+		check(Other.ScriptStructDef != NULL);
+
+		if (ScriptStructDef == Other.ScriptStructDef)
+		{
+			// struct types match exactly 
+			return true;
+		}
+
+		// returning false here prevents structs related through inheritance from being used interchangeably, such as passing a derived struct as the value for a parameter
+		// that expects the base struct, or vice versa.  An easier example is assignment (e.g. Vector = Plane or Plane = Vector).
+		// there are two cases to consider (let's use vector and plane for the example):
+		// - Vector = Plane;
+		//		in this expression, 'this' is the vector, and Other is the plane.  This is an unsafe conversion, as the destination property type is used to copy the r-value to the l-value
+		//		so in this case, the VM would call CopyCompleteValue on the FPlane struct, which would copy 16 bytes into the l-value's buffer;  However, the l-value buffer will only be
+		//		12 bytes because that is the size of FVector
+		// - Plane = Vector;
+		//		in this expression, 'this' is the plane, and Other is the vector.  This is a safe conversion, since only 12 bytes would be copied from the r-value into the l-value's buffer
+		//		(which would be 16 bytes).  The problem with allowing this conversion is that what to do with the extra member (e.g. Plane.W); should it be left alone? should it be zeroed?
+		//		difficult to say what the correct behavior should be, so let's just ignore inheritance for the sake of determining whether two structs are identical
+
+		// Previously, the logic for determining whether this is a generalization of Other was reversed; this is very likely the culprit behind all current issues with 
+		// using derived structs interchangeably with their base versions.  The inheritance check has been fixed; for now, allow struct generalization and see if we can find any further
+		// issues with allowing conversion.  If so, then we disable all struct generalization by returning false here.
+		// return false;
+
+		if (bDisallowGeneralization)
+		{
+			return false;
+		}
+
+		UScriptStruct* Struct = ScriptStructDef->GetScriptStruct();
+		UScriptStruct* OtherStruct = Other.ScriptStructDef->GetScriptStruct();
+
+		// Generalization is ok if this is not a dynamic array
+		if (ArrayType != EArrayType::Dynamic && Other.ArrayType != EArrayType::Dynamic)
+		{
+			if (!OtherStruct->IsChildOf(Struct) && Struct->IsChildOf(OtherStruct))
+			{
+				return true;
+			}
+		}
+
+		return false;
+	}
+	else
+	{
+		// General match.
+		return true;
+	}
+}
+
+
 /////////////////////////////////////////////////////
 // FToken
 
diff --git a/Engine/Source/Programs/UnrealHeaderTool/Private/ParserHelper.h b/Engine/Source/Programs/UnrealHeaderTool/Private/ParserHelper.h
--- a/Engine/Source/Programs/UnrealHeaderTool/Private/ParserHelper.h
+++ b/Engine/Source/Programs/UnrealHeaderTool/Private/ParserHelper.h
@@ -23,6 +23,11 @@ class UField;
 class UMetaData;
 class FHeaderParser;
 class FUnrealPropertyDefinitionInfo;
+class FUnrealTypeDefinitionInfo;
+class FUnrealEnumDefinitionInfo;
+class FUnrealClassDefinitionInfo;
+class FUnrealScriptStructDefinitionInfo;
+class FUnrealFunctionDefinitionInfo;
 
 /*-----------------------------------------------------------------------------
 	FPropertyBase.
@@ -136,21 +141,18 @@ public:
 	 * A mask of EPropertyHeaderExportFlags which are used for modifying how this property is exported to the native class header
 	 */
 	uint32 PropertyExportFlags;
+
 	union
 	{
-		class UEnum* Enum;
-		class UClass* PropertyClass;
-		class UScriptStruct* Struct;
-		class UFunction* Function;
-		class FFieldClass* PropertyPathClass;
-#if PLATFORM_64BITS
-		int64 StringSize;
-#else
-		int32 StringSize;
-#endif
+		FUnrealTypeDefinitionInfo* TypeDef = nullptr;
+		FUnrealEnumDefinitionInfo* EnumDef;
+		FUnrealScriptStructDefinitionInfo* ScriptStructDef;
+		FUnrealClassDefinitionInfo* ClassDef;
+		FUnrealFunctionDefinitionInfo* FunctionDef;
 	};
+	FFieldClass* PropertyPathClass = nullptr;
+	FUnrealClassDefinitionInfo* MetaClassDef = nullptr;
 
-	UClass* MetaClass;
 	FName   DelegateName;
 	UClass*	DelegateSignatureOwnerClass;
 	FName   RepNotifyName;
@@ -179,8 +181,6 @@ public:
 	, ImpliedPropertyFlags(CPF_None)
 	, RefQualifier        (ERefQualifier::None)
 	, PropertyExportFlags (PROPEXPORT_Public)
-	, StringSize          (0)
-	, MetaClass           (nullptr)
 	, DelegateName        (NAME_None)
 	, DelegateSignatureOwnerClass(nullptr)
 	, RepNotifyName       (NAME_None)
@@ -196,8 +196,6 @@ public:
 	, ImpliedPropertyFlags(CPF_None)
 	, RefQualifier        (ERefQualifier::None)
 	, PropertyExportFlags (PROPEXPORT_Public)
-	, StringSize          (0)
-	, MetaClass           (nullptr)
 	, DelegateName        (NAME_None)
 	, DelegateSignatureOwnerClass(nullptr)
 	, RepNotifyName       (NAME_None)
@@ -206,60 +204,11 @@ public:
 	{
 	}
 
-	explicit FPropertyBase(UEnum* InEnum, EPropertyType InType)
-	: Type                (InType)
-	, ArrayType           (EArrayType::None)
-	, PropertyFlags       (CPF_None)
-	, ImpliedPropertyFlags(CPF_None)
-	, RefQualifier        (ERefQualifier::None)
-	, PropertyExportFlags (PROPEXPORT_Public)
-	, Enum                (InEnum)
-	, MetaClass           (nullptr)
-	, DelegateName        (NAME_None)
-	, DelegateSignatureOwnerClass(nullptr)
-	, RepNotifyName       (NAME_None)
-	, PointerType         (EPointerType::None)
-	, IntType             (GetSizedIntTypeFromPropertyType(InType))
-	{
-	}
+	explicit FPropertyBase(FUnrealEnumDefinitionInfo& InEnumDef, EPropertyType InType);
 
-	explicit FPropertyBase(UClass* InClass, EPropertyType InType, bool bWeakIsAuto = false)
-	: Type                (InType)
-	, ArrayType           (EArrayType::None)
-	, PropertyFlags       (CPF_None)
-	, ImpliedPropertyFlags(CPF_None)
-	, RefQualifier        (ERefQualifier::None)
-	, PropertyExportFlags (PROPEXPORT_Public)
-	, PropertyClass       (InClass)
-	, MetaClass           (nullptr)
-	, DelegateName        (NAME_None)
-	, DelegateSignatureOwnerClass(nullptr)
-	, RepNotifyName       (NAME_None)
-	, PointerType         (EPointerType::None)
-	, IntType             (EIntType::None)
-	{
-		if ((Type == CPT_WeakObjectReference) && bWeakIsAuto)
-		{
-			PropertyFlags |= CPF_AutoWeak;
-		}
-	}
+	explicit FPropertyBase(FUnrealClassDefinitionInfo& InClassDef, EPropertyType InType, bool bWeakIsAuto = false);
 
-	explicit FPropertyBase(UScriptStruct* InStruct)
-	: Type                (CPT_Struct)
-	, ArrayType           (EArrayType::None)
-	, PropertyFlags       (CPF_None)
-	, ImpliedPropertyFlags(CPF_None)
-	, RefQualifier        (ERefQualifier::None)
-	, PropertyExportFlags (PROPEXPORT_Public)
-	, Struct              (InStruct)
-	, MetaClass           (NULL)
-	, DelegateName        (NAME_None)
-	, DelegateSignatureOwnerClass(nullptr)
-	, RepNotifyName       (NAME_None)
-	, PointerType         (EPointerType::None)
-	, IntType             (EIntType::None)
-	{
-	}
+	explicit FPropertyBase(FUnrealScriptStructDefinitionInfo& InStructDef);
 
 	explicit FPropertyBase(FFieldClass* InPropertyClass, EPropertyType InType)
 		: Type(InType)
@@ -269,7 +218,6 @@ public:
 		, RefQualifier(ERefQualifier::None)
 		, PropertyExportFlags(PROPEXPORT_Public)
 		, PropertyPathClass(InPropertyClass)
-		, MetaClass(nullptr)
 		, DelegateName(NAME_None)
 		, DelegateSignatureOwnerClass(nullptr)
 		, RepNotifyName(NAME_None)
@@ -318,175 +266,7 @@ public:
 	 * @param	bIgnoreImplementedInterfaces	controls whether two types can be considered a match if one type is an interface implemented
 	 *											by the other type.
 	 */
-	bool MatchesType( const FPropertyBase& Other, bool bDisallowGeneralization, bool bIgnoreImplementedInterfaces=false ) const
-	{
-		check(Type!=CPT_None || !bDisallowGeneralization);
-
-		bool bIsObjectType = IsObject();
-		bool bOtherIsObjectType = Other.IsObject();
-		bool bIsObjectComparison = bIsObjectType && bOtherIsObjectType;
-		bool bReverseClassChainCheck = true;
-
-		// If converting to an l-value, we require an exact match with an l-value.
-		if( (PropertyFlags&CPF_OutParm) != 0 )
-		{
-			// if the other type is not an l-value, disallow
-			if ( (Other.PropertyFlags&CPF_OutParm) == 0 )
-			{
-				return false;
-			}
-
-			// if the other type is const and we are not const, disallow
-			if ( (Other.PropertyFlags&CPF_ConstParm) != 0 && (PropertyFlags&CPF_ConstParm) == 0 )
-			{
-				return false;
-			}
-
-			if ( Type == CPT_Struct )
-			{
-				// Allow derived structs to be passed by reference, unless this is a dynamic array of structs
-				bDisallowGeneralization = bDisallowGeneralization || ArrayType == EArrayType::Dynamic || Other.ArrayType == EArrayType::Dynamic;
-			}
-
-			// if Type == CPT_ObjectReference, out object function parm; allow derived classes to be passed in
-			// if Type == CPT_Interface, out interface function parm; allow derived classes to be passed in
-			else if ( (PropertyFlags & CPF_ConstParm) == 0 || !IsObject() )
-			{
-				// all other variable types must match exactly when passed as the value to an 'out' parameter
-				bDisallowGeneralization = true;
-			}
-			
-			// both types are objects, but one is an interface and one is an object reference
-			else if ( bIsObjectComparison && Type != Other.Type )
-			{
-				return false;
-			}
-		}
-		else if ((Type == CPT_ObjectReference || Type == CPT_WeakObjectReference || Type == CPT_LazyObjectReference || Type == CPT_ObjectPtrReference || Type == CPT_SoftObjectReference) && Other.Type != CPT_Interface && (PropertyFlags & CPF_ReturnParm))
-		{
-			bReverseClassChainCheck = false;
-		}
-
-		// Check everything.
-		if( Type==CPT_None && (Other.Type==CPT_None || !bDisallowGeneralization) )
-		{
-			// If Other has no type, accept anything.
-			return true;
-		}
-		else if( Type != Other.Type && !bIsObjectComparison )
-		{
-			// Mismatched base types.
-			return false;
-		}
-		else if( ArrayType != Other.ArrayType )
-		{
-			// Mismatched array types.
-			return false;
-		}
-		else if( Type==CPT_Byte )
-		{
-			// Make sure enums match, or we're generalizing.
-			return Enum==Other.Enum || (Enum==NULL && !bDisallowGeneralization);
-		}
-		else if( bIsObjectType )
-		{
-			check(PropertyClass!=NULL);
-
-			// Make sure object types match, or we're generalizing.
-			if( bDisallowGeneralization )
-			{
-				// Exact match required.
-				return PropertyClass==Other.PropertyClass && MetaClass==Other.MetaClass;
-			}
-			else if( Other.PropertyClass==NULL )
-			{
-				// Cannonical 'None' matches all object classes.
-				return true;
-			}
-			else
-			{
-				// Generalization is ok (typical example of this check would look like: VarA = VarB;, where this is VarB and Other is VarA)
-				if ( Other.PropertyClass->IsChildOf(PropertyClass) )
-				{
-					if ( !bIgnoreImplementedInterfaces || ((Type == CPT_Interface) == (Other.Type == CPT_Interface)) )
-					{
-						if ( !PropertyClass->IsChildOf(UClass::StaticClass()) || MetaClass == NULL || Other.MetaClass->IsChildOf(MetaClass) ||
-							(bReverseClassChainCheck && (Other.MetaClass == NULL || MetaClass->IsChildOf(Other.MetaClass))) )
-						{
-							return true;
-						}
-					}
-				}
-				// check the opposite class chain for object types
-				else if (bReverseClassChainCheck && Type != CPT_Interface && bIsObjectComparison && PropertyClass != NULL && PropertyClass->IsChildOf(Other.PropertyClass))
-				{
-					if (!Other.PropertyClass->IsChildOf(UClass::StaticClass()) || MetaClass == NULL || Other.MetaClass == NULL || MetaClass->IsChildOf(Other.MetaClass) || Other.MetaClass->IsChildOf(MetaClass))
-					{
-						return true;
-					}
-				}
-
-				if ( PropertyClass->HasAnyClassFlags(CLASS_Interface) && !bIgnoreImplementedInterfaces )
-				{
-					if ( Other.PropertyClass->ImplementsInterface(PropertyClass) )
-					{
-						return true;
-					}
-				}
-
-				return false;
-			}
-		}
-		else if( Type==CPT_Struct )
-		{
-			check(Struct!=NULL);
-			check(Other.Struct!=NULL);
-
-			if ( Struct == Other.Struct )
-			{
-				// struct types match exactly 
-				return true;
-			}
-
-			// returning false here prevents structs related through inheritance from being used interchangeably, such as passing a derived struct as the value for a parameter
-			// that expects the base struct, or vice versa.  An easier example is assignment (e.g. Vector = Plane or Plane = Vector).
-			// there are two cases to consider (let's use vector and plane for the example):
-			// - Vector = Plane;
-			//		in this expression, 'this' is the vector, and Other is the plane.  This is an unsafe conversion, as the destination property type is used to copy the r-value to the l-value
-			//		so in this case, the VM would call CopyCompleteValue on the FPlane struct, which would copy 16 bytes into the l-value's buffer;  However, the l-value buffer will only be
-			//		12 bytes because that is the size of FVector
-			// - Plane = Vector;
-			//		in this expression, 'this' is the plane, and Other is the vector.  This is a safe conversion, since only 12 bytes would be copied from the r-value into the l-value's buffer
-			//		(which would be 16 bytes).  The problem with allowing this conversion is that what to do with the extra member (e.g. Plane.W); should it be left alone? should it be zeroed?
-			//		difficult to say what the correct behavior should be, so let's just ignore inheritance for the sake of determining whether two structs are identical
-
-			// Previously, the logic for determining whether this is a generalization of Other was reversed; this is very likely the culprit behind all current issues with 
-			// using derived structs interchangeably with their base versions.  The inheritance check has been fixed; for now, allow struct generalization and see if we can find any further
-			// issues with allowing conversion.  If so, then we disable all struct generalization by returning false here.
- 			// return false;
-
-			if ( bDisallowGeneralization )
-			{
-				return false;
-			}
-
-			// Generalization is ok if this is not a dynamic array
-			if ( ArrayType != EArrayType::Dynamic && Other.ArrayType != EArrayType::Dynamic )
-			{
-				if ( !Other.Struct->IsChildOf(Struct) && Struct->IsChildOf(Other.Struct) )
-				{
-					return true;
-				}
-			}
-
-			return false;
-		}
-		else
-		{
-			// General match.
-			return true;
-		}
-	}
+	bool MatchesType(const FPropertyBase& Other, bool bDisallowGeneralization, bool bIgnoreImplementedInterfaces = false) const;
 	//@}
 
 	EIntType GetSizedIntTypeFromPropertyType(EPropertyType PropType)
diff --git a/Engine/Source/Programs/UnrealHeaderTool/Private/PropertyTypes.cpp b/Engine/Source/Programs/UnrealHeaderTool/Private/PropertyTypes.cpp
--- a/Engine/Source/Programs/UnrealHeaderTool/Private/PropertyTypes.cpp
+++ b/Engine/Source/Programs/UnrealHeaderTool/Private/PropertyTypes.cpp
@@ -179,8 +179,8 @@ namespace
 		}
 
 		// Check if it's an enum class property
-		// NOTE: VarProperty.Enum is a union and might not be an enum
-		if (FUnrealEnumDefinitionInfo* EnumDef = GTypeDefinitionInfoMap.Find<FUnrealEnumDefinitionInfo>(VarProperty.Enum))
+		// NOTE: VarProperty.EnumDef is a union and might not be an enum
+		if (UHTCast<FUnrealEnumDefinitionInfo>(VarProperty.TypeDef) != nullptr)
 		{
 			return FuncDispatch<FPropertyTypeTraitsEnum>()(std::forward<Args>(args)...);
 		}
@@ -324,11 +324,12 @@ struct FPropertyTypeTraitsByte : public FPropertyTypeTraitsNumericBase
 	static FProperty* CreateEngineType(FUnrealPropertyDefinitionInfo& PropDef, FFieldVariant Scope, const FName& Name, EObjectFlags ObjectFlags, EVariableCategory::Type VariableCategory, const TCHAR* Dimensions)
 	{
 		const FPropertyBase& VarProperty = PropDef.GetPropertyBase();
+		UEnum* Enum = VarProperty.EnumDef ? VarProperty.EnumDef->GetEnum() : nullptr;
 #if UHT_ENABLE_VALUE_PROPERTY_TAG
-		PropDef.GetUnrealSourceFile().AddTypeDefIncludeIfNeeded(VarProperty.Enum);
+		PropDef.GetUnrealSourceFile().AddTypeDefIncludeIfNeeded(VarProperty.EnumDef);
 #endif
 		FByteProperty* Result = new FByteProperty(Scope, Name, ObjectFlags);
-		Result->Enum = VarProperty.Enum;
+		Result->Enum = Enum;
 		check(VarProperty.IntType == EIntType::Sized);
 		return Result;
 	}
@@ -610,7 +611,7 @@ struct FPropertyTypeTraitsEnum : public FPropertyTypeTraitsBase
 {
 	static bool DefaultValueStringCppFormatToInnerFormat(const FUnrealPropertyDefinitionInfo& PropDef, const FString& CppForm, FString& OutForm)
 	{
-		UEnum* Enum = PropDef.GetPropertyBase().Enum;
+		UEnum* Enum = PropDef.GetPropertyBase().EnumDef->GetEnum();
 		OutForm = FDefaultValueHelper::GetUnqualifiedEnumValue(FDefaultValueHelper::RemoveWhitespaces(CppForm));
 
 		const int32 EnumEntryIndex = Enum->GetIndexByName(*OutForm);
@@ -629,23 +630,23 @@ struct FPropertyTypeTraitsEnum : public FPropertyTypeTraitsBase
 	{
 		const FPropertyBase& VarProperty = PropDef.GetPropertyBase();
 
-		if (VarProperty.Enum->GetCppForm() != UEnum::ECppForm::EnumClass)
+		UEnum* Enum = VarProperty.EnumDef->GetEnum();
+
+		if (Enum->GetCppForm() != UEnum::ECppForm::EnumClass)
 		{
 			check(VarProperty.Type == EPropertyType::CPT_Byte);
 			return FPropertyTypeTraitsByte::CreateEngineType(PropDef, Scope, Name, ObjectFlags, VariableCategory, Dimensions);
 		}
 
-		FUnrealEnumDefinitionInfo& EnumDef = GTypeDefinitionInfoMap.FindChecked<FUnrealEnumDefinitionInfo>(VarProperty.Enum);
-
 #if UHT_ENABLE_VALUE_PROPERTY_TAG
-		PropDef.GetUnrealSourceFile().AddTypeDefIncludeIfNeeded(EnumDef);
+		PropDef.GetUnrealSourceFile().AddTypeDefIncludeIfNeeded(VarProperty.EnumDef);
 #endif
 
 		FPropertyBase UnderlyingProperty = VarProperty;
-		UnderlyingProperty.Enum = nullptr;
+		UnderlyingProperty.EnumDef = nullptr;
 		UnderlyingProperty.PropertyFlags = EPropertyFlags::CPF_None;
 		UnderlyingProperty.ArrayType = EArrayType::None;
-		switch (EnumDef.GetUnderlyingType())
+		switch (VarProperty.EnumDef->GetUnderlyingType())
 		{
 		case EUnderlyingEnumType::int8:        UnderlyingProperty.Type = CPT_Int8;   break;
 		case EUnderlyingEnumType::int16:       UnderlyingProperty.Type = CPT_Int16;  break;
@@ -668,7 +669,7 @@ struct FPropertyTypeTraitsEnum : public FPropertyTypeTraitsBase
 		FUnrealPropertyDefinitionInfo& SubProp = FPropertyTraits::CreateProperty(UnderlyingProperty, Result, TEXT("UnderlyingType"), ObjectFlags, VariableCategory, Dimensions, PropDef.GetUnrealSourceFile(), PropDef.GetLineNumber(), PropDef.GetParsePosition());
 		PropDef.SetValuePropDef(SubProp);
 		Result->UnderlyingProp = CastFieldChecked<FNumericProperty>(SubProp.GetProperty());
-		Result->Enum = VarProperty.Enum;
+		Result->Enum = Enum;
 		return Result;
 	}
 
@@ -707,29 +708,30 @@ struct FPropertyTypeTraitsObjectReference : public FPropertyTypeTraitsObjectBase
 	static FProperty* CreateEngineType(FUnrealPropertyDefinitionInfo& PropDef, FFieldVariant Scope, const FName& Name, EObjectFlags ObjectFlags, EVariableCategory::Type VariableCategory, const TCHAR* Dimensions)
 	{
 		FPropertyBase& VarProperty = PropDef.GetPropertyBase();
-		check(VarProperty.PropertyClass);
+		check(VarProperty.ClassDef);
+		UClass* PropertyClass = VarProperty.ClassDef->GetClass();
 
 #if UHT_ENABLE_PTR_PROPERTY_TAG
-		PropDef.GetUnrealSourceFile().AddTypeDefIncludeIfNeeded(VarProperty.PropertyClass);
+		PropDef.GetUnrealSourceFile().AddTypeDefIncludeIfNeeded(VarProperty.ClassDef);
 #endif
 
-		if (VarProperty.PropertyClass->IsChildOf(UClass::StaticClass()))
+		if (PropertyClass->IsChildOf(UClass::StaticClass()))
 		{
 			FClassProperty* Result = new FClassProperty(Scope, Name, ObjectFlags);
-			Result->MetaClass = VarProperty.MetaClass;
-			Result->PropertyClass = VarProperty.PropertyClass;
+			Result->MetaClass = VarProperty.MetaClassDef ? VarProperty.MetaClassDef->GetClass() : nullptr;
+			Result->PropertyClass = PropertyClass;
 			return Result;
 		}
 		else
 		{
-			if (FUnrealClassDefinitionInfo::HierarchyHasAnyClassFlags(VarProperty.PropertyClass, CLASS_DefaultToInstanced))
+			if (VarProperty.ClassDef->HierarchyHasAnyClassFlags(CLASS_DefaultToInstanced))
 			{
 				VarProperty.PropertyFlags |= CPF_InstancedReference;
 				AddEditInlineMetaData(VarProperty.MetaData);
 			}
 
 			FObjectProperty* Result = new FObjectProperty(Scope, Name, ObjectFlags);
-			Result->PropertyClass = VarProperty.PropertyClass;
+			Result->PropertyClass = PropertyClass;
 			return Result;
 		}
 	}
@@ -745,14 +747,15 @@ struct FPropertyTypeTraitsWeakObjectReference : public FPropertyTypeTraitsObject
 	static FProperty* CreateEngineType(FUnrealPropertyDefinitionInfo& PropDef, FFieldVariant Scope, const FName& Name, EObjectFlags ObjectFlags, EVariableCategory::Type VariableCategory, const TCHAR* Dimensions)
 	{
 		const FPropertyBase& VarProperty = PropDef.GetPropertyBase();
-		check(VarProperty.PropertyClass);
+		check(VarProperty.ClassDef);
+		UClass* PropertyClass = VarProperty.ClassDef->GetClass();
 
 #if UHT_ENABLE_PTR_PROPERTY_TAG
-		PropDef.GetUnrealSourceFile().AddTypeDefIncludeIfNeeded(VarProperty.PropertyClass);
+		PropDef.GetUnrealSourceFile().AddTypeDefIncludeIfNeeded(VarProperty.ClassDef);
 #endif
 
 		FWeakObjectProperty* Result = new FWeakObjectProperty(Scope, Name, ObjectFlags);
-		Result->PropertyClass = VarProperty.PropertyClass;
+		Result->PropertyClass = PropertyClass;
 		return Result;
 	}
 
@@ -767,14 +770,15 @@ struct FPropertyTypeTraitsLazyObjectReference : public FPropertyTypeTraitsObject
 	static FProperty* CreateEngineType(FUnrealPropertyDefinitionInfo& PropDef, FFieldVariant Scope, const FName& Name, EObjectFlags ObjectFlags, EVariableCategory::Type VariableCategory, const TCHAR* Dimensions)
 	{
 		const FPropertyBase& VarProperty = PropDef.GetPropertyBase();
-		check(VarProperty.PropertyClass);
+		check(VarProperty.ClassDef);
+		UClass* PropertyClass = VarProperty.ClassDef->GetClass();
 
 #if UHT_ENABLE_PTR_PROPERTY_TAG
-		PropDef.GetUnrealSourceFile().AddTypeDefIncludeIfNeeded(VarProperty.PropertyClass);
+		PropDef.GetUnrealSourceFile().AddTypeDefIncludeIfNeeded(VarProperty.ClassDef);
 #endif
 
 		FLazyObjectProperty* Result = new FLazyObjectProperty(Scope, Name, ObjectFlags);
-		Result->PropertyClass = VarProperty.PropertyClass;
+		Result->PropertyClass = PropertyClass;
 		return Result;
 	}
 };
@@ -784,29 +788,30 @@ struct FPropertyTypeTraitsObjectPtrReference : public FPropertyTypeTraitsObjectB
 	static FProperty* CreateEngineType(FUnrealPropertyDefinitionInfo& PropDef, FFieldVariant Scope, const FName& Name, EObjectFlags ObjectFlags, EVariableCategory::Type VariableCategory, const TCHAR* Dimensions)
 	{
 		FPropertyBase& VarProperty = PropDef.GetPropertyBase();
-		check(VarProperty.PropertyClass);
+		check(VarProperty.ClassDef);
+		UClass* PropertyClass = VarProperty.ClassDef->GetClass();
 
 #if UHT_ENABLE_PTR_PROPERTY_TAG
-		PropDef.GetUnrealSourceFile().AddTypeDefIncludeIfNeeded(VarProperty.PropertyClass);
+		PropDef.GetUnrealSourceFile().AddTypeDefIncludeIfNeeded(VarProperty.ClassDef);
 #endif
 
-		if (VarProperty.PropertyClass->IsChildOf(UClass::StaticClass()))
+		if (PropertyClass->IsChildOf(UClass::StaticClass()))
 		{
 			FClassPtrProperty* Result = new FClassPtrProperty(Scope, Name, ObjectFlags);
-			Result->MetaClass = VarProperty.MetaClass;
-			Result->PropertyClass = VarProperty.PropertyClass;
+			Result->MetaClass = VarProperty.MetaClassDef ? VarProperty.MetaClassDef->GetClass() : nullptr;
+			Result->PropertyClass = PropertyClass;
 			return Result;
 		}
 		else
 		{
-			if (FUnrealClassDefinitionInfo::HierarchyHasAnyClassFlags(VarProperty.PropertyClass, CLASS_DefaultToInstanced))
+			if (VarProperty.ClassDef->HierarchyHasAnyClassFlags(CLASS_DefaultToInstanced))
 			{
 				VarProperty.PropertyFlags |= CPF_InstancedReference;
 				AddEditInlineMetaData(VarProperty.MetaData);
 			}
 
 			FObjectPtrProperty* Result = new FObjectPtrProperty(Scope, Name, ObjectFlags);
-			Result->PropertyClass = VarProperty.PropertyClass;
+			Result->PropertyClass = PropertyClass;
 			return Result;
 		}
 	}
@@ -822,23 +827,24 @@ struct FPropertyTypeTraitsSoftObjectReference : public FPropertyTypeTraitsObject
 	static FProperty* CreateEngineType(FUnrealPropertyDefinitionInfo& PropDef, FFieldVariant Scope, const FName& Name, EObjectFlags ObjectFlags, EVariableCategory::Type VariableCategory, const TCHAR* Dimensions)
 	{
 		FPropertyBase& VarProperty = PropDef.GetPropertyBase();
-		check(VarProperty.PropertyClass);
+		check(VarProperty.ClassDef);
+		UClass* PropertyClass = VarProperty.ClassDef->GetClass();
 
 #if UHT_ENABLE_PTR_PROPERTY_TAG
-		PropDef.GetUnrealSourceFile().AddTypeDefIncludeIfNeeded(VarProperty.PropertyClass);
+		PropDef.GetUnrealSourceFile().AddTypeDefIncludeIfNeeded(VarProperty.ClassDef);
 #endif
 
-		if (VarProperty.PropertyClass->IsChildOf(UClass::StaticClass()))
+		if (PropertyClass->IsChildOf(UClass::StaticClass()))
 		{
 			FSoftClassProperty* Result = new FSoftClassProperty(Scope, Name, ObjectFlags);
-			Result->MetaClass = VarProperty.MetaClass;
-			Result->PropertyClass = VarProperty.PropertyClass;
+			Result->MetaClass = VarProperty.MetaClassDef ? VarProperty.MetaClassDef->GetClass() : nullptr;
+			Result->PropertyClass = PropertyClass;
 			return Result;
 		}
 		else
 		{
 			FSoftObjectProperty* Result = new FSoftObjectProperty(Scope, Name, ObjectFlags);
-			Result->PropertyClass = VarProperty.PropertyClass;
+			Result->PropertyClass = PropertyClass;
 			return Result;
 		}
 	}
@@ -854,15 +860,16 @@ struct FPropertyTypeTraitsInterface : public FPropertyTypeTraitsBase
 	static FProperty* CreateEngineType(FUnrealPropertyDefinitionInfo& PropDef, FFieldVariant Scope, const FName& Name, EObjectFlags ObjectFlags, EVariableCategory::Type VariableCategory, const TCHAR* Dimensions)
 	{
 		FPropertyBase& VarProperty = PropDef.GetPropertyBase();
-		check(VarProperty.PropertyClass);
-		check(VarProperty.PropertyClass->HasAnyClassFlags(CLASS_Interface));
+		check(VarProperty.ClassDef);
+		UClass* PropertyClass = VarProperty.ClassDef->GetClass();
+		check(PropertyClass->HasAnyClassFlags(CLASS_Interface));
 
 #if UHT_ENABLE_PTR_PROPERTY_TAG
-		PropDef.GetUnrealSourceFile().AddTypeDefIncludeIfNeeded(VarProperty.PropertyClass);
+		PropDef.GetUnrealSourceFile().AddTypeDefIncludeIfNeeded(VarProperty.ClassDef);
 #endif
 
 		FInterfaceProperty* Result = new  FInterfaceProperty(Scope, Name, ObjectFlags);
-		Result->InterfaceClass = VarProperty.PropertyClass;
+		Result->InterfaceClass = PropertyClass;
 		return Result;
 	}
 
@@ -989,7 +996,7 @@ struct FPropertyTypeTraitsStruct : public FPropertyTypeTraitsBase
 		static const UScriptStruct* LinearColorStruct = FClasses::FindObjectChecked<UScriptStruct>(CoreUObjectPackage, TEXT("LinearColor"));
 		static const UScriptStruct* ColorStruct = FClasses::FindObjectChecked<UScriptStruct>(CoreUObjectPackage, TEXT("Color"));
 
-		UScriptStruct* Struct = PropDef.GetPropertyBase().Struct;
+		UScriptStruct* Struct = PropDef.GetPropertyBase().ScriptStructDef->GetScriptStruct();
 		if (Struct == VectorStruct)
 		{
 			FString Parameters;
@@ -1187,18 +1194,19 @@ struct FPropertyTypeTraitsStruct : public FPropertyTypeTraitsBase
 	static FProperty* CreateEngineType(FUnrealPropertyDefinitionInfo& PropDef, FFieldVariant Scope, const FName& Name, EObjectFlags ObjectFlags, EVariableCategory::Type VariableCategory, const TCHAR* Dimensions)
 	{
 		FPropertyBase& VarProperty = PropDef.GetPropertyBase();
+		UScriptStruct* Struct = PropDef.GetPropertyBase().ScriptStructDef->GetScriptStruct();
 
 #if UHT_ENABLE_VALUE_PROPERTY_TAG
-		PropDef.GetUnrealSourceFile().AddTypeDefIncludeIfNeeded(VarProperty.Struct);
+		PropDef.GetUnrealSourceFile().AddTypeDefIncludeIfNeeded(VarProperty.ScriptStructDef);
 #endif
 
-		if (VarProperty.Struct->StructFlags & STRUCT_HasInstancedReference)
+		if (Struct->StructFlags & STRUCT_HasInstancedReference)
 		{
 			VarProperty.PropertyFlags |= CPF_ContainsInstancedReference;
 		}
 
 		FStructProperty* Result = new FStructProperty(Scope, Name, ObjectFlags);
-		Result->Struct = VarProperty.Struct;
+		Result->Struct = Struct;
 		return Result;
 	}
 
@@ -1227,9 +1235,10 @@ struct FPropertyTypeTraitsMulticastDelegate : public FPropertyTypeTraitsBase
 	static FProperty* CreateEngineType(FUnrealPropertyDefinitionInfo& PropDef, FFieldVariant Scope, const FName& Name, EObjectFlags ObjectFlags, EVariableCategory::Type VariableCategory, const TCHAR* Dimensions)
 	{
 		const FPropertyBase& VarProperty = PropDef.GetPropertyBase();
+		UFunction* Function = PropDef.GetPropertyBase().FunctionDef->GetFunction();
 
 		FMulticastDelegateProperty* Result;
-		if (VarProperty.Function->IsA<USparseDelegateFunction>())
+		if (Function->IsA<USparseDelegateFunction>())
 		{
 			Result = new FMulticastSparseDelegateProperty(Scope, Name, ObjectFlags);
 		}
diff --git a/Engine/Source/Programs/UnrealHeaderTool/Private/Scope.cpp b/Engine/Source/Programs/UnrealHeaderTool/Private/Scope.cpp
--- a/Engine/Source/Programs/UnrealHeaderTool/Private/Scope.cpp
+++ b/Engine/Source/Programs/UnrealHeaderTool/Private/Scope.cpp
@@ -30,8 +30,9 @@ void FScope::AddType(FUnrealFieldDefinitionInfo& Type)
  * @param Structs (Output parameter) Array to fill with structs.
  * @param DelegateFunctions (Output parameter) Array to fill with delegate functions.
  */
-void DispatchType(UField* Type, TArray<UEnum*> &Enums, TArray<UScriptStruct*> &Structs, TArray<UDelegateFunction*> &DelegateFunctions)
+void DispatchType(FUnrealFieldDefinitionInfo& FieldDef, TArray<FUnrealEnumDefinitionInfo*>& Enums, TArray<FUnrealScriptStructDefinitionInfo*>& Structs, TArray<FUnrealFunctionDefinitionInfo*>& DelegateFunctions)
 {
+	UField* Type = FieldDef.GetField();
 	UClass* TypeClass = Type->GetClass();
 
 	if (TypeClass == UClass::StaticClass() || TypeClass == UStruct::StaticClass())
@@ -41,13 +42,11 @@ void DispatchType(UField* Type, TArray<UEnum*> &Enums, TArray<UScriptStruct*> &S
 	}
 	else if (TypeClass == UEnum::StaticClass())
 	{
-		UEnum* Enum = (UEnum*)Type;
-		Enums.Add(Enum);
+		Enums.Add(&UHTCastChecked<FUnrealEnumDefinitionInfo>(FieldDef));
 	}
 	else if (TypeClass == UScriptStruct::StaticClass())
 	{
-		UScriptStruct* Struct = (UScriptStruct*)Type;
-		Structs.Add(Struct);
+		Structs.Add(&UHTCastChecked<FUnrealScriptStructDefinitionInfo>(FieldDef));
 	}
 	else if (TypeClass == UDelegateFunction::StaticClass() || TypeClass == USparseDelegateFunction::StaticClass())
 	{
@@ -56,7 +55,7 @@ void DispatchType(UField* Type, TArray<UEnum*> &Enums, TArray<UScriptStruct*> &S
 
 		if (Function->GetSuperFunction() == NULL)
 		{
-			DelegateFunctions.Add(Function);
+			DelegateFunctions.Add(&UHTCastChecked<FUnrealFunctionDefinitionInfo>(FieldDef));
 			bAdded = true;
 		}
 
@@ -64,12 +63,11 @@ void DispatchType(UField* Type, TArray<UEnum*> &Enums, TArray<UScriptStruct*> &S
 	}
 }
 
-void FScope::SplitTypesIntoArrays(TArray<UEnum*>& Enums, TArray<UScriptStruct*>& Structs, TArray<UDelegateFunction*>& DelegateFunctions)
+void FScope::SplitTypesIntoArrays(TArray<FUnrealEnumDefinitionInfo*>& Enums, TArray<FUnrealScriptStructDefinitionInfo*>& Structs, TArray<FUnrealFunctionDefinitionInfo*>& DelegateFunctions)
 {
 	for (TPair<FName, FUnrealFieldDefinitionInfo*>& TypePair : TypeMap)
 	{
-		UField* Type = TypePair.Value->GetField();
-		DispatchType(Type, Enums, Structs, DelegateFunctions);
+		DispatchType(*TypePair.Value, Enums, Structs, DelegateFunctions);
 	}
 }
 
diff --git a/Engine/Source/Programs/UnrealHeaderTool/Private/Scope.h b/Engine/Source/Programs/UnrealHeaderTool/Private/Scope.h
--- a/Engine/Source/Programs/UnrealHeaderTool/Private/Scope.h
+++ b/Engine/Source/Programs/UnrealHeaderTool/Private/Scope.h
@@ -139,7 +139,7 @@ public:
 	 * @param Structs (Output parameter) List of structs from this scope.
 	 * @param DelegateFunctions (Output parameter) List of delegate properties from this scope.
 	 */
-	void SplitTypesIntoArrays(TArray<UEnum*>& Enums, TArray<UScriptStruct*>& Structs, TArray<UDelegateFunction*>& DelegateFunctions);
+	void SplitTypesIntoArrays(TArray<FUnrealEnumDefinitionInfo*>& Enums, TArray<FUnrealScriptStructDefinitionInfo*>& Structs, TArray<FUnrealFunctionDefinitionInfo*>& DelegateFunctions);
 
 	/**
 	 * Gets scope name.
diff --git a/Engine/Source/Programs/UnrealHeaderTool/Private/UnrealSourceFile.cpp b/Engine/Source/Programs/UnrealHeaderTool/Private/UnrealSourceFile.cpp
--- a/Engine/Source/Programs/UnrealHeaderTool/Private/UnrealSourceFile.cpp
+++ b/Engine/Source/Programs/UnrealHeaderTool/Private/UnrealSourceFile.cpp
@@ -193,23 +193,14 @@ void FUnrealSourceFile::AddScriptStructIncludeIfNeeded(int32 InputLine, const FS
 	}
 }
 
-void FUnrealSourceFile::AddTypeDefIncludeIfNeeded(FUnrealTypeDefinitionInfo& TypeDef)
+void FUnrealSourceFile::AddTypeDefIncludeIfNeeded(FUnrealTypeDefinitionInfo* TypeDef)
 {
-	check(TypeDef.HasSource());
-	if (&TypeDef.GetUnrealSourceFile() != this)
+	if (TypeDef != nullptr)
 	{
-		GetIncludes().AddUnique(FHeaderProvider(TypeDef));
-	}
-}
-
-void FUnrealSourceFile::AddTypeDefIncludeIfNeeded(UField* Field)
-{
-	if (Field != nullptr)
-	{
-		if (TSharedRef<FUnrealTypeDefinitionInfo>* TypeDef = GTypeDefinitionInfoMap.Find(Field))
+		check(TypeDef->HasSource());
+		if (&TypeDef->GetUnrealSourceFile() != this)
 		{
-			AddTypeDefIncludeIfNeeded(**TypeDef);
+			GetIncludes().AddUnique(FHeaderProvider(*TypeDef));
 		}
 	}
 }
-
diff --git a/Engine/Source/Programs/UnrealHeaderTool/Private/UnrealSourceFile.h b/Engine/Source/Programs/UnrealHeaderTool/Private/UnrealSourceFile.h
--- a/Engine/Source/Programs/UnrealHeaderTool/Private/UnrealSourceFile.h
+++ b/Engine/Source/Programs/UnrealHeaderTool/Private/UnrealSourceFile.h
@@ -279,9 +279,9 @@ public:
 	void AddScriptStructIncludeIfNeeded(int32 InputLine, const FString& StructNameWithoutPrefix, const FString& DependencyStructName);
 
 	/**
-	 * Add an include for a type definition if required
+	 * Add an include for a type definition if required.  TypeDef can be null.
 	 */
-	void AddTypeDefIncludeIfNeeded(FUnrealTypeDefinitionInfo& TypeDef);
+	void AddTypeDefIncludeIfNeeded(FUnrealTypeDefinitionInfo* TypeDef);
 
 	/**
 	 * Add an include for a type definition if required
@@ -430,7 +430,7 @@ public:
 	/**
 	 * Get the collection of singletons 
 	 */
-	TArray<UField*>& GetSingletons()
+	TArray<FUnrealFieldDefinitionInfo*>& GetSingletons()
 	{
 		return Singletons;
 	}
@@ -523,5 +523,5 @@ private:
 	TMap<UStruct*, EGeneratedCodeVersion> GeneratedCodeVersions;
 
 	// Collection of all singletons found during code generation
-	TArray<UField*> Singletons;
+	TArray<FUnrealFieldDefinitionInfo*> Singletons;
 };
diff --git a/Engine/Source/Programs/UnrealHeaderTool/Private/UnrealTypeDefinitionInfo.cpp b/Engine/Source/Programs/UnrealHeaderTool/Private/UnrealTypeDefinitionInfo.cpp
--- a/Engine/Source/Programs/UnrealHeaderTool/Private/UnrealTypeDefinitionInfo.cpp
+++ b/Engine/Source/Programs/UnrealHeaderTool/Private/UnrealTypeDefinitionInfo.cpp
@@ -522,24 +522,6 @@ void FUnrealClassDefinitionInfo::PostParseFinalize()
 	FUnrealStructDefinitionInfo::PostParseFinalize();
 }
 
-bool FUnrealClassDefinitionInfo::HierarchyHasAnyClassFlags(UClass* Class, EClassFlags FlagsToCheck)
-{
-	if (FUnrealClassDefinitionInfo* ClassDef = UHTCast<FUnrealClassDefinitionInfo>(GTypeDefinitionInfoMap.FindChecked(Class)))
-	{
-		return ClassDef->HierarchyHasAnyClassFlags(CLASS_DefaultToInstanced);
-	}
-	return false;
-}
-
-bool FUnrealClassDefinitionInfo::HierarchyHasAllClassFlags(UClass* Class, EClassFlags FlagsToCheck)
-{
-	if (FUnrealClassDefinitionInfo* ClassDef = UHTCast<FUnrealClassDefinitionInfo>(GTypeDefinitionInfoMap.FindChecked(Class)))
-	{
-		return ClassDef->HierarchyHasAllClassFlags(CLASS_DefaultToInstanced);
-	}
-	return false;
-}
-
 void FUnrealClassDefinitionInfo::ParseClassProperties(TArray<FPropertySpecifier>&& InClassSpecifiers, const FString& InRequiredAPIMacroIfPresent)
 {
 	ClassFlags = CLASS_None;
diff --git a/Engine/Source/Programs/UnrealHeaderTool/Private/UnrealTypeDefinitionInfo.h b/Engine/Source/Programs/UnrealHeaderTool/Private/UnrealTypeDefinitionInfo.h
--- a/Engine/Source/Programs/UnrealHeaderTool/Private/UnrealTypeDefinitionInfo.h
+++ b/Engine/Source/Programs/UnrealHeaderTool/Private/UnrealTypeDefinitionInfo.h
@@ -1300,16 +1300,6 @@ public:
 		return false;
 	}
 
-	/**
-	 * Used to safely check whether the passed in flag is set in the whole hierarchy
-	 *
-	 * @param	Class				Class to check
-	 * @param	FlagsToCheck		Class flag(s) to check for
-	 * @return	true if the passed in flag is set, false otherwise
-	 *			(including no flag passed in, unless the FlagsToCheck is CLASS_AllFlags)
-	 */
-	static bool HierarchyHasAnyClassFlags(UClass* Class, EClassFlags FlagsToCheck);
-
 	/**
 	 * Used to safely check whether all of the passed in flags are set.
 	 *
@@ -1327,14 +1317,6 @@ public:
 		}
 		return false;
 	}
-	/**
-	 * Used to safely check whether all of the passed in flags are set.
-	 *
-	 * @param	Class				Class to check
-	 * @param FlagsToCheck	Class flags to check for
-	 * @return true if all of the passed in flags are set (including no flags passed in), false otherwise
-	 */
-	static bool HierarchyHasAllClassFlags(UClass* Class, EClassFlags FlagsToCheck);
 
 	/**
 	* Parse Class's properties to generate its declaration data.
