Remove FClasses

#rnx

diff --git a/Engine/Source/Programs/UnrealHeaderTool/Private/ClassMaps.h b/Engine/Source/Programs/UnrealHeaderTool/Private/ClassMaps.h
--- a/Engine/Source/Programs/UnrealHeaderTool/Private/ClassMaps.h
+++ b/Engine/Source/Programs/UnrealHeaderTool/Private/ClassMaps.h
@@ -3,12 +3,11 @@
 #pragma once
 
 #include "CoreMinimal.h"
-#include "UObject/Stack.h"
-#include "UObject/ErrorException.h"
-
+#include "StringUtils.h"
 #include "UnrealSourceFile.h"
-#include "UObject/ErrorException.h"
 #include "UnrealTypeDefinitionInfo.h"
+#include "UObject/Stack.h"
+#include "UObject/ErrorException.h"
 
 class UField;
 class UClass;
@@ -94,29 +93,34 @@ struct FTypeDefinitionInfoMap : public FFreezableContainer
 		return *TypeDef;
 	}
 
-	TSharedRef<FUnrealTypeDefinitionInfo>* FindByName(const FName Name)
+	// Finding by name must be done on the stripped name for classes and script structs
+	TSharedRef<FUnrealTypeDefinitionInfo>* FindByName(const TCHAR* Name)
 	{
-		check(bFrozen);
-		return DefinitionsByName.Find(Name);
+		FName SearchName(Name, EFindName::FNAME_Find);
+		if (SearchName != NAME_None)
+		{
+			return DefinitionsByName.Find(SearchName);
+		}
+		return nullptr;
 	}
 
 	template<typename To>
-	To* FindByName(const FName Name)
+	To* FindByName(const TCHAR* Name)
 	{
 		check(bFrozen);
-		return UHTCast<To>(DefinitionsByName.Find(Name));
+		return UHTCast<To>(FindByName(Name));
 	}
 
-	FUnrealTypeDefinitionInfo& FindByNameChecked(const FName Name)
+	FUnrealTypeDefinitionInfo& FindByNameChecked(const TCHAR* Name)
 	{
 		check(bFrozen);
-		TSharedRef<FUnrealTypeDefinitionInfo>* TypeDef = DefinitionsByName.Find(Name);
+		TSharedRef<FUnrealTypeDefinitionInfo>* TypeDef = FindByName(Name);
 		check(TypeDef);
 		return **TypeDef;
 	}
 
 	template <typename To>
-	To& FindByNameChecked(const FName Name)
+	To& FindByNameChecked(const TCHAR* Name)
 	{
 		check(bFrozen);
 		To* TypeDef = FindByName<To>(Name);
diff --git a/Engine/Source/Programs/UnrealHeaderTool/Private/Classes.cpp b/Engine/Source/Programs/UnrealHeaderTool/Private/Classes.cpp
deleted file mode 100644
diff --git a/Engine/Source/Programs/UnrealHeaderTool/Private/Classes.h b/Engine/Source/Programs/UnrealHeaderTool/Private/Classes.h
deleted file mode 100644
diff --git a/Engine/Source/Programs/UnrealHeaderTool/Private/CodeGenerator.cpp b/Engine/Source/Programs/UnrealHeaderTool/Private/CodeGenerator.cpp
--- a/Engine/Source/Programs/UnrealHeaderTool/Private/CodeGenerator.cpp
+++ b/Engine/Source/Programs/UnrealHeaderTool/Private/CodeGenerator.cpp
@@ -42,7 +42,7 @@
 #include "GeneratedCodeVersion.h"
 #include "UnrealSourceFile.h"
 #include "ParserHelper.h"
-#include "Classes.h"
+#include "EngineAPI.h"
 #include "ClassMaps.h"
 #include "NativeClassExporter.h"
 #include "ProfilingDebugging/ScopedTimers.h"
@@ -86,6 +86,11 @@ static FGraphEventArray GAsyncFileTasks;
 
 bool HasIdentifierExactMatch(const TCHAR* StringBegin, const TCHAR* StringEnd, const FString& Find);
 
+// Globals for common class definitions
+FUnrealClassDefinitionInfo* GUObjectDef = nullptr;
+FUnrealClassDefinitionInfo* GUClassDef = nullptr;
+FUnrealClassDefinitionInfo* GUInterfaceDef = nullptr;
+
 namespace
 {
 	static const FName NAME_SerializeToFArchive("SerializeToFArchive");
@@ -2336,17 +2341,15 @@ void FNativeClassHeaderGenerator::ExportNativeGeneratedInitCode(FOutputDevice& O
 
 	for (const FString& SparseClassDataString : SparseClassDataTypes)
 	{
-		UScriptStruct* SparseClassDataStruct = FClasses::FindObject<UScriptStruct>(ANY_PACKAGE, *SparseClassDataString);
-		if (!SparseClassDataStruct)
-		{
-			continue;
-		}
-		FUnrealScriptStructDefinitionInfo& SparseScriptStructDef = GTypeDefinitionInfoMap.FindChecked<FUnrealScriptStructDefinitionInfo>(SparseClassDataStruct);
-		GeneratedClassRegisterFunctionText.Logf(TEXT("%s\r\n"), *SparseClassDataStruct->GetName());
-		for (FUnrealPropertyDefinitionInfo* ChildDef : TUHTFieldRange<FUnrealPropertyDefinitionInfo>(SparseScriptStructDef))
+		if (FUnrealScriptStructDefinitionInfo* SparseScriptStructDef = GTypeDefinitionInfoMap.FindByName<FUnrealScriptStructDefinitionInfo>(*SparseClassDataString))
 		{
-			FProperty* Child = ChildDef->GetProperty();
-			GeneratedClassRegisterFunctionText.Logf(TEXT("%s %s\r\n"), *Child->GetCPPType(), *Child->GetNameCPP());
+			UScriptStruct* SparseClassDataStruct = SparseScriptStructDef->GetScriptStruct();
+			GeneratedClassRegisterFunctionText.Logf(TEXT("%s\r\n"), *SparseClassDataStruct->GetName());
+			for (FUnrealPropertyDefinitionInfo* ChildDef : TUHTFieldRange<FUnrealPropertyDefinitionInfo>(*SparseScriptStructDef))
+			{
+				FProperty* Child = ChildDef->GetProperty();
+				GeneratedClassRegisterFunctionText.Logf(TEXT("%s %s\r\n"), *Child->GetCPPType(), *Child->GetNameCPP());
+			}
 		}
 	}
 
@@ -2797,6 +2800,9 @@ void FNativeClassHeaderGenerator::ExportClassFromSourceFileInner(
 	UClass* Class = ClassDef.GetClass();
 	FStructMetaData& StructData = ClassDef.GetStructMetaData();
 
+	FUnrealClassDefinitionInfo* SuperClassDef = ClassDef.GetSuperClass();
+	UClass* SuperClass = SuperClassDef ? SuperClassDef->GetClass() : nullptr;
+
 	// C++ -> VM stubs (native function execs)
 	FUHTStringBuilder ClassMacroCalls;
 	FUHTStringBuilder ClassNoPureDeclsMacroCalls;
@@ -2857,8 +2863,6 @@ void FNativeClassHeaderGenerator::ExportClassFromSourceFileInner(
 	FUHTStringBuilder FriendText;
 	ExportNativeGeneratedInitCode(OutCpp, OutDeclarations, OutReferenceGatherers, SourceFile, ClassDef, FriendText);
 
-	UClass* SuperClass = Class->GetSuperClass();
-
 	// the name for the C++ version of the UClass
 	const FString ClassCPPName = FNameLookupCPP::GetNameCPP(Class);
 	const FString SuperClassCPPName = (SuperClass ? FNameLookupCPP::GetNameCPP(SuperClass) : TEXT("None"));
@@ -2950,7 +2954,7 @@ void FNativeClassHeaderGenerator::ExportClassFromSourceFileInner(
 			Boilerplate.Logf(TEXT("\t%s\r\n"), *GeneratedSerializeFunctionHeaderMacroName);
 		}
 
-		if (SuperClass && Class->ClassWithin != SuperClass->ClassWithin)
+		if (SuperClassDef && ClassDef.GetClassWithin() != SuperClassDef->GetClassWithin())
 		{
 			Boilerplate.Logf(TEXT("\tDECLARE_WITHIN(%s)\r\n"), *FNameLookupCPP::GetNameCPP(Class->ClassWithin));
 		}
@@ -5170,11 +5174,10 @@ void FNativeClassHeaderGenerator::ExportNativeFunctions(FOutputDevice& OutGenera
 		RuntimeStringBuilders.AutogeneratedStaticData.Logf(TEXT("\treturn (F%s*)(GetClass()->GetOrCreateSparseClassData());\r\n"), *SparseClassDataString);
 		RuntimeStringBuilders.AutogeneratedStaticData += TEXT("}\r\n");
 
-		UScriptStruct* SparseClassDataStruct = FClasses::FindObject<UScriptStruct>(ANY_PACKAGE, *SparseClassDataString);
-		while (SparseClassDataStruct != nullptr)
+		FUnrealScriptStructDefinitionInfo* SparseClassDataStructDef = GTypeDefinitionInfoMap.FindByName<FUnrealScriptStructDefinitionInfo>(*SparseClassDataString);
+		while (SparseClassDataStructDef != nullptr)
 		{
-			FUnrealScriptStructDefinitionInfo& SparseClassDataStructDef = GTypeDefinitionInfoMap.FindChecked<FUnrealScriptStructDefinitionInfo>(SparseClassDataStruct);
-			for (FUnrealPropertyDefinitionInfo* PropertyDef : SparseClassDataStructDef.GetProperties())
+			for (FUnrealPropertyDefinitionInfo* PropertyDef : SparseClassDataStructDef->GetProperties())
 			{
 				const FProperty* Child = PropertyDef->GetProperty();
 				FString ReturnExtendedType;
@@ -5218,7 +5221,7 @@ void FNativeClassHeaderGenerator::ExportNativeFunctions(FOutputDevice& OutGenera
 				}
 			}
 
-			SparseClassDataStruct = Cast<UScriptStruct>(SparseClassDataStruct->GetSuperStruct());
+			SparseClassDataStructDef = UHTCast<FUnrealScriptStructDefinitionInfo>(SparseClassDataStructDef->GetSuperClassInfo().Struct);
 		}
 	}
 
@@ -5610,17 +5613,11 @@ bool FNativeClassHeaderGenerator::WriteSource(const FManifestModule& Module, FGe
 		for (const TSharedRef<FUnrealTypeDefinitionInfo>& TypeDef : InSourceFile->GetDefinedClasses())
 		{
 			FUnrealClassDefinitionInfo& ClassDef = TypeDef->AsClassChecked();
-			UClass* Class = ClassDef.GetClass();
-			if (Class->ClassWithin && Class->ClassWithin != UObject::StaticClass())
+			FUnrealClassDefinitionInfo* ClassWithin = ClassDef.GetClassWithin();
+			if (ClassWithin && ClassWithin != GUObjectDef && ClassWithin->HasSource())
 			{
-				if (TSharedRef<FUnrealTypeDefinitionInfo>* WithinTypeInfo = GTypeDefinitionInfoMap.Find(Class->ClassWithin))
-				{
-					if ((*WithinTypeInfo)->HasSource())
-					{
-						FString Header = GetBuildPath((*WithinTypeInfo)->GetUnrealSourceFile());
-						RelativeIncludes.AddUnique(MoveTemp(Header));
-					}
-				}
+				FString Header = GetBuildPath(ClassWithin->GetUnrealSourceFile());
+				RelativeIncludes.AddUnique(MoveTemp(Header));
 			}
 
 			if (!bAddedStructuredArchiveFromArchiveHeader && ClassDef.GetArchiveType() == ESerializerArchiveType::StructuredArchiveRecord)
@@ -6128,7 +6125,7 @@ void FNativeClassHeaderGenerator::DeleteUnusedGeneratedHeaders(TSet<FString>&& P
 								FString BaseFilename = FPaths::GetBaseFilename(Filename);
 								const int32   GeneratedIndex = BaseFilename.Find(TEXT(".generated"), ESearchCase::IgnoreCase, ESearchDir::FromEnd);
 								const FString ClassName = MoveTemp(BaseFilename).Mid(0, GeneratedIndex);
-								UClass* IntrinsicClass = FClasses::FindObject<UClass>(ANY_PACKAGE, *ClassName);
+								UClass* IntrinsicClass = FEngineAPI::FindObject<UClass>(ANY_PACKAGE, *ClassName);
 								if (!IntrinsicClass || !IntrinsicClass->HasAnyClassFlags(CLASS_Intrinsic))
 								{
 									IFileManager::Get().Delete(*Fullpath);
@@ -6385,13 +6382,19 @@ void ResolveSuperClasses(UPackage* Package)
 
 		UClass* DefinedClass = Cast<UClass>(Object);
 
+		FUnrealClassDefinitionInfo& ClassDef = GTypeDefinitionInfoMap.FindChecked<FUnrealClassDefinitionInfo>(DefinedClass);
+
+		// Propagate any already set ClassWithin
+		if (DefinedClass->ClassWithin != nullptr)
+		{
+			ClassDef.SetClassWithin(&GTypeDefinitionInfoMap.FindChecked<FUnrealClassDefinitionInfo>(DefinedClass->ClassWithin));
+		}
+
 		if (DefinedClass->HasAnyClassFlags(CLASS_Intrinsic | CLASS_NoExport))
 		{
 			continue;
 		}
 
-		FUnrealClassDefinitionInfo& ClassDef = GTypeDefinitionInfoMap.FindChecked<FUnrealClassDefinitionInfo>(DefinedClass);
-
 		// Resolve the base class
 		{
 			FUnrealStructDefinitionInfo::FBaseClassInfo& SuperClassInfo = ClassDef.GetSuperClassInfo();
@@ -6400,22 +6403,21 @@ void ResolveSuperClasses(UPackage* Package)
 
 			if (!BaseClassNameStripped.IsEmpty() && !DefinedClass->GetSuperClass())
 			{
-				UClass* FoundBaseClass = FClasses::FindObject<UClass>(Package, *BaseClassNameStripped);
-
-				if (FoundBaseClass == nullptr)
+				FUnrealClassDefinitionInfo* FoundBaseClassDef = GTypeDefinitionInfoMap.FindByName<FUnrealClassDefinitionInfo>(*BaseClassNameStripped);
+				if (FoundBaseClassDef == nullptr)
 				{
-					FoundBaseClass = FClasses::FindObject<UClass>(ANY_PACKAGE, *BaseClassNameStripped);
+					FoundBaseClassDef = GTypeDefinitionInfoMap.FindByName<FUnrealClassDefinitionInfo>(*BaseClassName);
 				}
 
-				if (FoundBaseClass == nullptr)
+				if (FoundBaseClassDef == nullptr)
 				{
 					// Don't know its parent class. Raise error.
 					FError::Throwf(TEXT("Couldn't find parent type for '%s' named '%s' in current module (Package: %s) or any other module parsed so far."), *DefinedClass->GetName(), *BaseClassName, *GetNameSafe(Package));
 				}
 
-				DefinedClass->SetSuperStruct(FoundBaseClass);
-				DefinedClass->ClassCastFlags |= FoundBaseClass->ClassCastFlags;
-				SuperClassInfo.Struct = &GTypeDefinitionInfoMap.FindChecked<FUnrealStructDefinitionInfo>(FoundBaseClass);
+				DefinedClass->SetSuperStruct(FoundBaseClassDef->GetClass());
+				DefinedClass->ClassCastFlags |= FoundBaseClassDef->GetClass()->ClassCastFlags;
+				SuperClassInfo.Struct = FoundBaseClassDef;
 			}
 		}
 
@@ -6423,10 +6425,7 @@ void ResolveSuperClasses(UPackage* Package)
 		{
 			for (FUnrealStructDefinitionInfo::FBaseClassInfo& BaseClassInfo : ClassDef.GetBaseClassInfo())
 			{
-				if (UClass* FoundClass = FClasses::FindScriptClass(BaseClassInfo.Name))
-				{
-					BaseClassInfo.Struct = &GTypeDefinitionInfoMap.FindChecked<FUnrealStructDefinitionInfo>(FoundClass);
-				}
+				BaseClassInfo.Struct = FUnrealClassDefinitionInfo::FindScriptClass(BaseClassInfo.Name);
 			}
 		}
 	}
@@ -6434,7 +6433,7 @@ void ResolveSuperClasses(UPackage* Package)
 
 UPackage* GetModulePackage(FManifestModule& Module)
 {
-	UPackage* Package = FClasses::FindObjectFast<UPackage>(NULL, FName(*Module.LongPackageName), false, false);
+	UPackage* Package = FEngineAPI::FindObjectFast<UPackage>(NULL, FName(*Module.LongPackageName), false, false);
 	if (Package == NULL)
 	{
 		Package = CreatePackage(*Module.LongPackageName);
@@ -6720,6 +6719,10 @@ void DefineTypes(TArray<FUnrealPackageDefinitionInfo*>& PackageDefs)
 
 void ResolveParents(TArray<FUnrealPackageDefinitionInfo*>& PackageDefs)
 {
+	GUObjectDef = &GTypeDefinitionInfoMap.FindChecked<FUnrealClassDefinitionInfo>(UObject::StaticClass());
+	GUClassDef = &GTypeDefinitionInfoMap.FindChecked<FUnrealClassDefinitionInfo>(UClass::StaticClass());
+	GUInterfaceDef = &GTypeDefinitionInfoMap.FindChecked<FUnrealClassDefinitionInfo>(UInterface::StaticClass());
+
 	for (FUnrealPackageDefinitionInfo* PackageDef : PackageDefs)
 	{
 		FResults::Try([PackageDef]() { ResolveSuperClasses(PackageDef->GetPackage()); });
@@ -7118,7 +7121,7 @@ ECompilationResult::Type UnrealHeaderTool_Main(const FString& ModuleInfoFilename
 	UE_LOG(LogCompile, Log, TEXT("Preparing types for export took %.3f seconds"), TotalPrepareTypesForExportTime);
 	UE_LOG(LogCompile, Log, TEXT("Code generation took %.3f seconds"), TotalCodeGenTime);
 	UE_LOG(LogCompile, Log, TEXT("ScriptPlugin overhead was %.3f seconds"), TotalPluginTime);
-	UE_LOG(LogCompile, Log, TEXT("Sumary generation took %.3f seconds"), TotalSummaryTime);
+	UE_LOG(LogCompile, Log, TEXT("Summary generation took %.3f seconds"), TotalSummaryTime);
 	UE_LOG(LogCompile, Log, TEXT("Macroize time was %.3f seconds"), GMacroizeTime);
 
 	FUnrealHeaderToolStats& Stats = FUnrealHeaderToolStats::Get();
@@ -7197,7 +7200,7 @@ UClass* ProcessParsedClass(bool bClassIsAnInterface, const FString& ClassName, c
 		FError::Throwf(TEXT("Class '%s' cannot inherit from itself"), *ClassName);
 	}
 
-	UClass* ResultClass = FClasses::FindObject<UClass>(InParent, *ClassNameStripped);
+	UClass* ResultClass = FEngineAPI::FindObject<UClass>(InParent, *ClassNameStripped);
 
 	// if we aren't generating headers, then we shouldn't set misaligned object, since it won't get cleared
 
@@ -7208,7 +7211,7 @@ UClass* ProcessParsedClass(bool bClassIsAnInterface, const FString& ClassName, c
 		// detect if the same class name is used in multiple packages
 		if (ResultClass == nullptr)
 		{
-			UClass* ConflictingClass = FClasses::FindObject<UClass>(ANY_PACKAGE, *ClassNameStripped, true);
+			UClass* ConflictingClass = FEngineAPI::FindObject<UClass>(ANY_PACKAGE, *ClassNameStripped, true);
 			if (ConflictingClass != nullptr)
 			{
 				UE_LOG_WARNING_UHT(TEXT("Duplicate class name: %s also exists in file %s"), *ClassName, *ConflictingClass->GetOutermost()->GetName());
@@ -7237,7 +7240,7 @@ UClass* ProcessParsedClass(bool bClassIsAnInterface, const FString& ClassName, c
 
 UEnum* ProcessParsedEnum(const FString& EnumName, UObject* InParent, EObjectFlags Flags)
 {
-	if (UEnum* Existing = FClasses::FindObject<UEnum>(ANY_PACKAGE, *EnumName))
+	if (UEnum* Existing = FEngineAPI::FindObject<UEnum>(ANY_PACKAGE, *EnumName))
 	{
 		FError::Throwf(TEXT("Duplicate enum name: %s also exists in file %s"), *EnumName, *Existing->GetOutermost()->GetName());
 	}
@@ -7257,7 +7260,7 @@ UScriptStruct* ProcessParsedStruct(const FString& StructName, UObject* InParent,
 {
 	FString StructNameStripped = GetClassNameWithPrefixRemoved(*StructName);
 
-	if (UScriptStruct* Existing = FClasses::FindObject<UScriptStruct>(ANY_PACKAGE, *StructNameStripped))
+	if (UScriptStruct* Existing = FEngineAPI::FindObject<UScriptStruct>(ANY_PACKAGE, *StructNameStripped))
 	{
 		FError::Throwf(TEXT("Duplicate struct name: %s also exists in file %s"), *StructNameStripped, *Existing->GetOutermost()->GetName());
 	}
diff --git a/Engine/Source/Programs/UnrealHeaderTool/Private/EngineAPI.cpp b/Engine/Source/Programs/UnrealHeaderTool/Private/EngineAPI.cpp
new file mode 100644
--- /dev/null
+++ b/Engine/Source/Programs/UnrealHeaderTool/Private/EngineAPI.cpp
@@ -0,0 +1,3 @@
+// Copyright Epic Games, Inc. All Rights Reserved.
+
+#include "EngineAPI.h"
diff --git a/Engine/Source/Programs/UnrealHeaderTool/Private/EngineAPI.h b/Engine/Source/Programs/UnrealHeaderTool/Private/EngineAPI.h
new file mode 100644
--- /dev/null
+++ b/Engine/Source/Programs/UnrealHeaderTool/Private/EngineAPI.h
@@ -0,0 +1,33 @@
+// Copyright Epic Games, Inc. All Rights Reserved.
+
+#pragma once
+
+#include "CoreMinimal.h"
+#include "UObject\ObjectMacros.h"
+
+/**
+ * Wrapper class for access to the engine API.
+ */
+class FEngineAPI
+{
+public:
+	/**
+	 * Find an optional object.
+	 * @see StaticFindObject()
+	 */
+	template< class T >
+	static inline T* FindObject(UObject* Outer, const TCHAR* Name, bool ExactClass = false)
+	{
+		return (T*)StaticFindObject(T::StaticClass(), Outer, Name, ExactClass);
+	}
+
+	/**
+	 * Find an optional object, relies on the name being unqualified
+	 * @see StaticFindObjectFast()
+	 */
+	template< class T >
+	static inline T* FindObjectFast(UObject* Outer, FName Name, bool ExactClass = false, bool AnyPackage = false, EObjectFlags ExclusiveFlags = RF_NoFlags)
+	{
+		return (T*)StaticFindObjectFast(T::StaticClass(), Outer, Name, ExactClass, AnyPackage, ExclusiveFlags);
+	}
+};
diff --git a/Engine/Source/Programs/UnrealHeaderTool/Private/Exceptions.cpp b/Engine/Source/Programs/UnrealHeaderTool/Private/Exceptions.cpp
--- a/Engine/Source/Programs/UnrealHeaderTool/Private/Exceptions.cpp
+++ b/Engine/Source/Programs/UnrealHeaderTool/Private/Exceptions.cpp
@@ -40,21 +40,6 @@ FUHTException::FUHTException(ECompilationResult::Type InResult, FUnrealTypeDefin
 {
 }
 
-FUHTException::FUHTException(ECompilationResult::Type InResult, const UObject* InObject, FString&& InMessage)
-	: Result(InResult)
-	, Message(MoveTemp(InMessage))
-	, Line(1)
-{
-	if (const UField* Field = Cast<UField>(InObject))
-	{
-		if (TSharedRef<FUnrealTypeDefinitionInfo>* TypeDef = GTypeDefinitionInfoMap.Find(Field))
-		{
-			Filename = (*TypeDef)->GetUnrealSourceFile().GetFilename();
-			Line = (*TypeDef)->GetLineNumber();
-		}
-	}
-}
-
 namespace UE::UnrealHeaderTool::Exceptions::Private
 {
 	std::atomic<ECompilationResult::Type> OverallResults = ECompilationResult::Succeeded;
@@ -149,6 +134,15 @@ void FResults::LogError(const UObject& Object, const TCHAR* ErrorMsg, ECompilati
 	LogError(ErrorMsg, InResult);
 }
 
+void FResults::LogError(FUnrealTypeDefinitionInfo& InTypeDef, const TCHAR* ErrorMsg, ECompilationResult::Type InResult/* = ECompilationResult::OtherCompilationError*/)
+{
+	if (InTypeDef.HasSource())
+	{
+		LogError(InTypeDef.GetUnrealSourceFile(), InTypeDef.GetLineNumber(), ErrorMsg, InResult);
+	}
+	LogError(ErrorMsg, InResult);
+}
+
 void FResults::LogError(const TCHAR* ErrorMsg, ECompilationResult::Type InResult/* = ECompilationResult::OtherCompilationError*/)
 {
 	LogError(TEXT(""), 1, ErrorMsg, InResult);
diff --git a/Engine/Source/Programs/UnrealHeaderTool/Private/Exceptions.h b/Engine/Source/Programs/UnrealHeaderTool/Private/Exceptions.h
--- a/Engine/Source/Programs/UnrealHeaderTool/Private/Exceptions.h
+++ b/Engine/Source/Programs/UnrealHeaderTool/Private/Exceptions.h
@@ -99,33 +99,6 @@ public:
 		throw FUHTException(InResult, InTypeDef, MoveTemp(ResultString));
 	}
 
-	/**
-	 * Generate an exception
-	 * @param InObject The type object generating the exception.  The filename and line number will be retrieved from the object's type definition if possible.
-	 * @param InFmt The format string
-	 * @param InArgs Arguments supplied to the format string
-	 */
-	template <typename FmtType, typename... Types>
-	UE_NORETURN static void VARARGS Throwf(const UObject* InObject, const FmtType& InFmt, Types... InArgs)
-	{
-		FString ResultString = FString::Printf(InFmt, InArgs ...);
-		throw FUHTException(ECompilationResult::OtherCompilationError, InObject, MoveTemp(ResultString));
-	}
-
-	/**
-	 * Generate an exception
-	 * @param Result Result code to be returned as the overall result of the compilation process
-	 * @param Object The type object generating the exception.  The filename and line number will be retrieved from the object's type definition if possible.
-	 * @param Fmt The format string
-	 * @param Args Arguments supplied to the format string
-	 */
-	template <typename FmtType, typename... Types>
-	UE_NORETURN static void VARARGS Throwf(ECompilationResult::Type InResult, const UObject* InObject, const FmtType& InFmt, Types... InArgs)
-	{
-		FString ResultString = FString::Printf(InFmt, InArgs ...);
-		throw FUHTException(InResult, InObject, MoveTemp(ResultString));
-	}
-
 	/**
 	 * Return the result code of the exception
 	 */
@@ -242,6 +215,14 @@ struct FResults
 	 */
 	static void LogError(const UObject& Object, const TCHAR* ErrorMsg, ECompilationResult::Type Result = ECompilationResult::OtherCompilationError);
 
+	/**
+	 * Log an error for the given source file where the type is defined
+	 * @param InTypeDef The source file and line number will be determined from the type
+	 * @param ErrorMsg The text of the error
+	 * @param Result Compilation result of the error
+	 */
+	static void LogError(FUnrealTypeDefinitionInfo& InTypeDef, const TCHAR* ErrorMsg, ECompilationResult::Type Result = ECompilationResult::OtherCompilationError);
+	
 	/**
 	 * Log an error without any source file information
 	 * @param ErrorMsg The text of the error
diff --git a/Engine/Source/Programs/UnrealHeaderTool/Private/HeaderParser.cpp b/Engine/Source/Programs/UnrealHeaderTool/Private/HeaderParser.cpp
--- a/Engine/Source/Programs/UnrealHeaderTool/Private/HeaderParser.cpp
+++ b/Engine/Source/Programs/UnrealHeaderTool/Private/HeaderParser.cpp
@@ -12,7 +12,7 @@
 #include "GeneratedCodeVersion.h"
 #include "ProfilingDebugging/ScopedTimers.h"
 #include "NativeClassExporter.h"
-#include "Classes.h"
+#include "EngineAPI.h"
 #include "ClassMaps.h"
 #include "StringUtils.h"
 #include "Misc/DefaultValueHelper.h"
@@ -33,6 +33,11 @@
 #include "Specifiers/StructSpecifiers.h"
 #include "Specifiers/VariableSpecifiers.h"
 
+// Globals for common class definitions
+extern FUnrealClassDefinitionInfo* GUObjectDef;
+extern FUnrealClassDefinitionInfo* GUClassDef;
+extern FUnrealClassDefinitionInfo* GUInterfaceDef;
+
 /*-----------------------------------------------------------------------------
 	Constants & declarations.
 -----------------------------------------------------------------------------*/
@@ -1027,7 +1032,7 @@ FString FHeaderParser::GetContext()
 //
 // Get a qualified class.
 //
-UClass* FHeaderParser::GetQualifiedClass(const TCHAR* Thing)
+FUnrealClassDefinitionInfo* FHeaderParser::GetQualifiedClass(const TCHAR* Thing)
 {
 	TCHAR ClassName[256]=TEXT("");
 
@@ -1044,7 +1049,7 @@ UClass* FHeaderParser::GetQualifiedClass(const TCHAR* Thing)
 		FError::Throwf(TEXT("%s: Missing class name"), Thing );
 	}
 
-	return FClasses::FindScriptClassOrThrow(ClassName);
+	return FUnrealClassDefinitionInfo::FindScriptClassOrThrow(ClassName);
 }
 
 /*-----------------------------------------------------------------------------
@@ -2064,7 +2069,7 @@ UScriptStruct* FHeaderParser::CompileStructDeclaration()
 			{
 				StructFlags |= STRUCT_Immutable | STRUCT_Atomic;
 
-				if (!FPaths::IsSamePath(Filename, GTypeDefinitionInfoMap[UObject::StaticClass()]->GetUnrealSourceFile().GetFilename()))
+				if (!FPaths::IsSamePath(Filename, GUObjectDef->GetUnrealSourceFile().GetFilename()))
 				{
 					UE_LOG_ERROR_UHT(TEXT("Immutable is being phased out in favor of SerializeNative, and is only legal on the mirror structs declared in UObject"));
 				}
@@ -2102,23 +2107,17 @@ UScriptStruct* FHeaderParser::CompileStructDeclaration()
 			BaseStructDef = Scope->FindTypeByName(*ParentStructNameInScript);
 		}
 
-		// Resolve structs declared in another class  //@TODO: UCREMOVAL: This seems extreme
+		// Try to locate globally  //@TODO: UCREMOVAL: This seems extreme
 		if (BaseStructDef == nullptr)
 		{
-			UScriptStruct* Type = nullptr;
 			if (bOverrideParentStructName)
 			{
-				Type = FClasses::FindObject<UScriptStruct>(ANY_PACKAGE, *ParentStructNameStripped);
-			}
-
-			if (Type == nullptr)
-			{
-				Type = FClasses::FindObject<UScriptStruct>(ANY_PACKAGE, *ParentStructNameInScript);
+				BaseStructDef = GTypeDefinitionInfoMap.FindByName<FUnrealScriptStructDefinitionInfo>(*ParentStructNameStripped);
 			}
 
-			if (Type != nullptr)
+			if (BaseStructDef == nullptr)
 			{
-				BaseStructDef = &GTypeDefinitionInfoMap.FindChecked<FUnrealFieldDefinitionInfo>(Type);
+				BaseStructDef = GTypeDefinitionInfoMap.FindByName<FUnrealScriptStructDefinitionInfo>(*ParentStructNameInScript);
 			}
 		}
 
@@ -2469,11 +2468,11 @@ bool FHeaderParser::AllowReferenceToClass(UStruct* Scope, UClass* CheckClass) co
 	Nest management.
 -----------------------------------------------------------------------------*/
 
-void FHeaderParser::PushNest(ENestType NestType, UStruct* InNode, FUnrealSourceFile* InSourceFile)
+void FHeaderParser::PushNest(ENestType NestType, FUnrealStructDefinitionInfo* InNodeDef, FUnrealSourceFile* InSourceFile)
 {
 	// Update pointer to top nesting level.
 	TopNest = &Nest[NestLevel++];
-	TopNest->SetScope(NestType == ENestType::GlobalScope ? &InSourceFile->GetScope().Get() : &FScope::GetTypeScope(InNode).Get());
+	TopNest->SetScope(NestType == ENestType::GlobalScope ? &InSourceFile->GetScope().Get() : &InNodeDef->GetScope().Get());
 	TopNest->NestType = NestType;
 
 	// Prevent overnesting.
@@ -2612,7 +2611,7 @@ void FHeaderParser::FixupDelegateProperties(FUnrealStructDefinitionInfo& StructD
 						UObject* DelegateSignatureOuter = DelegatePropertyToken.DelegateSignatureOwnerClass 
 							? ((UObject*)DelegatePropertyToken.DelegateSignatureOwnerClass) 
 							: ((UObject*)ANY_PACKAGE);
-						SourceDelegateFunction = FClasses::FindObject<UFunction>(DelegateSignatureOuter, *NameOfDelegateFunction);
+						SourceDelegateFunction = FEngineAPI::FindObject<UFunction>(DelegateSignatureOuter, *NameOfDelegateFunction);
 
 						if (SourceDelegateFunction == nullptr)
 						{
@@ -2641,12 +2640,12 @@ void FHeaderParser::FixupDelegateProperties(FUnrealStructDefinitionInfo& StructD
 					}
 
 					// make sure that the class that contains the delegate can be referenced here
-					UClass* DelegateOwnerClass = FClasses::FindScriptClassOrThrow(DelegateClassName);
-					if (FScope::GetTypeScope(DelegateOwnerClass)->FindTypeByName(*DelegateName) != nullptr)
+					FUnrealClassDefinitionInfo* DelegateOwnerClassDef = FUnrealClassDefinitionInfo::FindScriptClassOrThrow(DelegateClassName);
+					if (DelegateOwnerClassDef->GetScope()->FindTypeByName(*DelegateName) != nullptr)
 					{
-						FError::Throwf(TEXT("Inaccessible type: '%s'"), *DelegateOwnerClass->GetPathName());
+						FError::Throwf(TEXT("Inaccessible type: '%s'"), *DelegateOwnerClassDef->GetClass()->GetPathName());
 					}
-					SourceDelegateFunction = Cast<UFunction>(FindField(GTypeDefinitionInfoMap.FindChecked<FUnrealStructDefinitionInfo>(DelegateOwnerClass), *DelegateName, false, UFunction::StaticClass(), NULL));	
+					SourceDelegateFunction = Cast<UFunction>(FindField(*DelegateOwnerClassDef, *DelegateName, false, UFunction::StaticClass(), NULL));
 				}
 
 				if ( SourceDelegateFunction == NULL )
@@ -2770,18 +2769,19 @@ void FHeaderParser::CheckSparseClassData(const FUnrealStructDefinitionInfo& Stru
 
 	for (const FString& SparseClassDataTypeName : SparseClassDataTypes)
 	{
-		UScriptStruct* SparseClassDataStruct = FClasses::FindObject<UScriptStruct>(ANY_PACKAGE, *SparseClassDataTypeName);
+		FUnrealScriptStructDefinitionInfo* SparseClassDataStructDef = GTypeDefinitionInfoMap.FindByName<FUnrealScriptStructDefinitionInfo>(*SparseClassDataTypeName);
 
 		// make sure the sparse class data struct actually exists
-		if (!SparseClassDataStruct)
+		if (!SparseClassDataStructDef)
 		{
 			FError::Throwf(TEXT("Unable to find sparse data type %s for class %s."), *SparseClassDataTypeName, *ClassToCheck->GetName());
 			return;
 		}
 
+		UScriptStruct* SparseClassDataStruct = SparseClassDataStructDef->GetScriptStruct();
+
 		// check the data struct for invalid properties
-		FUnrealScriptStructDefinitionInfo& SparseClassDataStructDef = GTypeDefinitionInfoMap.FindChecked<FUnrealScriptStructDefinitionInfo>(SparseClassDataStruct);
-		for (FUnrealPropertyDefinitionInfo* PropertyDef : TUHTFieldRange<FUnrealPropertyDefinitionInfo>(SparseClassDataStructDef))
+		for (FUnrealPropertyDefinitionInfo* PropertyDef : TUHTFieldRange<FUnrealPropertyDefinitionInfo>(*SparseClassDataStructDef))
 		{
 			FProperty* Property = PropertyDef->GetProperty();
 			if (Property->HasAnyPropertyFlags(CPF_BlueprintAssignable))
@@ -2808,10 +2808,13 @@ void FHeaderParser::CheckSparseClassData(const FUnrealStructDefinitionInfo& Stru
 		ClassDef->GetSuperClass()->GetSparseClassDataTypes(ParentSparseClassDataTypeNames);
 		for (FString& ParentSparseClassDataTypeName : ParentSparseClassDataTypeNames)
 		{
-			UScriptStruct* ParentSparseClassDataStruct = FClasses::FindObject<UScriptStruct>(ANY_PACKAGE, *ParentSparseClassDataTypeName);
-			if (ParentSparseClassDataStruct && !SparseClassDataStruct->IsChildOf(ParentSparseClassDataStruct))
+			if (FUnrealScriptStructDefinitionInfo* ParentSparseClassDataStructDef = GTypeDefinitionInfoMap.FindByName<FUnrealScriptStructDefinitionInfo>(*ParentSparseClassDataTypeName))
 			{
-				FError::Throwf(TEXT("Class %s is a child of %s but its sparse class data struct, %s, does not inherit from %s."), *ClassToCheck->GetName(), *ParentClass->GetName(), *SparseClassDataStruct->GetName(), *ParentSparseClassDataStruct->GetName());
+				UScriptStruct* ParentSparseClassDataStruct = ParentSparseClassDataStructDef->GetScriptStruct();
+				if (!SparseClassDataStruct->IsChildOf(ParentSparseClassDataStruct))
+				{
+					FError::Throwf(TEXT("Class %s is a child of %s but its sparse class data struct, %s, does not inherit from %s."), *ClassToCheck->GetName(), *ParentClass->GetName(), *SparseClassDataStruct->GetName(), *ParentSparseClassDataStruct->GetName());
+				}
 			}
 		}
 	}
@@ -4161,10 +4164,10 @@ void FHeaderParser::GetVarType(
 		FToken InnerEnumType;
 		if (GetIdentifier(InnerEnumType, true))
 		{
-			if (UEnum* Enum = FClasses::FindObject<UEnum>(ANY_PACKAGE, InnerEnumType.Identifier))
+			if (FUnrealEnumDefinitionInfo* EnumDef = GTypeDefinitionInfoMap.FindByName<FUnrealEnumDefinitionInfo>(InnerEnumType.Identifier))
 			{
 				// In-scope enumeration.
-				VarProperty = FPropertyBase(GTypeDefinitionInfoMap.FindChecked<FUnrealEnumDefinitionInfo>(Enum), CPT_Byte);
+				VarProperty = FPropertyBase(*EnumDef, CPT_Byte);
 				bFoundEnum  = true;
 			}
 		}
@@ -4214,15 +4217,15 @@ void FHeaderParser::GetVarType(
 
 		VarProperty = FPropertyBase(PropertyClass, CPT_FieldPath);
 	}
-	else if (UEnum* Enum = FClasses::FindObject<UEnum>(ANY_PACKAGE, VarType.Identifier))
+	else if (FUnrealEnumDefinitionInfo* EnumDef = GTypeDefinitionInfoMap.FindByName<FUnrealEnumDefinitionInfo>(VarType.Identifier))
 	{
 		EPropertyType UnderlyingType = CPT_Byte;
 
 		if (VariableCategory == EVariableCategory::Member)
 		{
-			if (Enum->GetCppForm() != UEnum::ECppForm::EnumClass)
+			if (EnumDef->GetEnum()->GetCppForm() != UEnum::ECppForm::EnumClass)
 			{
-				FError::Throwf(TEXT("You cannot use the raw enum name as a type for member variables, instead use TEnumAsByte or a C++11 enum class with an explicit underlying type."), *Enum->CppType);
+				FError::Throwf(TEXT("You cannot use the raw enum name as a type for member variables, instead use TEnumAsByte or a C++11 enum class with an explicit underlying type."), *EnumDef->GetEnum()->CppType);
 			}
 		}
 
@@ -4238,7 +4241,7 @@ void FHeaderParser::GetVarType(
 		}
 
 		// In-scope enumeration.
-		VarProperty            = FPropertyBase(GTypeDefinitionInfoMap.FindChecked<FUnrealEnumDefinitionInfo>(Enum), UnderlyingType);
+		VarProperty            = FPropertyBase(*EnumDef, UnderlyingType);
 		bUnconsumedEnumKeyword = false;
 	}
 	else
@@ -4247,10 +4250,10 @@ void FHeaderParser::GetVarType(
 		bool bHandledType = false;
 		FString IdentifierStripped = GetClassNameWithPrefixRemoved(VarType.Identifier);
 		bool bStripped = false;
-		UScriptStruct* Struct = FClasses::FindObject<UScriptStruct>(ANY_PACKAGE, VarType.Identifier);
-		if (!Struct)
+		FUnrealScriptStructDefinitionInfo* StructDef = GTypeDefinitionInfoMap.FindByName<FUnrealScriptStructDefinitionInfo>(VarType.Identifier);
+		if (!StructDef)
 		{
-			Struct = FClasses::FindObject<UScriptStruct>(ANY_PACKAGE, *IdentifierStripped);
+			StructDef = GTypeDefinitionInfoMap.FindByName<FUnrealScriptStructDefinitionInfo>(*IdentifierStripped);
 			bStripped = true;
 		}
 
@@ -4268,22 +4271,20 @@ void FHeaderParser::GetVarType(
 			}
 		};
 
-		if (!Struct && MatchSymbol(TEXT("::")))
+		if (!StructDef && MatchSymbol(TEXT("::")))
 		{
 			FToken DelegateName;
 			if (GetIdentifier(DelegateName))
 			{
-				UClass* LocalOwnerClass = FClasses::FindClass(*IdentifierStripped);
-				if (LocalOwnerClass)
+				if (FUnrealClassDefinitionInfo* LocalOwnerClassDef = FUnrealClassDefinitionInfo::FindClass(*IdentifierStripped))
 				{
-					TSharedRef<FScope> LocScope = FScope::GetTypeScope(LocalOwnerClass);
 					const FString DelegateIdentifierStripped = GetClassNameWithPrefixRemoved(DelegateName.Identifier);
-					if (FUnrealFieldDefinitionInfo* FoundDef = LocScope->FindTypeByName(*(DelegateIdentifierStripped + HEADER_GENERATED_DELEGATE_SIGNATURE_SUFFIX)))
+					if (FUnrealFieldDefinitionInfo* FoundDef = LocalOwnerClassDef->GetScope()->FindTypeByName(*(DelegateIdentifierStripped + HEADER_GENERATED_DELEGATE_SIGNATURE_SUFFIX)))
 					{
 						if (FUnrealFunctionDefinitionInfo* DelegateFuncDef = FoundDef->AsFunction())
 						{
 							SetDelegateType(*DelegateFuncDef, DelegateIdentifierStripped);
-							VarProperty.DelegateSignatureOwnerClass = LocalOwnerClass;
+							VarProperty.DelegateSignatureOwnerClass = LocalOwnerClassDef->GetClass();
 						}
 					}
 				}
@@ -4297,8 +4298,9 @@ void FHeaderParser::GetVarType(
 		if (bHandledType)
 		{
 		}
-		else if (Struct)
+		else if (StructDef)
 		{
+			UScriptStruct* Struct = StructDef->GetScriptStruct();
 			if (bStripped)
 			{
 				const TCHAR* PrefixCPP = UHTConfig.StructsWithTPrefix.Contains(IdentifierStripped) ? TEXT("T") : Struct->GetPrefixCPP();
@@ -4316,7 +4318,7 @@ void FHeaderParser::GetVarType(
 
 			bHandledType = true;
 
-			VarProperty = FPropertyBase(GTypeDefinitionInfoMap.FindChecked<FUnrealScriptStructDefinitionInfo>(Struct));
+			VarProperty = FPropertyBase(*StructDef);
 			if ((Struct->StructFlags & STRUCT_HasInstancedReference) && !(Disallow & CPF_ContainsInstancedReference))
 			{
 				Flags |= CPF_ContainsInstancedReference;
@@ -4324,7 +4326,7 @@ void FHeaderParser::GetVarType(
 			// Struct keyword in front of a struct is legal, we 'consume' it
 			bUnconsumedStructKeyword = false;
 		}
-		else if (FClasses::FindObject<UScriptStruct>(ANY_PACKAGE, *IdentifierStripped) != nullptr)
+		else if (GTypeDefinitionInfoMap.FindByName<FUnrealScriptStructDefinitionInfo>(*IdentifierStripped) != nullptr)
 		{
 			bHandledType = true;
 
@@ -4342,7 +4344,7 @@ void FHeaderParser::GetVarType(
 			else
 			{
 				// An object reference of some type (maybe a restricted class?)
-				UClass* TempClass = NULL;
+				FUnrealClassDefinitionInfo* TempClassDef = nullptr;
 
 				EPropertyType PropertyType = CPT_ObjectReference;
 				const bool bIsSoftObjectPtrTemplate = VarType.Matches(TEXT("TSoftObjectPtr"), ESearchCase::CaseSensitive);
@@ -4351,16 +4353,16 @@ void FHeaderParser::GetVarType(
 
 				if (VarType.Matches(TEXT("TSubclassOf"), ESearchCase::CaseSensitive))
 				{
-					TempClass = UClass::StaticClass();
+					TempClassDef = GUClassDef;
 				}
 				else if (VarType.Matches(TEXT("FScriptInterface"), ESearchCase::CaseSensitive))
 				{
-					TempClass = UInterface::StaticClass();
+					TempClassDef = GUInterfaceDef;
 					Flags |= CPF_UObjectWrapper;
 				}
 				else if (VarType.Matches(TEXT("TSoftClassPtr"), ESearchCase::CaseSensitive))
 				{
-					TempClass = UClass::StaticClass();
+					TempClassDef = GUClassDef;
 					PropertyType = CPT_SoftObjectReference;
 				}
 				else
@@ -4387,8 +4389,8 @@ void FHeaderParser::GetVarType(
 						{
 							RedirectTypeIdentifier(InnerClass);
 
-							TempClass = FClasses::FindScriptClass(InnerClass.Identifier);
-							if (TempClass == nullptr)
+							TempClassDef = FUnrealClassDefinitionInfo::FindScriptClass(InnerClass.Identifier);
+							if (TempClassDef == nullptr)
 							{
 								FError::Throwf(TEXT("Unrecognized type '%s' (in expression %s<%s>) - type must be a UCLASS"), InnerClass.Identifier, VarType.Identifier, InnerClass.Identifier);
 							}
@@ -4429,12 +4431,14 @@ void FHeaderParser::GetVarType(
 					}
 					else
 					{
-						TempClass = FClasses::FindScriptClass(VarType.Identifier);
+						TempClassDef = FUnrealClassDefinitionInfo::FindScriptClass(VarType.Identifier);
 					}
 				}
 
-				if (TempClass != NULL)
+				if (TempClassDef != NULL)
 				{
+					UClass* TempClass = TempClassDef->GetClass();
+
 					bHandledType = true;
 
 					if ((PropertyType == CPT_WeakObjectReference) && (Disallow & CPF_AutoWeak)) // if it is not allowing anything, force it strong. this is probably a function arg
@@ -4448,8 +4452,7 @@ void FHeaderParser::GetVarType(
 						PropertyType = CPT_Interface;
 					}
 
-					FUnrealClassDefinitionInfo& TempClassDef = GTypeDefinitionInfoMap.FindChecked<FUnrealClassDefinitionInfo>(TempClass);
-					VarProperty = FPropertyBase(TempClassDef, PropertyType, bWeakIsAuto);
+					VarProperty = FPropertyBase(*TempClassDef, PropertyType, bWeakIsAuto);
 					if (TempClass->IsChildOf(UClass::StaticClass()))
 					{
 						if (MatchSymbol(TEXT('<')))
@@ -4468,13 +4471,13 @@ void FHeaderParser::GetVarType(
 
 							RedirectTypeIdentifier(Limitor);
 
-							VarProperty.MetaClassDef = &GTypeDefinitionInfoMap.FindChecked<FUnrealClassDefinitionInfo>(FClasses::FindScriptClassOrThrow(Limitor.Identifier));
+							VarProperty.MetaClassDef = FUnrealClassDefinitionInfo::FindScriptClassOrThrow(Limitor.Identifier);
 
 							RequireSymbol(TEXT('>'), TEXT("'class limitor'"), ESymbolParseOption::CloseTemplateBracket);
 						}
 						else
 						{
-							VarProperty.MetaClassDef = &GTypeDefinitionInfoMap.FindChecked<FUnrealClassDefinitionInfo>(UObject::StaticClass());
+							VarProperty.MetaClassDef = GUObjectDef;
 						}
 
 						if (PropertyType == CPT_WeakObjectReference)
@@ -4493,7 +4496,7 @@ void FHeaderParser::GetVarType(
 					}
 
 					// Inherit instancing flags
-					if (TempClassDef.HierarchyHasAnyClassFlags(CLASS_DefaultToInstanced))
+					if (TempClassDef->HierarchyHasAnyClassFlags(CLASS_DefaultToInstanced))
 					{
 						Flags |= ((CPF_InstancedReference | CPF_ExportObject) & (~Disallow));
 					}
@@ -4536,9 +4539,9 @@ void FHeaderParser::GetVarType(
 			// Resolve delegates declared in another class  //@TODO: UCREMOVAL: This seems extreme
 			if (!bHandledType)
 			{
-				if (UFunction* ExternDelegateFunc = FClasses::FindObject<UFunction>(ANY_PACKAGE, *(IdentifierStripped + HEADER_GENERATED_DELEGATE_SIGNATURE_SUFFIX)))
+				if (FUnrealFunctionDefinitionInfo* ExternDelegateFuncDef = GTypeDefinitionInfoMap.FindByName<FUnrealFunctionDefinitionInfo>(*(IdentifierStripped + HEADER_GENERATED_DELEGATE_SIGNATURE_SUFFIX)))
 				{
-					SetDelegateType(GTypeDefinitionInfoMap.FindChecked<FUnrealFunctionDefinitionInfo>(ExternDelegateFunc), IdentifierStripped);
+					SetDelegateType(*ExternDelegateFuncDef, IdentifierStripped);
 				}
 
 				if (!bHandledType)
@@ -4981,7 +4984,7 @@ FUnrealPropertyDefinitionInfo& FHeaderParser::GetVarNameAndDim
 		if (!Enum)
 		{
 			// If the enum wasn't declared in this scope, then try to find it anywhere we can
-			Enum = FClasses::FindObject<UEnum>(ANY_PACKAGE, Dimensions.String);
+			Enum = FEngineAPI::FindObject<UEnum>(ANY_PACKAGE, Dimensions.String);
 		}
 
 		if (Enum)
@@ -5617,12 +5620,11 @@ FUnrealClassDefinitionInfo& FHeaderParser::ParseClassNameDeclaration(FString& De
 {
 	ParseNameWithPotentialAPIMacroPrefix(/*out*/ DeclaredClassName, /*out*/ RequiredAPIMacroIfPresent, TEXT("class"));
 
-	UClass* FoundClass = FClasses::FindClass(*GetClassNameWithPrefixRemoved(*DeclaredClassName));
-	check(FoundClass);
+	FUnrealClassDefinitionInfo* ClassDef = FUnrealClassDefinitionInfo::FindClass(*GetClassNameWithPrefixRemoved(*DeclaredClassName));
+	check(ClassDef);
+	UClass* FoundClass = ClassDef->GetClass();
 	FoundClass->ClassCastFlags |= ClassCastFlagMap::Get().GetCastFlag(DeclaredClassName);
 
-	FUnrealClassDefinitionInfo& ClassDef = GTypeDefinitionInfoMap.FindChecked<FUnrealClassDefinitionInfo>(FoundClass);
-
 	// Skip optional final keyword
 	MatchIdentifier(TEXT("final"), ESearchCase::CaseSensitive);
 
@@ -5630,11 +5632,11 @@ FUnrealClassDefinitionInfo& FHeaderParser::ParseClassNameDeclaration(FString& De
 	ParseInheritance(TEXT("class"), [](const TCHAR* ClassName, bool bIsSuperClass) {}); // Eat the results, already been parsed
 
 	// Collect data about the super and base classes
-	if (FUnrealClassDefinitionInfo* SuperClassDef = UHTCast<FUnrealClassDefinitionInfo>(ClassDef.GetSuperClassInfo().Struct))
+	if (FUnrealClassDefinitionInfo* SuperClassDef = UHTCast<FUnrealClassDefinitionInfo>(ClassDef->GetSuperClassInfo().Struct))
 	{
 		FoundClass->ClassCastFlags |= SuperClassDef->GetClass()->ClassCastFlags;
 	}
-	for (FUnrealStructDefinitionInfo::FBaseClassInfo& BaseClassInfo : ClassDef.GetBaseClassInfo())
+	for (FUnrealStructDefinitionInfo::FBaseClassInfo& BaseClassInfo : ClassDef->GetBaseClassInfo())
 	{
 		if (FUnrealClassDefinitionInfo* BaseClassDef = UHTCast<FUnrealClassDefinitionInfo>(BaseClassInfo.Struct))
 		{
@@ -5649,7 +5651,7 @@ FUnrealClassDefinitionInfo& FHeaderParser::ParseClassNameDeclaration(FString& De
 			FoundClass->Interfaces.Emplace(BaseClass, 0, false);
 		}
 	}
-	return ClassDef;
+	return *ClassDef;
 }
 
 /**
@@ -5739,7 +5741,7 @@ UClass* FHeaderParser::CompileClassDeclaration()
 
 	Class->ClassFlags |= CLASS_Parsed;
 
-	PushNest(ENestType::Class, Class);
+	PushNest(ENestType::Class, &ClassDef);
 	
 	const uint32 PrevClassFlags = Class->ClassFlags;
 	ResetClassData();
@@ -5829,29 +5831,31 @@ UClass* FHeaderParser::CompileClassDeclaration()
 	return Class;
 }
 
-UClass* FHeaderParser::ParseInterfaceNameDeclaration(FString& DeclaredInterfaceName, FString& RequiredAPIMacroIfPresent)
+FUnrealClassDefinitionInfo* FHeaderParser::ParseInterfaceNameDeclaration(FString& DeclaredInterfaceName, FString& RequiredAPIMacroIfPresent)
 {
 	ParseNameWithPotentialAPIMacroPrefix(/*out*/ DeclaredInterfaceName, /*out*/ RequiredAPIMacroIfPresent, TEXT("interface"));
 
-	UClass* FoundClass = FClasses::FindClass(*GetClassNameWithPrefixRemoved(*DeclaredInterfaceName));
-	if (FoundClass == nullptr)
+	FUnrealClassDefinitionInfo* ClassDef = FUnrealClassDefinitionInfo::FindClass(*GetClassNameWithPrefixRemoved(*DeclaredInterfaceName));
+	if (ClassDef == nullptr)
 	{
 		return nullptr;
 	}
+	UClass* FoundClass = ClassDef->GetClass();
 
 	// Get super interface
 	bool bSpecifiesParentClass = MatchSymbol(TEXT(':'));
 	if (!bSpecifiesParentClass)
 	{
-		return FoundClass;
+		return ClassDef;
 	}
 
 	RequireIdentifier(TEXT("public"), ESearchCase::CaseSensitive, TEXT("class inheritance"));
 
 	// verify if our super class is an interface class
 	// the super class should have been marked as CLASS_Interface at the importing stage, if it were an interface
-	UClass* TempClass = GetQualifiedClass(TEXT("'extends'"));
-	check(TempClass);
+	FUnrealClassDefinitionInfo* TempClassDef = GetQualifiedClass(TEXT("'extends'"));
+	check(TempClassDef);
+	UClass* TempClass = TempClassDef->GetClass();
 	if( !(TempClass->ClassFlags & CLASS_Interface) )
 	{
 		// UInterface is special and actually extends from UObject, which isn't an interface
@@ -5869,7 +5873,7 @@ UClass* FHeaderParser::ParseInterfaceNameDeclaration(FString& DeclaredInterfaceN
 		FError::Throwf(TEXT("%s's superclass must be %s, not %s"), *FoundClass->GetPathName(), *SuperClass->GetPathName(), *TempClass->GetPathName());
 	}
 
-	return FoundClass;
+	return ClassDef;
 }
 
 bool FHeaderParser::TryParseIInterfaceClass()
@@ -5895,8 +5899,8 @@ bool FHeaderParser::TryParseIInterfaceClass()
 		return false;
 	}
 
-	UClass* FoundClass = nullptr;
-	if ((FoundClass = FClasses::FindClass(*DeclaredInterfaceName.Mid(1))) == nullptr)
+	FUnrealClassDefinitionInfo* FoundClassDef = nullptr;
+	if ((FoundClassDef = FUnrealClassDefinitionInfo::FindClass(*DeclaredInterfaceName.Mid(1))) == nullptr)
 	{
 		return false;
 	}
@@ -5905,7 +5909,7 @@ bool FHeaderParser::TryParseIInterfaceClass()
 	RequireSymbol(TEXT('{'), TEXT("C++ interface mix-in class declaration"));
 
 	// Push the interface class nesting again.
-	PushNest(ENestType::NativeInterface, FoundClass);
+	PushNest(ENestType::NativeInterface, FoundClassDef);
 
 	return true;
 }
@@ -5933,8 +5937,9 @@ void FHeaderParser::CompileInterfaceDeclaration()
 
 	// New style files have the interface name / extends afterwards
 	RequireIdentifier(TEXT("class"), ESearchCase::CaseSensitive, TEXT("Interface declaration"));
-	UClass* InterfaceClass = ParseInterfaceNameDeclaration(/*out*/ DeclaredInterfaceName, /*out*/ RequiredAPIMacroIfPresent);
-	FUnrealClassDefinitionInfo& ClassDef = GTypeDefinitionInfoMap.FindChecked<FUnrealClassDefinitionInfo>(InterfaceClass);
+	FUnrealClassDefinitionInfo* ClassDef = ParseInterfaceNameDeclaration(/*out*/ DeclaredInterfaceName, /*out*/ RequiredAPIMacroIfPresent);
+	check(ClassDef);
+	UClass* InterfaceClass = ClassDef->GetClass();
 	ClassDefinitionRanges.Add(InterfaceClass, ClassDefinitionRange(&Input[InputPos], nullptr));
 
 	// Record that this interface is RequiredAPI if the CORE_API style macro was present
@@ -5998,12 +6003,12 @@ void FHeaderParser::CompileInterfaceDeclaration()
 	}
 
 	// Try parsing metadata for the interface
-	FStructMetaData& StructData = ClassDef.GetStructMetaData();
+	FStructMetaData& StructData = ClassDef->GetStructMetaData();
 	StructData.SetPrologLine(PrologFinishLine);
 
 	// Register the metadata
-	AddModuleRelativePathToMetadata(ClassDef, MetaData);
-	AddMetaDataToClassData(ClassDef, MoveTemp(MetaData));
+	AddModuleRelativePathToMetadata(*ClassDef, MetaData);
+	AddMetaDataToClassData(*ClassDef, MoveTemp(MetaData));
 
 	// Handle the start of the rest of the interface
 	RequireSymbol( TEXT('{'), TEXT("'Class'") );
@@ -6015,7 +6020,7 @@ void FHeaderParser::CompileInterfaceDeclaration()
 
 	// Push the interface class nesting.
 	// we need a more specific set of allow flags for ENestType::Interface, only function declaration is allowed, no other stuff are allowed
-	PushNest(ENestType::Interface, InterfaceClass);
+	PushNest(ENestType::Interface, ClassDef);
 }
 
 void FHeaderParser::CompileRigVMMethodDeclaration(UStruct* Struct)
@@ -9026,20 +9031,24 @@ void FHeaderParser::CompileVersionDeclaration(UStruct* Struct)
 
 void FHeaderParser::ResetClassData()
 {
-	UClass* CurrentClass = GetCurrentClass();
+	FUnrealClassDefinitionInfo& CurrentClassDef = UHTCastChecked<FUnrealClassDefinitionInfo>(GetCurrentClassDef());
+	UClass* CurrentClass = CurrentClassDef.GetClass();
 	CurrentClass->PropertiesSize = 0;
 
 	// Set class flags and within.
 	CurrentClass->ClassFlags &= ~CLASS_RecompilerClear;
 
-	if (UClass* SuperClass = CurrentClass->GetSuperClass())
+	if (FUnrealClassDefinitionInfo* SuperClassDef = CurrentClassDef.GetSuperClass())
 	{
+		UClass* SuperClass = SuperClassDef->GetClass();
 		CurrentClass->ClassFlags |= (SuperClass->ClassFlags) & CLASS_ScriptInherit;
 		CurrentClass->ClassConfigName = SuperClass->ClassConfigName;
 		check(SuperClass->ClassWithin);
+		check(SuperClassDef->GetClassWithin());
 		if (CurrentClass->ClassWithin == nullptr)
 		{
 			CurrentClass->ClassWithin = SuperClass->ClassWithin;
+			CurrentClassDef.SetClassWithin(SuperClassDef->GetClassWithin());
 		}
 
 		// Copy special categories from parent
@@ -9070,6 +9079,7 @@ void FHeaderParser::ResetClassData()
 	}
 
 	check(CurrentClass->ClassWithin);
+	check(CurrentClassDef.GetClassWithin());
 }
 
 void FHeaderParser::PostPopNestClass(FUnrealClassDefinitionInfo& CurrentClassDef)
diff --git a/Engine/Source/Programs/UnrealHeaderTool/Private/HeaderParser.h b/Engine/Source/Programs/UnrealHeaderTool/Private/HeaderParser.h
--- a/Engine/Source/Programs/UnrealHeaderTool/Private/HeaderParser.h
+++ b/Engine/Source/Programs/UnrealHeaderTool/Private/HeaderParser.h
@@ -25,7 +25,6 @@ class UEnum;
 class UScriptStruct;
 class UDelegateFunction;
 class UStruct;
-class FClasses;
 class FScope;
 class FHeaderProvider;
 
@@ -744,7 +743,7 @@ protected:
 	void CompileRigVMMethodDeclaration(UStruct* Struct);
 	void ParseRigVMMethodParameters(FUnrealStructDefinitionInfo& StructDef);
 
-	UClass* ParseInterfaceNameDeclaration(FString& DeclaredInterfaceName, FString& RequiredAPIMacroIfPresent);
+	FUnrealClassDefinitionInfo* ParseInterfaceNameDeclaration(FString& DeclaredInterfaceName, FString& RequiredAPIMacroIfPresent);
 	bool TryParseIInterfaceClass();
 
 	bool CompileStatement(TArray<FUnrealFunctionDefinitionInfo*>& DelegatesToFixup);
@@ -832,7 +831,7 @@ protected:
 
 	const TCHAR* NestTypeName( ENestType NestType );
 
-	UClass* GetQualifiedClass(const TCHAR* Thing);
+	FUnrealClassDefinitionInfo* GetQualifiedClass(const TCHAR* Thing);
 
 	/**
 	 * Increase the nesting level, setting the new top nesting level to
@@ -842,7 +841,7 @@ protected:
 	 * @param	NestType	the new nesting type
 	 * @param	InNode		@todo
 	 */
-	void PushNest(ENestType NestType, UStruct* InNode, FUnrealSourceFile* SourceFile = nullptr);
+	void PushNest(ENestType NestType, FUnrealStructDefinitionInfo* InNodeDef, FUnrealSourceFile* SourceFile = nullptr);
 	void PopNest(ENestType NestType, const TCHAR* Descr);
 
 	/**
diff --git a/Engine/Source/Programs/UnrealHeaderTool/Private/HeaderProvider.cpp b/Engine/Source/Programs/UnrealHeaderTool/Private/HeaderProvider.cpp
--- a/Engine/Source/Programs/UnrealHeaderTool/Private/HeaderProvider.cpp
+++ b/Engine/Source/Programs/UnrealHeaderTool/Private/HeaderProvider.cpp
@@ -31,8 +31,7 @@ FUnrealSourceFile* FHeaderProvider::Resolve(const FUnrealSourceFile& ParentSourc
 		{
 		case EHeaderProviderSourceType::ClassName:
 		{
-			FName IdName(*Id, FNAME_Find);
-			if (TSharedRef<FUnrealTypeDefinitionInfo>* Source = GTypeDefinitionInfoMap.FindByName(IdName))
+			if (TSharedRef<FUnrealTypeDefinitionInfo>* Source = GTypeDefinitionInfoMap.FindByName(*Id))
 			{
 				Cache = (*Source)->HasSource() ? &(*Source)->GetUnrealSourceFile() : nullptr;
 				// There is an edge case with interfaces.  If you define the UMyInterface and IMyInterface in the same
@@ -51,8 +50,7 @@ FUnrealSourceFile* FHeaderProvider::Resolve(const FUnrealSourceFile& ParentSourc
 		{
 			if (!FUHTConfig::Get().StructsWithNoPrefix.Contains(Id))
 			{
-				FName IdName(*GetClassNameWithPrefixRemoved(Id), FNAME_Find);
-				if (TSharedRef<FUnrealTypeDefinitionInfo>* Source = GTypeDefinitionInfoMap.FindByName(IdName))
+				if (TSharedRef<FUnrealTypeDefinitionInfo>* Source = GTypeDefinitionInfoMap.FindByName(*GetClassNameWithPrefixRemoved(Id)))
 				{
 					Cache = (*Source)->HasSource() ? &(*Source)->GetUnrealSourceFile() : nullptr;
 				}
@@ -60,7 +58,7 @@ FUnrealSourceFile* FHeaderProvider::Resolve(const FUnrealSourceFile& ParentSourc
 			if (Cache == nullptr)
 			{
 				FName IdName(*Id, FNAME_Find);
-				if (TSharedRef<FUnrealTypeDefinitionInfo>* Source = GTypeDefinitionInfoMap.FindByName(IdName))
+				if (TSharedRef<FUnrealTypeDefinitionInfo>* Source = GTypeDefinitionInfoMap.FindByName(*Id))
 				{
 					Cache = (*Source)->HasSource() ? &(*Source)->GetUnrealSourceFile() : nullptr;
 				}
diff --git a/Engine/Source/Programs/UnrealHeaderTool/Private/NativeClassExporter.h b/Engine/Source/Programs/UnrealHeaderTool/Private/NativeClassExporter.h
--- a/Engine/Source/Programs/UnrealHeaderTool/Private/NativeClassExporter.h
+++ b/Engine/Source/Programs/UnrealHeaderTool/Private/NativeClassExporter.h
@@ -3,7 +3,6 @@
 #pragma once
 
 #include "CoreMinimal.h"
-#include "Classes.h"
 #include "Async/TaskGraphInterfaces.h"
 
 class FUnrealSourceFile;
@@ -46,8 +45,6 @@ namespace EExportFunctionType
 	};
 }
 
-class FClass;
-class FClasses;
 class FModuleClasses;
 class FScope;
 
diff --git a/Engine/Source/Programs/UnrealHeaderTool/Private/PropertyTypes.cpp b/Engine/Source/Programs/UnrealHeaderTool/Private/PropertyTypes.cpp
--- a/Engine/Source/Programs/UnrealHeaderTool/Private/PropertyTypes.cpp
+++ b/Engine/Source/Programs/UnrealHeaderTool/Private/PropertyTypes.cpp
@@ -2,7 +2,6 @@
 
 #include "PropertyTypes.h"
 #include "BaseParser.h"
-#include "Classes.h"
 #include "ClassMaps.h"
 #include "HeaderParser.h"
 #include "UnrealHeaderTool.h"
@@ -989,15 +988,14 @@ struct FPropertyTypeTraitsStruct : public FPropertyTypeTraitsBase
 	static bool DefaultValueStringCppFormatToInnerFormat(const FUnrealPropertyDefinitionInfo& PropDef, const FString& CppForm, FString& OutForm)
 	{
 		// Cache off the struct types, in case we need them later
-		UPackage* CoreUObjectPackage = UObject::StaticClass()->GetOutermost();
-		static const UScriptStruct* VectorStruct = FClasses::FindObjectChecked<UScriptStruct>(CoreUObjectPackage, TEXT("Vector"));
-		static const UScriptStruct* Vector2DStruct = FClasses::FindObjectChecked<UScriptStruct>(CoreUObjectPackage, TEXT("Vector2D"));
-		static const UScriptStruct* RotatorStruct = FClasses::FindObjectChecked<UScriptStruct>(CoreUObjectPackage, TEXT("Rotator"));
-		static const UScriptStruct* LinearColorStruct = FClasses::FindObjectChecked<UScriptStruct>(CoreUObjectPackage, TEXT("LinearColor"));
-		static const UScriptStruct* ColorStruct = FClasses::FindObjectChecked<UScriptStruct>(CoreUObjectPackage, TEXT("Color"));
-
-		UScriptStruct* Struct = PropDef.GetPropertyBase().ScriptStructDef->GetScriptStruct();
-		if (Struct == VectorStruct)
+		static const FUnrealScriptStructDefinitionInfo& VectorStructDef = GTypeDefinitionInfoMap.FindByNameChecked<FUnrealScriptStructDefinitionInfo>(TEXT("Vector"));
+		static const FUnrealScriptStructDefinitionInfo& Vector2DStructDef = GTypeDefinitionInfoMap.FindByNameChecked<FUnrealScriptStructDefinitionInfo>(TEXT("Vector2D"));
+		static const FUnrealScriptStructDefinitionInfo& RotatorStructDef = GTypeDefinitionInfoMap.FindByNameChecked<FUnrealScriptStructDefinitionInfo>(TEXT("Rotator"));
+		static const FUnrealScriptStructDefinitionInfo& LinearColorStructDef = GTypeDefinitionInfoMap.FindByNameChecked<FUnrealScriptStructDefinitionInfo>(TEXT("LinearColor"));
+		static const FUnrealScriptStructDefinitionInfo& ColorStructDef = GTypeDefinitionInfoMap.FindByNameChecked<FUnrealScriptStructDefinitionInfo>(TEXT("Color"));
+
+		FUnrealScriptStructDefinitionInfo* ScriptStructDef = PropDef.GetPropertyBase().ScriptStructDef;
+		if (ScriptStructDef == &VectorStructDef)
 		{
 			FString Parameters;
 			if (FDefaultValueHelper::Is(CppForm, TEXT("FVector::ZeroVector")))
@@ -1039,7 +1037,7 @@ struct FPropertyTypeTraitsStruct : public FPropertyTypeTraitsBase
 				}
 			}
 		}
-		else if (Struct == RotatorStruct)
+		else if (ScriptStructDef == &RotatorStructDef)
 		{
 			if (FDefaultValueHelper::Is(CppForm, TEXT("FRotator::ZeroRotator")))
 			{
@@ -1060,7 +1058,7 @@ struct FPropertyTypeTraitsStruct : public FPropertyTypeTraitsBase
 				}
 			}
 		}
-		else if (Struct == Vector2DStruct)
+		else if (ScriptStructDef == &Vector2DStructDef)
 		{
 			if (FDefaultValueHelper::Is(CppForm, TEXT("FVector2D::ZeroVector")))
 			{
@@ -1086,7 +1084,7 @@ struct FPropertyTypeTraitsStruct : public FPropertyTypeTraitsBase
 				}
 			}
 		}
-		else if (Struct == LinearColorStruct)
+		else if (ScriptStructDef == &LinearColorStructDef)
 		{
 			if (FDefaultValueHelper::Is(CppForm, TEXT("FLinearColor::White")))
 			{
@@ -1137,7 +1135,7 @@ struct FPropertyTypeTraitsStruct : public FPropertyTypeTraitsBase
 				}
 			}
 		}
-		else if (Struct == ColorStruct)
+		else if (ScriptStructDef == &ColorStructDef)
 		{
 			if (FDefaultValueHelper::Is(CppForm, TEXT("FColor::White")))
 			{
diff --git a/Engine/Source/Programs/UnrealHeaderTool/Private/Scope.cpp b/Engine/Source/Programs/UnrealHeaderTool/Private/Scope.cpp
--- a/Engine/Source/Programs/UnrealHeaderTool/Private/Scope.cpp
+++ b/Engine/Source/Programs/UnrealHeaderTool/Private/Scope.cpp
@@ -38,7 +38,7 @@ void DispatchType(FUnrealFieldDefinitionInfo& FieldDef, TArray<FUnrealEnumDefini
 	if (TypeClass == UClass::StaticClass() || TypeClass == UStruct::StaticClass())
 	{
 		// Inner scopes.
-		FScope::GetTypeScope((UStruct*)Type)->SplitTypesIntoArrays(Enums, Structs, DelegateFunctions);
+		FieldDef.GetScope()->SplitTypesIntoArrays(Enums, Structs, DelegateFunctions);
 	}
 	else if (TypeClass == UEnum::StaticClass())
 	{
@@ -71,17 +71,6 @@ void FScope::SplitTypesIntoArrays(TArray<FUnrealEnumDefinitionInfo*>& Enums, TAr
 	}
 }
 
-TSharedRef<FScope> FScope::GetTypeScope(UStruct* Type)
-{
-	TSharedRef<FUnrealTypeDefinitionInfo>* TypeDef = GTypeDefinitionInfoMap.Find(Type);
-	if (TypeDef == nullptr)
-	{
-		FError::Throwf(TEXT("Couldn't find scope for the type %s."), *Type->GetName());
-	}
-
-	return (*TypeDef)->GetScope();
-}
-
 FUnrealFieldDefinitionInfo* FScope::FindTypeByName(FName Name)
 {
 	if (!Name.IsNone())
diff --git a/Engine/Source/Programs/UnrealHeaderTool/Private/Scope.h b/Engine/Source/Programs/UnrealHeaderTool/Private/Scope.h
--- a/Engine/Source/Programs/UnrealHeaderTool/Private/Scope.h
+++ b/Engine/Source/Programs/UnrealHeaderTool/Private/Scope.h
@@ -123,15 +123,6 @@ public:
 		return false;
 	}
 
-	/**
-	 * Gets struct's or class's scope.
-	 *
-	 * @param Type A struct or class for which to return the scope.
-	 *
-	 * @returns Struct's or class's scope.
-	 */
-	static TSharedRef<FScope> GetTypeScope(UStruct* Type);
-
 	/**
 	 * Gets structs, enums and delegate functions from this scope.
 	 *
@@ -411,16 +402,14 @@ public:
 		{
 			ScopesToTraverse.Add(Scope);
 
-			UStruct* Struct = ((FStructScope*)CurrentScope)->GetStruct();
+			FUnrealStructDefinitionInfo& StructDef = ((FStructScope*)CurrentScope)->GetStructDef();
 
-			if (Struct->IsA<UClass>())
+			if (FUnrealClassDefinitionInfo* ClassDef = UHTCast<FUnrealClassDefinitionInfo>(StructDef))
 			{
-				UClass* Class = ((UClass*)Struct)->GetSuperClass();
-
-				while (Class && !(Class->ClassFlags & EClassFlags::CLASS_Intrinsic))
+				// Skip myself when starting this loop, we only care about the parents
+				for (ClassDef = ClassDef->GetSuperClass(); ClassDef && (ClassDef->GetClass()->ClassFlags & EClassFlags::CLASS_Intrinsic) == 0; ClassDef = ClassDef->GetSuperClass())
 				{
-					ScopesToTraverse.Add(&FScope::GetTypeScope(Class).Get());
-					Class = Class->GetSuperClass();
+					ScopesToTraverse.Add(&ClassDef->GetScope().Get());
 				}
 			}
 
diff --git a/Engine/Source/Programs/UnrealHeaderTool/Private/UnrealTypeDefinitionInfo.cpp b/Engine/Source/Programs/UnrealHeaderTool/Private/UnrealTypeDefinitionInfo.cpp
--- a/Engine/Source/Programs/UnrealHeaderTool/Private/UnrealTypeDefinitionInfo.cpp
+++ b/Engine/Source/Programs/UnrealHeaderTool/Private/UnrealTypeDefinitionInfo.cpp
@@ -2,18 +2,20 @@
 
 #include "UnrealTypeDefinitionInfo.h"
 #include "BaseParser.h"
-#include "Classes.h"
 #include "ClassMaps.h"
 #include "HeaderParser.h"
 #include "NativeClassExporter.h"
 #include "Scope.h"
+#include "StringUtils.h"
 #include "UnrealHeaderTool.h"
 #include "UnrealSourceFile.h"
+#include "Specifiers/ClassMetadataSpecifiers.h"
+
+#include "Algo/FindSortedStringCaseInsensitive.h"
 #include "Misc/PackageName.h"
 #include "UObject/ErrorException.h"
 #include "UObject/Interface.h"
-#include "Specifiers/ClassMetadataSpecifiers.h"
-#include "Algo/FindSortedStringCaseInsensitive.h"
+#include "UObject/ObjectRedirector.h"
 
 namespace
 {
@@ -107,6 +109,39 @@ namespace
 		}
 		return PackageName;
 	}
+
+	/**
+	 * Returns True if the given class name includes a valid Unreal prefix and matches based on the given class.
+	 *
+	 * @param InNameToCheck - Name w/ potential prefix to check
+	 * @param OriginalClass - Class to check against
+	 */
+	bool ClassNameHasValidPrefix(const FString& InNameToCheck, const UClass* OriginalClass)
+	{
+		bool bIsLabledDeprecated;
+		GetClassPrefix(InNameToCheck, bIsLabledDeprecated);
+
+		// If the class is labeled deprecated, don't try to resolve it during header generation, valid results can't be guaranteed.
+		if (bIsLabledDeprecated)
+		{
+			return true;
+		}
+
+		const FString OriginalClassName = FUnrealTypeDefinitionInfo::GetNameWithPrefix(OriginalClass);
+
+		bool bNamesMatch = (InNameToCheck == OriginalClassName);
+
+		if (!bNamesMatch)
+		{
+			//@TODO: UCREMOVAL: I/U interface hack - Ignoring prefixing for this call
+			if (OriginalClass->HasAnyClassFlags(CLASS_Interface))
+			{
+				bNamesMatch = InNameToCheck.Mid(1) == OriginalClassName.Mid(1);
+			}
+		}
+
+		return bNamesMatch;
+	}
 }
 
 FUnrealPropertyDefinitionInfo* FUnrealTypeDefinitionInfo::AsProperty()
@@ -496,6 +531,83 @@ FUnrealClassDefinitionInfo::FUnrealClassDefinitionInfo(FUnrealSourceFile& InSour
 	}
 }
 
+FUnrealClassDefinitionInfo* FUnrealClassDefinitionInfo::FindClass(const TCHAR* ClassName)
+{
+	FUnrealClassDefinitionInfo* ClassDef = nullptr;
+	{
+		ClassDef = GTypeDefinitionInfoMap.FindByName<FUnrealClassDefinitionInfo>(ClassName);
+	}
+
+	if (ClassDef != nullptr)
+	{
+		if (UObjectRedirector* RenamedClassRedirector = FindObject<UObjectRedirector>(ANY_PACKAGE, ClassName))
+		{
+			UClass* RedierClass = CastChecked<UClass>(RenamedClassRedirector->DestinationObject);
+			if (RedierClass != nullptr)
+			{
+				ClassDef = &GTypeDefinitionInfoMap.FindChecked<FUnrealClassDefinitionInfo>(RedierClass);
+			}
+		}
+	}
+
+	return ClassDef;
+}
+
+FUnrealClassDefinitionInfo* FUnrealClassDefinitionInfo::FindScriptClassOrThrow(const FString& InClassName)
+{
+	FString ErrorMsg;
+	if (FUnrealClassDefinitionInfo* ResultDef = FindScriptClass(InClassName, &ErrorMsg))
+	{
+		return ResultDef;
+	}
+
+	FError::Throwf(*ErrorMsg);
+
+	// Unreachable, but compiler will warn otherwise because FError::Throwf isn't declared noreturn
+	return nullptr;
+}
+
+FUnrealClassDefinitionInfo* FUnrealClassDefinitionInfo::FindScriptClass(const FString& InClassName, FString* OutErrorMsg)
+{
+	// Strip the class name of its prefix and then do a search for the class
+	FString ClassNameStripped = GetClassNameWithPrefixRemoved(InClassName);
+	if (FUnrealClassDefinitionInfo* FoundClassDef = FindClass(*ClassNameStripped))
+	{
+		// If the class was found with the stripped class name, verify that the correct prefix was used and throw an error otherwise
+		if (!ClassNameHasValidPrefix(InClassName, FoundClassDef->GetClass()))
+		{
+			if (OutErrorMsg)
+			{
+				*OutErrorMsg = FString::Printf(TEXT("Class '%s' has an incorrect prefix, expecting '%s'"), *InClassName, *FUnrealTypeDefinitionInfo::GetNameWithPrefix(FoundClassDef->GetClass()));
+			}
+			return nullptr;
+		}
+
+		return FoundClassDef;
+	}
+
+	// Couldn't find the class with a class name stripped of prefix (or a prefix was not found)
+	// See if the prefix was forgotten by trying to find the class with the given identifier
+	if (FUnrealClassDefinitionInfo* FoundClassDef = FindClass(*InClassName))
+	{
+		// If the class was found with the given identifier, the user forgot to use the correct Unreal prefix	
+		if (OutErrorMsg)
+		{
+			*OutErrorMsg = FString::Printf(TEXT("Class '%s' is missing a prefix, expecting '%s'"), *InClassName, *FUnrealTypeDefinitionInfo::GetNameWithPrefix(FoundClassDef->GetClass()));
+		}
+	}
+	else
+	{
+		// If the class was still not found, it wasn't a valid identifier
+		if (OutErrorMsg)
+		{
+			*OutErrorMsg = FString::Printf(TEXT("Class '%s' not found."), *InClassName);
+		}
+	}
+
+	return nullptr;
+}
+
 uint32 FUnrealClassDefinitionInfo::GetHash(bool bIncludeNoExport) const
 {
 	if (!bIncludeNoExport)
@@ -556,7 +668,7 @@ void FUnrealClassDefinitionInfo::ParseClassProperties(TArray<FPropertySpecifier>
 
 		case EClassMetadataSpecifier::Within:
 
-			ClassWithin = FHeaderParser::RequireExactlyOneSpecifierValue(PropSpecifier);
+			ClassWithinStr = FHeaderParser::RequireExactlyOneSpecifierValue(PropSpecifier);
 			break;
 
 		case EClassMetadataSpecifier::EditInlineNew:
@@ -995,13 +1107,14 @@ void FUnrealClassDefinitionInfo::SetAndValidateWithinClass()
 	UClass* Class = GetClass();
 
 	// Process all of the class specifiers
-	if (ClassWithin.IsEmpty() == false)
+	if (ClassWithinStr.IsEmpty() == false)
 	{
-		UClass* RequiredWithinClass = FClasses::FindClass(*ClassWithin);
-		if (!RequiredWithinClass)
+		FUnrealClassDefinitionInfo* RequiredWithinClassDef = FUnrealClassDefinitionInfo::FindClass(*ClassWithinStr);
+		if (!RequiredWithinClassDef)
 		{
-			FError::Throwf(TEXT("Within class '%s' not found."), *ClassWithin);
+			FError::Throwf(TEXT("Within class '%s' not found."), *ClassWithinStr);
 		}
+		UClass* RequiredWithinClass = RequiredWithinClassDef->GetClass();
 		if (RequiredWithinClass->IsChildOf(UInterface::StaticClass()))
 		{
 			FError::Throwf(TEXT("Classes cannot be 'within' interfaces"));
@@ -1031,6 +1144,8 @@ void FUnrealClassDefinitionInfo::SetAndValidateWithinClass()
 	{
 		FError::Throwf(TEXT("Parent class declared within '%s'.  Cannot override within class with '%s' since it isn't a child"), *ExpectedWithin->GetName(), *Class->ClassWithin->GetName());
 	}
+
+	ClassWithin = &GTypeDefinitionInfoMap.FindChecked<FUnrealClassDefinitionInfo>(Class->ClassWithin);
 }
 
 void FUnrealClassDefinitionInfo::GetHideCategories(TArray<FString>& OutHideCategories) const
diff --git a/Engine/Source/Programs/UnrealHeaderTool/Private/UnrealTypeDefinitionInfo.h b/Engine/Source/Programs/UnrealHeaderTool/Private/UnrealTypeDefinitionInfo.h
--- a/Engine/Source/Programs/UnrealHeaderTool/Private/UnrealTypeDefinitionInfo.h
+++ b/Engine/Source/Programs/UnrealHeaderTool/Private/UnrealTypeDefinitionInfo.h
@@ -11,7 +11,6 @@
 #include "ParserHelper.h"
 
 // Forward declarations.
-class FClass;
 class FScope;
 class FUnrealSourceFile;
 struct FManifestModule;
@@ -1171,6 +1170,31 @@ public:
 		: FUnrealStructDefinitionInfo(MoveTemp(InNameCPP))
 	{ }
 
+	/**
+	 * Attempts to find a class definition based on the given name
+	 */
+	static FUnrealClassDefinitionInfo* FindClass(const TCHAR* ClassName);
+
+	/**
+	 * Attempts to find a script class based on the given name. Will attempt to strip
+	 * the prefix of the given name while searching. Throws an exception with the script error
+	 * if the class could not be found.
+	 *
+	 * @param   InClassName  Name w/ Unreal prefix to use when searching for a class
+	 * @return               The found class.
+	 */
+	static FUnrealClassDefinitionInfo* FindScriptClassOrThrow(const FString& InClassName);
+
+	/**
+	 * Attempts to find a script class based on the given name. Will attempt to strip
+	 * the prefix of the given name while searching. Optionally returns script errors when appropriate.
+	 *
+	 * @param   InClassName  Name w/ Unreal prefix to use when searching for a class
+	 * @param   OutErrorMsg  Error message (if any) giving the caller flexibility in how they present an error
+	 * @return               The found class, or NULL if the class was not found.
+	 */
+	static FUnrealClassDefinitionInfo* FindScriptClass(const FString& InClassName, FString* OutErrorMsg = nullptr);
+
 	virtual FUnrealClassDefinitionInfo* AsClass() override
 	{
 		return this;
@@ -1345,13 +1369,26 @@ public:
 	 */
 	void MergeCategoryMetaData(TMap<FName, FString>& InMetaData) const;
 
-
 	void GetSparseClassDataTypes(TArray<FString>& OutSparseClassDataTypes) const;
 
+	/**
+	 * Get the class's class within setting
+	 */
+	FUnrealClassDefinitionInfo* GetClassWithin() const
+	{
+		return ClassWithin;
+	}
+
+	/**
+	 * Set the class's class within setting
+	 */
+	void SetClassWithin(FUnrealClassDefinitionInfo* InClassWithin)
+	{
+		ClassWithin = InClassWithin;
+	}
+
 public:
 	EClassFlags ClassFlags = CLASS_None;
-	FString ClassWithin;
-	FString ConfigName;
 	TMap<FName, FString> MetaData;
 
 private:
@@ -1365,8 +1402,6 @@ private:
 	void GetHideCategories(TArray<FString>& OutHideCategories) const;
 	void GetShowCategories(TArray<FString>& OutShowCategories) const;
 
-	FString EnclosingDefine;
-	ESerializerArchiveType ArchiveType = ESerializerArchiveType::None;
 	TArray<FString> ShowCategories;
 	TArray<FString> ShowFunctions;
 	TArray<FString> DontAutoCollapseCategories;
@@ -1378,6 +1413,11 @@ private:
 	TArray<FString> DependsOn;
 	TArray<FString> ClassGroupNames;
 	TArray<FString> SparseClassDataTypes;
+	FString EnclosingDefine;
+	FString ClassWithinStr;
+	FString ConfigName;
+	FUnrealClassDefinitionInfo* ClassWithin = nullptr;
+	ESerializerArchiveType ArchiveType = ESerializerArchiveType::None;
 	bool bIsInterface = false;
 	bool bWantsToBePlaceable = false;
 };
