Fix to the Mac build.
#rnx

diff --git a/Engine/Source/Programs/UnrealHeaderTool/Private/EngineAPI.h b/Engine/Source/Programs/UnrealHeaderTool/Private/EngineAPI.h
--- a/Engine/Source/Programs/UnrealHeaderTool/Private/EngineAPI.h
+++ b/Engine/Source/Programs/UnrealHeaderTool/Private/EngineAPI.h
@@ -3,7 +3,7 @@
 #pragma once
 
 #include "CoreMinimal.h"
-#include "UObject\ObjectMacros.h"
+#include "UObject/UObjectGlobals.h"
 
 /**
  * Wrapper class for access to the engine API.
