Remove most of the engine types from the exporter API

#rnx

diff --git a/Engine/Source/Programs/UnrealHeaderTool/Private/CodeGenerator.cpp b/Engine/Source/Programs/UnrealHeaderTool/Private/CodeGenerator.cpp
--- a/Engine/Source/Programs/UnrealHeaderTool/Private/CodeGenerator.cpp
+++ b/Engine/Source/Programs/UnrealHeaderTool/Private/CodeGenerator.cpp
@@ -266,18 +266,19 @@ namespace
 		const ClassDefinitionRange& ClassRange,
 		const TCHAR* ClassCPPName,
 		const TCHAR* API,
-		UClass* Class,
+		FUnrealClassDefinitionInfo& ClassDef,
 		UClass* SuperClass,
 		FOutputDevice& Writer,
 		const FUnrealSourceFile& SourceFile,
 		EExportClassOutFlags& OutFlags)
 	{
+		UClass* Class = ClassDef.GetClass();
 		const bool bHasGetLifetimeReplicatedProps = HasIdentifierExactMatch(ClassRange.Start, ClassRange.End, STRING_GetLifetimeReplicatedPropsStr);
 
 		if (!bHasGetLifetimeReplicatedProps)
 		{
 			// Default version autogenerates declarations.
-			if (SourceFile.GetGeneratedCodeVersionForStruct(Class) == EGeneratedCodeVersion::V1)
+			if (SourceFile.GetGeneratedCodeVersionForStruct(ClassDef) == EGeneratedCodeVersion::V1)
 			{
 				Writer.Logf(TEXT("\tvoid GetLifetimeReplicatedProps(TArray<FLifetimeProperty>& OutLifetimeProps) const override;\r\n"));
 			}
@@ -762,7 +763,7 @@ void FNativeClassHeaderGenerator::ExportProperties(FOutputDevice& Out, FUnrealSt
 		// Indent code and export CPP text.
 		FUHTStringBuilder JustPropertyDecl;
 		Property->ExportCppDeclaration( JustPropertyDecl, EExportedDeclaration::Member, PropertyDef->GetArrayDimensions());
-		ApplyAlternatePropertyExportText(Property, JustPropertyDecl, EExportingState::TypeEraseDelegates);
+		ApplyAlternatePropertyExportText(*PropertyDef, JustPropertyDecl, EExportingState::TypeEraseDelegates);
 
 		// Finish up line.
 		Out.Logf(TEXT("%s%s;\r\n"), FCString::Tab(TextIndent + 1), *JustPropertyDecl);
@@ -2941,7 +2942,7 @@ void FNativeClassHeaderGenerator::ExportClassFromSourceFileInner(
 			*ClassCPPName,
 			*SuperClassCPPName,
 			Class->HasAnyClassFlags(CLASS_Abstract) ? TEXT("CLASS_Abstract") : TEXT("0"),
-			*GetClassFlagExportText(Class),
+			*GetClassFlagExportText(ClassDef),
 			bCastedClass ? *FString::Printf(TEXT("CASTCLASS_%s"), *ClassCPPName) : TEXT("CASTCLASS_None"),
 			*ClassDef.GetTypePackageName(),
 			*APIArg);
@@ -3029,7 +3030,7 @@ void FNativeClassHeaderGenerator::ExportClassFromSourceFileInner(
 
 			if (ClassHasReplicatedProperties(ClassDef))
 			{
-				WriteReplicatedMacroData(ClassRange, *ClassCPPName, *APIArg, Class, SuperClass, InterfaceBoilerplate, SourceFile, OutFlags);
+				WriteReplicatedMacroData(ClassRange, *ClassCPPName, *APIArg, ClassDef, SuperClass, InterfaceBoilerplate, SourceFile, OutFlags);
 			}
 
 			FString NoPureDeclsMacroName = SourceFile.GetGeneratedMacroName(StructData, TEXT("_INCLASS_IINTERFACE_NO_PURE_DECLS"));
@@ -3056,7 +3057,7 @@ void FNativeClassHeaderGenerator::ExportClassFromSourceFileInner(
 
 			if (ClassHasReplicatedProperties(ClassDef))
 			{
-				WriteReplicatedMacroData(ClassRange, *ClassCPPName, *APIArg, Class, SuperClass, Boilerplate, SourceFile, OutFlags);
+				WriteReplicatedMacroData(ClassRange, *ClassCPPName, *APIArg, ClassDef, SuperClass, Boilerplate, SourceFile, OutFlags);
 			}
 
 			{
@@ -3363,11 +3364,11 @@ bool FNativeClassHeaderGenerator::WriteHeader(FGeneratedFileInfo& FileInfo, cons
  * Returns a string in the format CLASS_Something|CLASS_Something which represents all class flags that are set for the specified
  * class which need to be exported as part of the DECLARE_CLASS macro
  */
-FString FNativeClassHeaderGenerator::GetClassFlagExportText( UClass* Class )
+FString FNativeClassHeaderGenerator::GetClassFlagExportText(FUnrealClassDefinitionInfo& ClassDef)
 {
 	FString StaticClassFlagText;
 
-	check(Class);
+	UClass* Class = ClassDef.GetClass();
 	if ( Class->HasAnyClassFlags(CLASS_Transient) )
 	{
 		StaticClassFlagText += TEXT(" | CLASS_Transient");
@@ -3482,7 +3483,7 @@ void FNativeClassHeaderGenerator::ExportGeneratedStructBodyMacros(FOutputDevice&
 	// for RigVM methods we need to generated a macro used for implementing the static method
 	// and prepare two prologs: one for the virtual function implementation, and one for the stub
 	// invoking the static method.
-	const FRigVMStructInfo* StructRigVMInfo = FHeaderParser::StructRigVMMap.Find(Struct);
+	const FRigVMStructInfo* StructRigVMInfo = FHeaderParser::StructRigVMMap.Find(&ScriptStructDef);
 	if(StructRigVMInfo)
 	{
 		//RigVMStubProlog.Add(FString::Printf(TEXT("ensure(RigVMOperandMemory.Num() == %d);"), StructRigVMInfo->Members.Num()));
@@ -4405,7 +4406,7 @@ void FNativeClassHeaderGenerator::ExportEventParm(FUHTStringBuilder& Out, TSet<F
 		}
 
 		Prop->ExportCppDeclaration(PropertyText, EExportedDeclaration::Local, PropDef->GetArrayDimensions());
-		ApplyAlternatePropertyExportText(Prop, PropertyText, ExportingState);
+		ApplyAlternatePropertyExportText(*PropDef, PropertyText, ExportingState);
 
 		PropertyText.Log(TEXT(";\r\n"));
 		Out += *PropertyText;
@@ -4555,7 +4556,7 @@ FString FNativeClassHeaderGenerator::GetFunctionReturnString(FUnrealFunctionDefi
 		FString ReturnType = Return->GetCPPType(&ExtendedReturnType, CPPF_ArgumentOrReturnValue);
 		FUHTStringBuilder ReplacementText;
 		ReplacementText += MoveTemp(ReturnType);
-		ApplyAlternatePropertyExportText(Return, ReplacementText, EExportingState::Normal);
+		ApplyAlternatePropertyExportText(*ReturnDef, ReplacementText, EExportingState::Normal);
 		Result = MoveTemp(ReplacementText) + MoveTemp(ExtendedReturnType);
 	}
 	else
@@ -4772,7 +4773,7 @@ void FNativeClassHeaderGenerator::ExportNativeFunctionHeader(
 		OutFwdDecls.Add(ReturnProperty->GetCPPTypeForwardDeclaration());
 		FUHTStringBuilder ReplacementText;
 		ReplacementText += ReturnType;
-		ApplyAlternatePropertyExportText(ReturnProperty, ReplacementText, EExportingState::Normal);
+		ApplyAlternatePropertyExportText(*ReturnPropertyDef, ReplacementText, EExportingState::Normal);
 		Out.Logf(TEXT("%s%s"), *ReplacementText, *ExtendedReturnType);
 	}
 	else
@@ -4828,7 +4829,7 @@ void FNativeClassHeaderGenerator::ExportNativeFunctionHeader(
 		FUHTStringBuilder PropertyText;
 
 		Property->ExportCppDeclaration( PropertyText, EExportedDeclaration::Parameter, PropertyDef->GetArrayDimensions() );
-		ApplyAlternatePropertyExportText(Property, PropertyText, EExportingState::Normal);
+		ApplyAlternatePropertyExportText(*PropertyDef, PropertyText, EExportingState::Normal);
 
 		Out.Logf(TEXT("%s"), *PropertyText);
 	}
@@ -4930,7 +4931,7 @@ void FNativeClassHeaderGenerator::ExportFunctionThunk(FUHTStringBuilder& RPCWrap
 		FUHTStringBuilder ReplacementText;
 		ReplacementText += TypeText;
 
-		ApplyAlternatePropertyExportText(Param, ReplacementText, EExportingState::Normal);
+		ApplyAlternatePropertyExportText(*ParamDef, ReplacementText, EExportingState::Normal);
 		TypeText = ReplacementText;
 
 		FString DefaultValueText;
@@ -5080,7 +5081,7 @@ void FNativeClassHeaderGenerator::ExportFunctionThunk(FUHTStringBuilder& RPCWrap
 		FUHTStringBuilder ReplacementText;
 		FString ReturnExtendedType;
 		ReplacementText += Return->GetCPPType(&ReturnExtendedType);
-		ApplyAlternatePropertyExportText(Return, ReplacementText, EExportingState::Normal);
+		ApplyAlternatePropertyExportText(*ReturnDef, ReplacementText, EExportingState::Normal);
 
 		FString ReturnType = ReplacementText;
 		if (Return->HasAnyPropertyFlags(CPF_ConstParm) && CastField<FObjectProperty>(Return))
@@ -5123,7 +5124,7 @@ FString FNativeClassHeaderGenerator::GetFunctionParameterString(FUnrealFunctionD
 		}
 
 		Property->ExportCppDeclaration(PropertyText, EExportedDeclaration::Parameter, PropertyDef->GetArrayDimensions(), 0, true);
-		ApplyAlternatePropertyExportText(Property, PropertyText, EExportingState::Normal);
+		ApplyAlternatePropertyExportText(*PropertyDef, PropertyText, EExportingState::Normal);
 
 		ParameterList += PropertyText;
 		PropertyText.Reset();
@@ -5301,7 +5302,7 @@ void FNativeClassHeaderGenerator::ExportNativeFunctions(FOutputDevice& OutGenera
 			}
 
 			// Versions that skip function autodeclaration throw an error when a function is missing.
-			if (ClassRange.bHasGeneratedBody && (SourceFile.GetGeneratedCodeVersionForStruct(Class) > EGeneratedCodeVersion::V1))
+			if (ClassRange.bHasGeneratedBody && (SourceFile.GetGeneratedCodeVersionForStruct(ClassDef) > EGeneratedCodeVersion::V1))
 			{
 				CheckRPCFunctions(OutReferenceGatherers, *FunctionDef, ClassCPPName, ImplementationPosition, ValidatePosition, SourceFile);
 			}
@@ -5352,7 +5353,7 @@ void FNativeClassHeaderGenerator::ExportNativeFunctions(FOutputDevice& OutGenera
 
 		// Put static checks before RPCWrappers to get proper messages from static asserts before compiler errors.
 		FString NoPureDeclsMacroName = SourceFile.GetGeneratedMacroName(StructData, TEXT("_RPC_WRAPPERS_NO_PURE_DECLS"));
-		if (SourceFile.GetGeneratedCodeVersionForStruct(Class) > EGeneratedCodeVersion::V1)
+		if (SourceFile.GetGeneratedCodeVersionForStruct(ClassDef) > EGeneratedCodeVersion::V1)
 		{
 			WriteMacro(OutGeneratedHeaderText, NoPureDeclsMacroName, RuntimeStringBuilders.RPCWrappers);
 		}
@@ -5385,7 +5386,7 @@ void FNativeClassHeaderGenerator::ExportNativeFunctions(FOutputDevice& OutGenera
 
 		// Put static checks before RPCWrappers to get proper messages from static asserts before compiler errors.
 		FString NoPureDeclsMacroName = SourceFile.GetGeneratedMacroName(StructData, TEXT("_EDITOR_ONLY_RPC_WRAPPERS_NO_PURE_DECLS"));
-		if (SourceFile.GetGeneratedCodeVersionForStruct(Class) > EGeneratedCodeVersion::V1)
+		if (SourceFile.GetGeneratedCodeVersionForStruct(ClassDef) > EGeneratedCodeVersion::V1)
 		{
 			WriteMacro(OutGeneratedHeaderText, NoPureDeclsMacroName, EditorStringBuilders.RPCWrappers);
 		}
@@ -5521,8 +5522,9 @@ void FNativeClassHeaderGenerator::ExportCallbackFunctions(
  * @param	Prop			the property that is being exported
  * @param	PropertyText	the string containing the text exported from ExportCppDeclaration
  */
-void FNativeClassHeaderGenerator::ApplyAlternatePropertyExportText(FProperty* Prop, FUHTStringBuilder& PropertyText, EExportingState ExportingState)
+void FNativeClassHeaderGenerator::ApplyAlternatePropertyExportText(FUnrealPropertyDefinitionInfo& PropertyDef, FUHTStringBuilder& PropertyText, EExportingState ExportingState)
 {
+	FProperty* Prop = PropertyDef.GetProperty();
 	FArrayProperty* ArrayProperty = CastField<FArrayProperty>(Prop);
 	FProperty* InnerProperty = ArrayProperty ? ArrayProperty->Inner : nullptr;
 	if (InnerProperty && (
diff --git a/Engine/Source/Programs/UnrealHeaderTool/Private/HeaderParser.cpp b/Engine/Source/Programs/UnrealHeaderTool/Private/HeaderParser.cpp
--- a/Engine/Source/Programs/UnrealHeaderTool/Private/HeaderParser.cpp
+++ b/Engine/Source/Programs/UnrealHeaderTool/Private/HeaderParser.cpp
@@ -1214,7 +1214,7 @@ void AddModuleRelativePathToMetadata(FUnrealTypeDefinitionInfo& Type, TMap<FName
 //
 // Compile an enumeration definition.
 //
-UEnum* FHeaderParser::CompileEnum()
+FUnrealEnumDefinitionInfo& FHeaderParser::CompileEnum()
 {
 	TSharedPtr<FFileScope> Scope = SourceFile.GetScope();
 
@@ -1557,7 +1557,7 @@ UEnum* FHeaderParser::CompileEnum()
 		FError::Throwf(TEXT("Unable to generate enum MAX entry '%s' due to name collision"), *MaxEnumItem.ToString());
 	}
 
-	CheckDocumentationPolicyForEnum(Enum, EnumValueMetaData, EntryMetaData);
+	CheckDocumentationPolicyForEnum(EnumDef, EnumValueMetaData, EntryMetaData);
 
 	// Add the metadata gathered for the enum to the package
 	if (EnumValueMetaData.Num() > 0)
@@ -1573,7 +1573,7 @@ UEnum* FHeaderParser::CompileEnum()
 		UE_LOG_WARNING_UHT(TEXT("'%s' does not have a 0 entry! (This is a problem when the enum is initalized by default)"), *Enum->GetName());
 	}
 
-	return Enum;
+	return EnumDef;
 }
 
 /**
@@ -1990,7 +1990,7 @@ EAccessSpecifier FHeaderParser::ParseAccessProtectionSpecifier(const FToken& Tok
 /**
  * Compile a struct definition.
  */
-UScriptStruct* FHeaderParser::CompileStructDeclaration()
+FUnrealScriptStructDefinitionInfo& FHeaderParser::CompileStructDeclaration()
 {
 	int32 DeclInputLine = InputLine;
 	TSharedPtr<FFileScope> Scope = SourceFile.GetScope();
@@ -2214,7 +2214,7 @@ UScriptStruct* FHeaderParser::CompileStructDeclaration()
 		}
 		else if (Token.Matches(TEXT("RIGVM_METHOD"), ESearchCase::CaseSensitive))
 		{
-			CompileRigVMMethodDeclaration(Struct);
+			CompileRigVMMethodDeclaration(StructDef);
 		}
 		else if (Token.Matches(TEXT("GENERATED_USTRUCT_BODY"), ESearchCase::CaseSensitive) || Token.Matches(TEXT("GENERATED_BODY"), ESearchCase::CaseSensitive))
 		{
@@ -2232,7 +2232,7 @@ UScriptStruct* FHeaderParser::CompileStructDeclaration()
 			Struct->StructMacroDeclaredLineNumber = InputLine;
 			RequireSymbol(TEXT('('), TEXT("'struct'"));
 
-			CompileVersionDeclaration(Struct);
+			CompileVersionDeclaration(StructDef);
 
 			RequireSymbol(TEXT(')'), TEXT("'struct'"));
 
@@ -2366,7 +2366,7 @@ UScriptStruct* FHeaderParser::CompileStructDeclaration()
 	// Link the properties within the struct
 	Struct->StaticLink(true);
 
-	return Struct;
+	return StructDef;
 }
 
 /*-----------------------------------------------------------------------------
@@ -2455,15 +2455,6 @@ void FHeaderParser::CheckAllow( const TCHAR* Thing, ENestAllowFlags AllowFlags )
 	}
 }
 
-bool FHeaderParser::AllowReferenceToClass(UStruct* Scope, UClass* CheckClass) const
-{
-	check(CheckClass);
-
-	return	(Scope->GetOutermost() == CheckClass->GetOutermost())
-		|| ((CheckClass->ClassFlags&CLASS_Parsed) != 0)
-		|| ((CheckClass->ClassFlags&CLASS_Intrinsic) != 0);
-}
-
 /*-----------------------------------------------------------------------------
 	Nest management.
 -----------------------------------------------------------------------------*/
@@ -2565,7 +2556,7 @@ void FHeaderParser::PopNest(ENestType NestType, const TCHAR* Descr)
 	}
 }
 
-void FHeaderParser::FixupDelegateProperties(FUnrealStructDefinitionInfo& StructDef, FScope& Scope, TMap<FName, UFunction*>& DelegateCache )
+void FHeaderParser::FixupDelegateProperties(FUnrealStructDefinitionInfo& StructDef, FScope& Scope, TMap<FName, FUnrealFunctionDefinitionInfo*>& DelegateCache )
 {
 	UStruct* Struct = StructDef.GetStruct();
 
@@ -2591,8 +2582,8 @@ void FHeaderParser::FixupDelegateProperties(FUnrealStructDefinitionInfo& StructD
 			FPropertyBase& DelegatePropertyToken = PropertyDef->GetPropertyBase();
 
 			// attempt to find the delegate function in the map of functions we've already found
-			UFunction* SourceDelegateFunction = DelegateCache.FindRef(DelegatePropertyToken.DelegateName);
-			if (SourceDelegateFunction == nullptr)
+			FUnrealFunctionDefinitionInfo* SourceDelegateFunctionDef = DelegateCache.FindRef(DelegatePropertyToken.DelegateName);
+			if (SourceDelegateFunctionDef == nullptr)
 			{
 				FString NameOfDelegateFunction = DelegatePropertyToken.DelegateName.ToString() + FString( HEADER_GENERATED_DELEGATE_SIGNATURE_SUFFIX );
 				if ( !NameOfDelegateFunction.Contains(TEXT(".")) )
@@ -2600,24 +2591,25 @@ void FHeaderParser::FixupDelegateProperties(FUnrealStructDefinitionInfo& StructD
 					// an unqualified delegate function name - search for a delegate function by this name within the current scope
 					if (FUnrealFieldDefinitionInfo* FoundDef = Scope.FindTypeByName(*NameOfDelegateFunction))
 					{
-						if (FUnrealFunctionDefinitionInfo* SourceDelegateFunctionDef = FoundDef->AsFunction())
-						{
-							SourceDelegateFunction = SourceDelegateFunctionDef->GetFunction();
-						}
+						SourceDelegateFunctionDef = FoundDef->AsFunction();
 					}
-					if (SourceDelegateFunction == nullptr)
+					if (SourceDelegateFunctionDef == nullptr)
 					{
 						// Try to find in other packages.
 						UObject* DelegateSignatureOuter = DelegatePropertyToken.DelegateSignatureOwnerClass 
 							? ((UObject*)DelegatePropertyToken.DelegateSignatureOwnerClass) 
 							: ((UObject*)ANY_PACKAGE);
-						SourceDelegateFunction = FEngineAPI::FindObject<UFunction>(DelegateSignatureOuter, *NameOfDelegateFunction);
+						UFunction* SourceDelegateFunction = FEngineAPI::FindObject<UFunction>(DelegateSignatureOuter, *NameOfDelegateFunction);
 
 						if (SourceDelegateFunction == nullptr)
 						{
 							// convert this into a fully qualified path name for the error message.
 							NameOfDelegateFunction = Scope.GetName().ToString() + TEXT(".") + NameOfDelegateFunction;
 						}
+						else
+						{
+							SourceDelegateFunctionDef = &GTypeDefinitionInfoMap.FindChecked<FUnrealFunctionDefinitionInfo>(SourceDelegateFunction);
+						}
 					}
 				}
 				else
@@ -2645,15 +2637,19 @@ void FHeaderParser::FixupDelegateProperties(FUnrealStructDefinitionInfo& StructD
 					{
 						FError::Throwf(TEXT("Inaccessible type: '%s'"), *DelegateOwnerClassDef->GetClass()->GetPathName());
 					}
-					SourceDelegateFunction = Cast<UFunction>(FindField(*DelegateOwnerClassDef, *DelegateName, false, UFunction::StaticClass(), NULL));
+					UFunction* SourceDelegateFunction = Cast<UFunction>(FindField(*DelegateOwnerClassDef, *DelegateName, false, UFunction::StaticClass(), NULL));
+					if (SourceDelegateFunction != nullptr)
+					{
+						SourceDelegateFunctionDef = &GTypeDefinitionInfoMap.FindChecked<FUnrealFunctionDefinitionInfo>(SourceDelegateFunction);
+					}
 				}
 
-				if ( SourceDelegateFunction == NULL )
+				if (SourceDelegateFunctionDef == NULL )
 				{
 					UngetToken(PropertyDef->GetLineNumber(), PropertyDef->GetParsePosition());
 					FError::Throwf(TEXT("Failed to find delegate function '%s'"), *NameOfDelegateFunction);
 				}
-				else if ( (SourceDelegateFunction->FunctionFlags&FUNC_Delegate) == 0 )
+				else if ( (SourceDelegateFunctionDef->GetFunction()->FunctionFlags&FUNC_Delegate) == 0 )
 				{
 					UngetToken(PropertyDef->GetLineNumber(), PropertyDef->GetParsePosition());
 					FError::Throwf(TEXT("Only delegate functions can be used as the type for a delegate property; '%s' is not a delegate."), *NameOfDelegateFunction);
@@ -2662,10 +2658,9 @@ void FHeaderParser::FixupDelegateProperties(FUnrealStructDefinitionInfo& StructD
 
 			// successfully found the delegate function that this delegate property corresponds to
 
-			// save this into the delegate cache for faster lookup later
-			DelegateCache.Add(DelegatePropertyToken.DelegateName, SourceDelegateFunction);
-
-			FUnrealFunctionDefinitionInfo& SourceDelegateFunctionDef = GTypeDefinitionInfoMap.FindChecked<FUnrealFunctionDefinitionInfo>(SourceDelegateFunction);
+				// save this into the delegate cache for faster lookup later
+			DelegateCache.Add(DelegatePropertyToken.DelegateName, SourceDelegateFunctionDef);
+			UFunction* SourceDelegateFunction = SourceDelegateFunctionDef->GetFunction();
 
 			// bind it to the delegate property
 			if( DelegateProperty != NULL )
@@ -2673,7 +2668,7 @@ void FHeaderParser::FixupDelegateProperties(FUnrealStructDefinitionInfo& StructD
 				if( !SourceDelegateFunction->HasAnyFunctionFlags( FUNC_MulticastDelegate ) )
 				{
 					DelegateProperty->SignatureFunction = SourceDelegateFunction;
-					DelegatePropertyToken.FunctionDef = &SourceDelegateFunctionDef;
+					DelegatePropertyToken.FunctionDef = SourceDelegateFunctionDef;
 				}
 				else
 				{
@@ -2685,11 +2680,11 @@ void FHeaderParser::FixupDelegateProperties(FUnrealStructDefinitionInfo& StructD
 				if( SourceDelegateFunction->HasAnyFunctionFlags( FUNC_MulticastDelegate ) )
 				{
 					MulticastDelegateProperty->SignatureFunction = SourceDelegateFunction;
-					DelegatePropertyToken.FunctionDef = &SourceDelegateFunctionDef;
+					DelegatePropertyToken.FunctionDef = SourceDelegateFunctionDef;
 
 					if(MulticastDelegateProperty->HasAnyPropertyFlags(CPF_BlueprintAssignable | CPF_BlueprintCallable))
 					{
-						for (FUnrealPropertyDefinitionInfo* FuncParamDef : SourceDelegateFunctionDef.GetProperties())
+						for (FUnrealPropertyDefinitionInfo* FuncParamDef : SourceDelegateFunctionDef->GetProperties())
 						{
 							FProperty* FuncParam = FuncParamDef->GetProperty();
 
@@ -2829,9 +2824,10 @@ void FHeaderParser::ValidateClassFlags(const FUnrealClassDefinitionInfo& ToValid
 	}
 }
 
-void FHeaderParser::VerifyBlueprintPropertyGetter(FProperty* Prop, FUnrealFunctionDefinitionInfo* TargetFuncDef)
+void FHeaderParser::VerifyBlueprintPropertyGetter(FUnrealPropertyDefinitionInfo& PropertyDef, FUnrealFunctionDefinitionInfo* TargetFuncDef)
 {
 	check(TargetFuncDef);
+	FProperty* Prop = PropertyDef.GetProperty();
 	UFunction* TargetFunc = TargetFuncDef->GetFunction();
 
 	const TArray<FUnrealPropertyDefinitionInfo*>& Properties = TargetFuncDef->GetProperties();
@@ -2858,9 +2854,10 @@ void FHeaderParser::VerifyBlueprintPropertyGetter(FProperty* Prop, FUnrealFuncti
 	}
 }
 
-void FHeaderParser::VerifyBlueprintPropertySetter(FProperty* Prop, FUnrealFunctionDefinitionInfo* TargetFuncDef)
+void FHeaderParser::VerifyBlueprintPropertySetter(FUnrealPropertyDefinitionInfo& PropertyDef, FUnrealFunctionDefinitionInfo* TargetFuncDef)
 {
 	check(TargetFuncDef);
+	FProperty* Prop = PropertyDef.GetProperty();
 	UFunction* TargetFunc = TargetFuncDef->GetFunction();
 	FUnrealPropertyDefinitionInfo* ReturnPropDef = TargetFuncDef->GetReturn();
 
@@ -2893,8 +2890,9 @@ void FHeaderParser::VerifyBlueprintPropertySetter(FProperty* Prop, FUnrealFuncti
 	}
 }
 
-void FHeaderParser::VerifyRepNotifyCallback(FProperty* Prop, FUnrealFunctionDefinitionInfo* TargetFuncDef)
+void FHeaderParser::VerifyRepNotifyCallback(FUnrealPropertyDefinitionInfo& PropertyDef, FUnrealFunctionDefinitionInfo* TargetFuncDef)
 {
+	FProperty* Prop = PropertyDef.GetProperty();
 	if (TargetFuncDef)
 	{
 		UFunction* TargetFunc = TargetFuncDef->GetFunction();
@@ -2977,7 +2975,7 @@ void FHeaderParser::VerifyPropertyMarkups(FUnrealClassDefinitionInfo& TargetClas
 
 		if (Prop->HasAnyPropertyFlags(CPF_RepNotify))
 		{
-			VerifyRepNotifyCallback(Prop, FindTargetFunction(Prop->RepNotifyFunc));
+			VerifyRepNotifyCallback(*PropertyDef, FindTargetFunction(Prop->RepNotifyFunc));
 		}
 
 		if (Prop->HasAnyPropertyFlags(CPF_BlueprintVisible))
@@ -2987,7 +2985,7 @@ void FHeaderParser::VerifyPropertyMarkups(FUnrealClassDefinitionInfo& TargetClas
 			{
 				if (FUnrealFunctionDefinitionInfo* TargetFuncDef = FindTargetFunction(*GetterFuncName))
 				{
-					VerifyBlueprintPropertyGetter(Prop, TargetFuncDef);
+					VerifyBlueprintPropertyGetter(*PropertyDef, TargetFuncDef);
 				}
 				else
 				{
@@ -3003,7 +3001,7 @@ void FHeaderParser::VerifyPropertyMarkups(FUnrealClassDefinitionInfo& TargetClas
 				{
 					if (FUnrealFunctionDefinitionInfo* TargetFuncDef = FindTargetFunction(*SetterFuncName))
 					{
-						VerifyBlueprintPropertySetter(Prop, TargetFuncDef);
+						VerifyBlueprintPropertySetter(*PropertyDef, TargetFuncDef);
 					}
 					else
 					{
@@ -5097,7 +5095,7 @@ bool FHeaderParser::CompileDeclaration(TArray<FUnrealFunctionDefinitionInfo*>& D
 			FError::Throwf(TEXT("%s must occur inside the native interface definition"), Token.Identifier);
 		}
 		RequireSymbol(TEXT('('), Token.Identifier);
-		CompileVersionDeclaration(GetCurrentClass());
+		CompileVersionDeclaration(GetCurrentClassDef());
 		RequireSymbol(TEXT(')'), Token.Identifier);
 
 		FStructMetaData& StructData = GetCurrentClassData();
@@ -5140,7 +5138,7 @@ bool FHeaderParser::CompileDeclaration(TArray<FUnrealFunctionDefinitionInfo*>& D
 			FError::Throwf(TEXT("%s must occur inside the interface definition"), Token.Identifier);
 		}
 		RequireSymbol(TEXT('('), Token.Identifier);
-		CompileVersionDeclaration(GetCurrentClass());
+		CompileVersionDeclaration(GetCurrentClassDef());
 		RequireSymbol(TEXT(')'), Token.Identifier);
 
 		FStructMetaData& StructData = GetCurrentClassData();
@@ -5183,7 +5181,7 @@ bool FHeaderParser::CompileDeclaration(TArray<FUnrealFunctionDefinitionInfo*>& D
 		}
 
 		RequireSymbol(TEXT('('), Token.Identifier);
-		CompileVersionDeclaration(GetCurrentClass());
+		CompileVersionDeclaration(GetCurrentClassDef());
 		RequireSymbol(TEXT(')'), Token.Identifier);
 
 		StructData.SetGeneratedBodyLine(InputLine);
@@ -5705,7 +5703,7 @@ void PostParsingClassSetup(FUnrealClassDefinitionInfo& ClassDef)
 /**
  * Compiles a class declaration.
  */
-UClass* FHeaderParser::CompileClassDeclaration()
+FUnrealClassDefinitionInfo& FHeaderParser::CompileClassDeclaration()
 {
 	// Start of a class block.
 	CheckAllow(TEXT("'class'"), ENestAllowFlags::Class);
@@ -5828,7 +5826,7 @@ UClass* FHeaderParser::CompileClassDeclaration()
 	// Check that the class has appropriate class flags set
 	ValidateClassFlags(ClassDef);
 
-	return Class;
+	return ClassDef;
 }
 
 FUnrealClassDefinitionInfo* FHeaderParser::ParseInterfaceNameDeclaration(FString& DeclaredInterfaceName, FString& RequiredAPIMacroIfPresent)
@@ -6023,7 +6021,7 @@ void FHeaderParser::CompileInterfaceDeclaration()
 	PushNest(ENestType::Interface, ClassDef);
 }
 
-void FHeaderParser::CompileRigVMMethodDeclaration(UStruct* Struct)
+void FHeaderParser::CompileRigVMMethodDeclaration(FUnrealStructDefinitionInfo& StructDef)
 {
 	if (!MatchSymbol(TEXT("(")))
 	{
@@ -6139,8 +6137,8 @@ void FHeaderParser::CompileRigVMMethodDeclaration(UStruct* Struct)
 		}
 	}
 
-	FRigVMStructInfo& StructRigVMInfo = StructRigVMMap.FindOrAdd(Struct);
-	StructRigVMInfo.Name = Struct->GetName();
+	FRigVMStructInfo& StructRigVMInfo = StructRigVMMap.FindOrAdd(&StructDef);
+	StructRigVMInfo.Name = StructDef.GetStruct()->GetName();
 	StructRigVMInfo.Methods.Add(MethodInfo);
 }
 
@@ -6162,7 +6160,7 @@ const TCHAR* FHeaderParser::GetDynamicArrayText = TEXT("GetDynamicArray");
 void FHeaderParser::ParseRigVMMethodParameters(FUnrealStructDefinitionInfo& StructDef)
 {
 	UStruct* Struct = StructDef.GetStruct();
-	FRigVMStructInfo* StructRigVMInfo = StructRigVMMap.Find(Struct);
+	FRigVMStructInfo* StructRigVMInfo = StructRigVMMap.Find(&StructDef);
 	if (StructRigVMInfo == nullptr)
 	{
 		return;
@@ -7170,7 +7168,7 @@ void FHeaderParser::CompileFunctionDeclaration()
 	}
 
 	// perform documentation policy tests
-	CheckDocumentationPolicyForFunc(GetCurrentClass(), FuncDef);
+	CheckDocumentationPolicyForFunc(UHTCastChecked<FUnrealClassDefinitionInfo>(GetCurrentClassDef()), FuncDef);
 }
 
 /** Parses optional metadata text. */
@@ -7850,7 +7848,7 @@ ECompilationResult::Type FHeaderParser::ParseHeader()
 				FUnrealFieldDefinitionInfo* TypeDef = *ScopeTypeIterator;
 				if (TypeDef->AsScriptStruct() != nullptr || TypeDef->AsClass() != nullptr)
 				{
-					TMap<FName, UFunction*> DelegateCache;
+					TMap<FName, FUnrealFunctionDefinitionInfo*> DelegateCache;
 					FixupDelegateProperties(TypeDef->AsStructChecked(), *TypeDef->GetScope(), DelegateCache);
 				}
 			}
@@ -7858,7 +7856,7 @@ ECompilationResult::Type FHeaderParser::ParseHeader()
 
 		// Fix up any delegates themselves, if they refer to other delegates
 		{
-			TMap<FName, UFunction*> DelegateCache;
+			TMap<FName, FUnrealFunctionDefinitionInfo*> DelegateCache;
 			for (FUnrealFunctionDefinitionInfo* DelegateDef : DelegatesToFixup)
 			{
 				FixupDelegateProperties(*DelegateDef, SourceFile.GetScope().Get(), DelegateCache);
@@ -8997,7 +8995,7 @@ bool FHeaderParser::TryToMatchConstructorParameterList(FToken Token)
 	return true;
 }
 
-void FHeaderParser::CompileVersionDeclaration(UStruct* Struct)
+void FHeaderParser::CompileVersionDeclaration(FUnrealStructDefinitionInfo& StructDef)
 {
 	// Do nothing if we're at the end of file.
 	FToken Token;
@@ -9018,7 +9016,7 @@ void FHeaderParser::CompileVersionDeclaration(UStruct* Struct)
 	if (Token.TokenType == ETokenType::TOKEN_Symbol
 		&& Token.Matches(TEXT(')')))
 	{
-		SourceFile.GetGeneratedCodeVersions().FindOrAdd(Struct) = Version;
+		SourceFile.GetGeneratedCodeVersions().FindOrAdd(&StructDef) = Version;
 		UngetToken(Token);
 		return;
 	}
@@ -9026,7 +9024,7 @@ void FHeaderParser::CompileVersionDeclaration(UStruct* Struct)
 	// Overwrite with version specified by macro.
 	Version = ToGeneratedCodeVersion(Token.Identifier);
 
-	SourceFile.GetGeneratedCodeVersions().FindOrAdd(Struct) = Version;
+	SourceFile.GetGeneratedCodeVersions().FindOrAdd(&StructDef) = Version;
 }
 
 void FHeaderParser::ResetClassData()
@@ -9192,7 +9190,7 @@ void FHeaderParser::PostPopFunctionDeclaration(FUnrealFunctionDefinitionInfo& Po
 	if (!GetCurrentScope()->IsFileScope() && GetCurrentClassDef().ContainsDelegates())
 	{
 		// now validate all delegate variables declared in the class
-		TMap<FName, UFunction*> DelegateCache;
+		TMap<FName, FUnrealFunctionDefinitionInfo*> DelegateCache;
 		FixupDelegateProperties(PoppedFunctionDef, *GetCurrentScope(), DelegateCache);
 	}
 }
@@ -9202,7 +9200,7 @@ void FHeaderParser::PostPopNestInterface(FUnrealClassDefinitionInfo& CurrentInte
 	if (CurrentInterfaceDef.ContainsDelegates())
 	{
 		UClass* CurrentInterface = CurrentInterfaceDef.GetClass();
-		TMap<FName, UFunction*> DelegateCache;
+		TMap<FName, FUnrealFunctionDefinitionInfo*> DelegateCache;
 		FixupDelegateProperties(CurrentInterfaceDef, *CurrentInterfaceDef.GetScope(), DelegateCache);
 	}
 }
@@ -9226,26 +9224,24 @@ FDocumentationPolicy FHeaderParser::GetDocumentationPolicyFromName(const FString
 }
 
 
-FDocumentationPolicy FHeaderParser::GetDocumentationPolicyForStruct(UStruct* Struct)
+FDocumentationPolicy FHeaderParser::GetDocumentationPolicyForStruct(FUnrealStructDefinitionInfo& StructDef)
 {
 	SCOPE_SECONDS_COUNTER_UHT(DocumentationPolicy);
 
-	check(Struct!= nullptr);
-
 	FDocumentationPolicy DocumentationPolicy;
 	FString DocumentationPolicyName;
-	if (Struct->GetStringMetaDataHierarchical(NAME_DocumentationPolicy, &DocumentationPolicyName))
+	if (StructDef.GetStruct()->GetStringMetaDataHierarchical(NAME_DocumentationPolicy, &DocumentationPolicyName))
 	{
 		DocumentationPolicy = GetDocumentationPolicyFromName(DocumentationPolicyName);
 	}
 	return DocumentationPolicy;
 }
 
-void FHeaderParser::CheckDocumentationPolicyForEnum(UEnum* Enum, const TMap<FName, FString>& MetaData, const TArray<TMap<FName, FString>>& Entries)
+void FHeaderParser::CheckDocumentationPolicyForEnum(FUnrealEnumDefinitionInfo& EnumDef, const TMap<FName, FString>& MetaData, const TArray<TMap<FName, FString>>& Entries)
 {
 	SCOPE_SECONDS_COUNTER_UHT(DocumentationPolicy);
 
-	check(Enum != nullptr);
+	UEnum* Enum = EnumDef.GetEnum();
 
 	const FString* DocumentationPolicyName = MetaData.Find(NAME_DocumentationPolicy);
 	if (DocumentationPolicyName == nullptr)
@@ -9297,7 +9293,7 @@ void FHeaderParser::CheckDocumentationPolicyForStruct(FUnrealStructDefinitionInf
 	UStruct* Struct = StructDef.GetStruct();
 	check(Struct != nullptr);
 
-	FDocumentationPolicy DocumentationPolicy = GetDocumentationPolicyForStruct(Struct);
+	FDocumentationPolicy DocumentationPolicy = GetDocumentationPolicyForStruct(StructDef);
 	if (DocumentationPolicy.bClassOrStructCommentRequired)
 	{
 		FString ClassTooltip;
@@ -9384,15 +9380,14 @@ bool FHeaderParser::DoesCPPTypeRequireDocumentation(const FString& CPPType)
 }
 
 // Validates the documentation for a given method
-void FHeaderParser::CheckDocumentationPolicyForFunc(UClass* Class, FUnrealFunctionDefinitionInfo& FunctionDef)
+void FHeaderParser::CheckDocumentationPolicyForFunc(FUnrealClassDefinitionInfo& ClassDef, FUnrealFunctionDefinitionInfo& FunctionDef)
 {
 	SCOPE_SECONDS_COUNTER_UHT(DocumentationPolicy);
 
+	UClass* Class = ClassDef.GetClass();
 	UFunction* Func = FunctionDef.GetFunction();
-	check(Class != nullptr);
-	check(Func != nullptr);
 
-	FDocumentationPolicy DocumentationPolicy = GetDocumentationPolicyForStruct(Class);
+	FDocumentationPolicy DocumentationPolicy = GetDocumentationPolicyForStruct(ClassDef);
 	if (DocumentationPolicy.bFunctionToolTipsRequired)
 	{
 		const FString* FunctionTooltip = Func->FindMetaData(NAME_ToolTip);
diff --git a/Engine/Source/Programs/UnrealHeaderTool/Private/HeaderParser.h b/Engine/Source/Programs/UnrealHeaderTool/Private/HeaderParser.h
--- a/Engine/Source/Programs/UnrealHeaderTool/Private/HeaderParser.h
+++ b/Engine/Source/Programs/UnrealHeaderTool/Private/HeaderParser.h
@@ -18,13 +18,7 @@ class UPackage;
 struct FManifestModule;
 class IScriptGeneratorPluginInterface;
 class FStringOutputDevice;
-class FProperty;
 class FUnrealSourceFile;
-class UFunction;
-class UEnum;
-class UScriptStruct;
-class UDelegateFunction;
-class UStruct;
 class FScope;
 class FHeaderProvider;
 
@@ -373,7 +367,7 @@ struct FRigVMStructInfo
 	TArray<FRigVMMethodInfo> Methods;
 };
 
-typedef TMap<UStruct*, FRigVMStructInfo> FRigVMStructMap;
+typedef TMap<FUnrealStructDefinitionInfo*, FRigVMStructInfo> FRigVMStructMap;
 
 struct ClassDefinitionRange
 {
@@ -699,8 +693,8 @@ protected:
 	ECompilationResult::Type ParseHeader();
 	void CompileDirective();
 	void FinalizeScriptExposedFunctions(FUnrealClassDefinitionInfo& ClassDef);
-	UEnum* CompileEnum();
-	UScriptStruct* CompileStructDeclaration();
+	FUnrealEnumDefinitionInfo& CompileEnum();
+	FUnrealScriptStructDefinitionInfo& CompileStructDeclaration();
 	bool CompileDeclaration(TArray<FUnrealFunctionDefinitionInfo*>& DelegatesToFixup, FToken& Token);
 
 	/** Skip C++ (noexport) declaration. */
@@ -735,12 +729,12 @@ protected:
 	template<typename T>
 	FUnrealFunctionDefinitionInfo& CreateDelegateFunction(FFuncInfo&& FuncInfo) const;
 
-	UClass* CompileClassDeclaration();
+	FUnrealClassDefinitionInfo& CompileClassDeclaration();
 	FUnrealFunctionDefinitionInfo& CompileDelegateDeclaration(const TCHAR* DelegateIdentifier, EDelegateSpecifierAction::Type SpecifierAction = EDelegateSpecifierAction::DontParse);
 	void CompileFunctionDeclaration();
 	void CompileVariableDeclaration (FUnrealStructDefinitionInfo& StructDef);
 	void CompileInterfaceDeclaration();
-	void CompileRigVMMethodDeclaration(UStruct* Struct);
+	void CompileRigVMMethodDeclaration(FUnrealStructDefinitionInfo& StructDef);
 	void ParseRigVMMethodParameters(FUnrealStructDefinitionInfo& StructDef);
 
 	FUnrealClassDefinitionInfo* ParseInterfaceNameDeclaration(FString& DeclaredInterfaceName, FString& RequiredAPIMacroIfPresent);
@@ -795,16 +789,6 @@ protected:
 		EVariableCategory::Type VariableCategory,
 		ELayoutMacroType LayoutMacroType = ELayoutMacroType::None);
 	
-	/**
-	 * Returns whether the specified class can be referenced from the class currently being compiled.
-	 *
-	 * @param	Scope		The scope we are currently parsing.
-	 * @param	CheckClass	The class we want to reference.
-	 *
-	 * @return	true if the specified class is an intrinsic type or if the class has successfully been parsed
-	 */
-	bool AllowReferenceToClass(UStruct* Scope, UClass* CheckClass) const;
-
 	/**
 	 * Parses optional metadata text.
 	 *
@@ -878,7 +862,7 @@ protected:
 	 * @param	Scope				the current scope
 	 * @param	DelegateCache		cached map of delegates that have already been found; used for faster lookup.
 	 */
-	void FixupDelegateProperties(FUnrealStructDefinitionInfo& StructDef, FScope& Scope, TMap<FName, UFunction*>& DelegateCache);
+	void FixupDelegateProperties(FUnrealStructDefinitionInfo& StructDef, FScope& Scope, TMap<FName, FUnrealFunctionDefinitionInfo*>& DelegateCache);
 
 	// Retry functions.
 	void InitScriptLocation( FScriptLocation& Retry );
@@ -920,25 +904,25 @@ private:
 	bool TryToMatchConstructorParameterList(FToken Token);
 
 	// Parses possible version declaration in generated code, e.g. GENERATED_BODY(<some_version>).
-	void CompileVersionDeclaration(UStruct* Struct);
+	void CompileVersionDeclaration(FUnrealStructDefinitionInfo& StructDef);
 
 	// Verifies that all specified class's UProperties with function associations have valid targets
 	void VerifyPropertyMarkups(FUnrealClassDefinitionInfo& TargetClassDef);
 
 	// Verifies the target function meets the criteria for a blueprint property getter
-	void VerifyBlueprintPropertyGetter(FProperty* Property, FUnrealFunctionDefinitionInfo* TargetFuncDef);
+	void VerifyBlueprintPropertyGetter(FUnrealPropertyDefinitionInfo& PropertyDef, FUnrealFunctionDefinitionInfo* TargetFuncDef);
 
 	// Verifies the target function meets the criteria for a blueprint property setter
-	void VerifyBlueprintPropertySetter(FProperty* Property, FUnrealFunctionDefinitionInfo* TargetFuncDef);
+	void VerifyBlueprintPropertySetter(FUnrealPropertyDefinitionInfo& PropertyDef, FUnrealFunctionDefinitionInfo* TargetFuncDef);
 
 	// Verifies the target function meets the criteria for a replication notify callback
-	void VerifyRepNotifyCallback(FProperty* Property, FUnrealFunctionDefinitionInfo* TargetFuncDef);
+	void VerifyRepNotifyCallback(FUnrealPropertyDefinitionInfo& PropertyDef, FUnrealFunctionDefinitionInfo* TargetFuncDef);
 
 	// Constructs the policy from a string
 	static FDocumentationPolicy GetDocumentationPolicyFromName(const FString& PolicyName);
 
 	// Constructs the policy for documentation checks for a given struct
-	static FDocumentationPolicy GetDocumentationPolicyForStruct(UStruct* Struct);
+	static FDocumentationPolicy GetDocumentationPolicyForStruct(FUnrealStructDefinitionInfo& StructDef);
 
 	// Property types to provide UI Min and Max ranges
 	static TArray<FString> PropertyCPPTypesRequiringUIRanges;
@@ -947,13 +931,13 @@ private:
 	static bool DoesCPPTypeRequireDocumentation(const FString& CPPType);
 
 	// Validates the documentation for a given enum
-	void CheckDocumentationPolicyForEnum(UEnum* Enum, const TMap<FName, FString>& MetaData, const TArray<TMap<FName, FString>>& Entries);
+	void CheckDocumentationPolicyForEnum(FUnrealEnumDefinitionInfo& EnumDef, const TMap<FName, FString>& MetaData, const TArray<TMap<FName, FString>>& Entries);
 
 	// Validates the documentation for a given struct
 	void CheckDocumentationPolicyForStruct(FUnrealStructDefinitionInfo& StructDef);
 
 	// Validates the documentation for a given method
-	void CheckDocumentationPolicyForFunc(UClass* Class, FUnrealFunctionDefinitionInfo& FunctionDef);
+	void CheckDocumentationPolicyForFunc(FUnrealClassDefinitionInfo& ClassDef, FUnrealFunctionDefinitionInfo& FunctionDef);
 
 	// Checks if a valid range has been found on the provided metadata
 	bool CheckUIMinMaxRangeFromMetaData(const FString& UIMin, const FString& UIMax);
diff --git a/Engine/Source/Programs/UnrealHeaderTool/Private/NativeClassExporter.h b/Engine/Source/Programs/UnrealHeaderTool/Private/NativeClassExporter.h
--- a/Engine/Source/Programs/UnrealHeaderTool/Private/NativeClassExporter.h
+++ b/Engine/Source/Programs/UnrealHeaderTool/Private/NativeClassExporter.h
@@ -6,16 +6,7 @@
 #include "Async/TaskGraphInterfaces.h"
 
 class FUnrealSourceFile;
-class UPackage;
-class FProperty;
-class UFunction;
-class UStruct;
 class UField;
-class UClass;
-class UEnum;
-class UScriptStruct;
-class UDelegateFunction;
-class FStructMetaData;
 class FOutputDevice;
 class FUnrealPackageDefinitionInfo;
 class FUnrealPropertyDefinitionInfo;
@@ -396,7 +387,7 @@ private:
 	 * Returns a string in the format CLASS_Something|CLASS_Something which represents all class flags that are set for the specified
 	 * class which need to be exported as part of the DECLARE_CLASS macro
 	 */
-	static FString GetClassFlagExportText( UClass* Class );
+	static FString GetClassFlagExportText(FUnrealClassDefinitionInfo& ClassDef);
 
 	/**
 	 * Exports the header text for the enum specified
@@ -638,7 +629,7 @@ private:
 	 * @param	Prop			the property that is being exported
 	 * @param	PropertyText	the string containing the text exported from ExportCppDeclaration
 	 */
-	static void ApplyAlternatePropertyExportText(FProperty* Prop, FUHTStringBuilder& PropertyText, EExportingState ExportingState);
+	static void ApplyAlternatePropertyExportText(FUnrealPropertyDefinitionInfo& PropertyDef, FUHTStringBuilder& PropertyText, EExportingState ExportingState);
 
 	/**
 	* Create a temp header file name from the header name
diff --git a/Engine/Source/Programs/UnrealHeaderTool/Private/UnrealSourceFile.cpp b/Engine/Source/Programs/UnrealHeaderTool/Private/UnrealSourceFile.cpp
--- a/Engine/Source/Programs/UnrealHeaderTool/Private/UnrealSourceFile.cpp
+++ b/Engine/Source/Programs/UnrealHeaderTool/Private/UnrealSourceFile.cpp
@@ -128,9 +128,9 @@ const FString& FUnrealSourceFile::GetContent() const
 	return Content;
 }
 
-EGeneratedCodeVersion FUnrealSourceFile::GetGeneratedCodeVersionForStruct(UStruct* Struct) const
+EGeneratedCodeVersion FUnrealSourceFile::GetGeneratedCodeVersionForStruct(FUnrealStructDefinitionInfo& StructDef) const
 {
-	if (const EGeneratedCodeVersion* Version = GeneratedCodeVersions.Find(Struct))
+	if (const EGeneratedCodeVersion* Version = GeneratedCodeVersions.Find(&StructDef))
 	{
 		return *Version;
 	}
diff --git a/Engine/Source/Programs/UnrealHeaderTool/Private/UnrealSourceFile.h b/Engine/Source/Programs/UnrealHeaderTool/Private/UnrealSourceFile.h
--- a/Engine/Source/Programs/UnrealHeaderTool/Private/UnrealSourceFile.h
+++ b/Engine/Source/Programs/UnrealHeaderTool/Private/UnrealSourceFile.h
@@ -291,12 +291,12 @@ public:
 	/**
 	 * Gets generated code version for given UStruct.
 	 */
-	EGeneratedCodeVersion GetGeneratedCodeVersionForStruct(UStruct* Struct) const;
+	EGeneratedCodeVersion GetGeneratedCodeVersionForStruct(FUnrealStructDefinitionInfo& StructDef) const;
 
 	/**
 	 * Gets generated code versions.
 	 */
-	TMap<UStruct*, EGeneratedCodeVersion>& GetGeneratedCodeVersions()
+	TMap<FUnrealStructDefinitionInfo*, EGeneratedCodeVersion>& GetGeneratedCodeVersions()
 	{
 		return GeneratedCodeVersions;
 	}
@@ -304,7 +304,7 @@ public:
 	/**
 	 * Gets generated code versions. Const version.
 	 */
-	const TMap<UStruct*, EGeneratedCodeVersion>& GetGeneratedCodeVersions() const
+	const TMap<FUnrealStructDefinitionInfo*, EGeneratedCodeVersion>& GetGeneratedCodeVersions() const
 	{
 		return GeneratedCodeVersions;
 	}
@@ -519,8 +519,8 @@ private:
 	// List of structs defined in this source file along with parsing info.
 	TArray<TSharedRef<FUnrealTypeDefinitionInfo>> DefinedStructs;
 
-	// Mapping of UStructs to versions, according to which their code should be generated.
-	TMap<UStruct*, EGeneratedCodeVersion> GeneratedCodeVersions;
+	// Mapping of FUnrealStructDefinitionInfo to versions, according to which their code should be generated.
+	TMap<FUnrealStructDefinitionInfo*, EGeneratedCodeVersion> GeneratedCodeVersions;
 
 	// Collection of all singletons found during code generation
 	TArray<FUnrealFieldDefinitionInfo*> Singletons;
