Removed the global maps for the code generation version, RigVM, and declaration ranges.

#rnx

diff --git a/Engine/Source/Programs/UnrealHeaderTool/Private/CodeGenerator.cpp b/Engine/Source/Programs/UnrealHeaderTool/Private/CodeGenerator.cpp
--- a/Engine/Source/Programs/UnrealHeaderTool/Private/CodeGenerator.cpp
+++ b/Engine/Source/Programs/UnrealHeaderTool/Private/CodeGenerator.cpp
@@ -263,7 +263,6 @@ namespace
 	static const FString STRING_GetLifetimeReplicatedPropsStr(TEXT("GetLifetimeReplicatedProps"));
 
 	static void WriteReplicatedMacroData(
-		const ClassDefinitionRange& ClassRange,
 		const TCHAR* ClassCPPName,
 		const TCHAR* API,
 		FUnrealClassDefinitionInfo& ClassDef,
@@ -273,12 +272,12 @@ namespace
 		EExportClassOutFlags& OutFlags)
 	{
 		UClass* Class = ClassDef.GetClass();
-		const bool bHasGetLifetimeReplicatedProps = HasIdentifierExactMatch(ClassRange.Start, ClassRange.End, STRING_GetLifetimeReplicatedPropsStr);
+		const bool bHasGetLifetimeReplicatedProps = HasIdentifierExactMatch(ClassDef.GetDefinitionRange().Start, ClassDef.GetDefinitionRange().End, STRING_GetLifetimeReplicatedPropsStr);
 
 		if (!bHasGetLifetimeReplicatedProps)
 		{
 			// Default version autogenerates declarations.
-			if (SourceFile.GetGeneratedCodeVersionForStruct(ClassDef) == EGeneratedCodeVersion::V1)
+			if (ClassDef.GetGeneratedCodeVersion() == EGeneratedCodeVersion::V1)
 			{
 				Writer.Logf(TEXT("\tvoid GetLifetimeReplicatedProps(TArray<FLifetimeProperty>& OutLifetimeProps) const override;\r\n"));
 			}
@@ -2876,12 +2875,7 @@ void FNativeClassHeaderGenerator::ExportClassFromSourceFileInner(
 
 	FString PPOMacroName;
 
-	ClassDefinitionRange ClassRange;
-	if (ClassDefinitionRange* FoundRange = ClassDefinitionRanges.Find(Class))
-	{
-		ClassRange = *FoundRange;
-		ClassRange.Validate();
-	}
+	ClassDef.GetDefinitionRange().Validate();
 
 	FString GeneratedSerializeFunctionCPP;
 	FString GeneratedSerializeFunctionHeaderMacroName;
@@ -3030,7 +3024,7 @@ void FNativeClassHeaderGenerator::ExportClassFromSourceFileInner(
 
 			if (ClassHasReplicatedProperties(ClassDef))
 			{
-				WriteReplicatedMacroData(ClassRange, *ClassCPPName, *APIArg, ClassDef, SuperClass, InterfaceBoilerplate, SourceFile, OutFlags);
+				WriteReplicatedMacroData(*ClassCPPName, *APIArg, ClassDef, SuperClass, InterfaceBoilerplate, SourceFile, OutFlags);
 			}
 
 			FString NoPureDeclsMacroName = SourceFile.GetGeneratedMacroName(StructData, TEXT("_INCLASS_IINTERFACE_NO_PURE_DECLS"));
@@ -3057,7 +3051,7 @@ void FNativeClassHeaderGenerator::ExportClassFromSourceFileInner(
 
 			if (ClassHasReplicatedProperties(ClassDef))
 			{
-				WriteReplicatedMacroData(ClassRange, *ClassCPPName, *APIArg, ClassDef, SuperClass, Boilerplate, SourceFile, OutFlags);
+				WriteReplicatedMacroData(*ClassCPPName, *APIArg, ClassDef, SuperClass, Boilerplate, SourceFile, OutFlags);
 			}
 
 			{
@@ -3483,15 +3477,15 @@ void FNativeClassHeaderGenerator::ExportGeneratedStructBodyMacros(FOutputDevice&
 	// for RigVM methods we need to generated a macro used for implementing the static method
 	// and prepare two prologs: one for the virtual function implementation, and one for the stub
 	// invoking the static method.
-	const FRigVMStructInfo* StructRigVMInfo = FHeaderParser::StructRigVMMap.Find(&ScriptStructDef);
-	if(StructRigVMInfo)
+	const FRigVMStructInfo& StructRigVMInfo = ScriptStructDef.GetRigVMInfo();
+	if(StructRigVMInfo.bHasRigVM)
 	{
 		//RigVMStubProlog.Add(FString::Printf(TEXT("ensure(RigVMOperandMemory.Num() == %d);"), StructRigVMInfo->Members.Num()));
 
 		int32 OperandIndex = 0;
-		for (int32 ParameterIndex = 0; ParameterIndex < StructRigVMInfo->Members.Num(); ParameterIndex++)
+		for (int32 ParameterIndex = 0; ParameterIndex < StructRigVMInfo.Members.Num(); ParameterIndex++)
 		{
-			const FRigVMParameter& Parameter = StructRigVMInfo->Members[ParameterIndex];
+			const FRigVMParameter& Parameter = StructRigVMInfo.Members[ParameterIndex];
 			if(Parameter.RequiresCast())
 			{
 				if (Parameter.IsArray() && !Parameter.IsConst() && !Parameter.ArraySize.IsEmpty())
@@ -3595,10 +3589,10 @@ void FNativeClassHeaderGenerator::ExportGeneratedStructBodyMacros(FOutputDevice&
 			}
 		}
 
-		FString StructMembers = StructRigVMInfo->Members.Declarations(false, TEXT(", \\\r\n\t\t"), true, false);
+		FString StructMembers = StructRigVMInfo.Members.Declarations(false, TEXT(", \\\r\n\t\t"), true, false);
 
 		OutGeneratedHeaderText.Log(TEXT("\n"));
-		for (const FRigVMMethodInfo& MethodInfo : StructRigVMInfo->Methods)
+		for (const FRigVMMethodInfo& MethodInfo : StructRigVMInfo.Methods)
 		{
 			FString ParameterSuffix = MethodInfo.Parameters.Declarations(true, TEXT(", \\\r\n\t\t"));
 			FString RigVMParameterPrefix2 = RigVMParameterPrefix + FString((StructMembers.IsEmpty() && ParameterSuffix.IsEmpty()) ? TEXT("") : TEXT(", \\\r\n\t\t"));
@@ -3621,12 +3615,12 @@ void FNativeClassHeaderGenerator::ExportGeneratedStructBodyMacros(FOutputDevice&
 		// if we have RigVM methods on this struct we need to 
 		// declare the static method as well as the stub method
 		FString RigVMMethodsDeclarations;
-		if (StructRigVMInfo)
+		if (StructRigVMInfo.bHasRigVM)
 		{
-			FString StructMembers = StructRigVMInfo->Members.Declarations(false, TEXT(",\r\n\t\t"), true, false);
-			for (const FRigVMMethodInfo& MethodInfo : StructRigVMInfo->Methods)
+			FString StructMembers = StructRigVMInfo.Members.Declarations(false, TEXT(",\r\n\t\t"), true, false);
+			for (const FRigVMMethodInfo& MethodInfo : StructRigVMInfo.Methods)
 			{
-				FString StructMembersForStub = StructRigVMInfo->Members.Names(false, TEXT(",\r\n\t\t\t"), false);
+				FString StructMembersForStub = StructRigVMInfo.Members.Names(false, TEXT(",\r\n\t\t\t"), false);
 				FString ParameterSuffix = MethodInfo.Parameters.Declarations(true, TEXT(",\r\n\t\t"));
 				FString ParameterNamesSuffix = MethodInfo.Parameters.Names(true, TEXT(",\r\n\t\t\t"));
 				FString RigVMParameterPrefix2 = RigVMParameterPrefix + FString((StructMembers.IsEmpty() && ParameterSuffix.IsEmpty()) ? TEXT("") : TEXT(",\r\n\t\t"));
@@ -3661,7 +3655,7 @@ void FNativeClassHeaderGenerator::ExportGeneratedStructBodyMacros(FOutputDevice&
 				RigVMMethodsDeclarations += FString::Printf(TEXT("\t}\r\n"));
 			}
 
-			for (const FRigVMParameter& StructMember : StructRigVMInfo->Members)
+			for (const FRigVMParameter& StructMember : StructRigVMInfo.Members)
 			{
 				if (!StructMember.ArraySize.IsEmpty())
 				{
@@ -3724,9 +3718,9 @@ void FNativeClassHeaderGenerator::ExportGeneratedStructBodyMacros(FOutputDevice&
 
 		// if this struct has RigVM methods - we need to register the method to our central
 		// registry on construction of the static struct
-		if (StructRigVMInfo)
+		if (StructRigVMInfo.bHasRigVM)
 		{
-			for (const FRigVMMethodInfo& MethodInfo : StructRigVMInfo->Methods)
+			for (const FRigVMMethodInfo& MethodInfo : StructRigVMInfo.Methods)
 			{
 				Out.Logf(TEXT("\t\tFRigVMRegistry::Get().Register(TEXT(\"%s::%s\"), &%s::RigVM%s, Singleton);\r\n"),
 					*StructNameCPP, *MethodInfo.Name, *StructNameCPP, *MethodInfo.Name);
@@ -3915,11 +3909,11 @@ void FNativeClassHeaderGenerator::ExportGeneratedStructBodyMacros(FOutputDevice&
 	// if this struct has RigVM methods we need to implement both the 
 	// virtual function as well as the stub method here.
 	// The static method is implemented by the user using a macro.
-	if (StructRigVMInfo)
+	if (StructRigVMInfo.bHasRigVM)
 	{
-		FString StructMembersForVirtualFunc = StructRigVMInfo->Members.Names(false, TEXT(",\r\n\t\t"), true);
+		FString StructMembersForVirtualFunc = StructRigVMInfo.Members.Names(false, TEXT(",\r\n\t\t"), true);
 
-		for (const FRigVMMethodInfo& MethodInfo : StructRigVMInfo->Methods)
+		for (const FRigVMMethodInfo& MethodInfo : StructRigVMInfo.Methods)
 		{
 			Out.Log(TEXT("\r\n"));
 
@@ -3959,7 +3953,7 @@ void FNativeClassHeaderGenerator::ExportGeneratedStructBodyMacros(FOutputDevice&
 			Out.Log(TEXT("\r\n"));
 
 		bool bHasGetArraySize = false;
-		for (const FRigVMParameter& StructMember : StructRigVMInfo->Members)
+		for (const FRigVMParameter& StructMember : StructRigVMInfo.Members)
 			{
 			if (!StructMember.ArraySize.IsEmpty())
 				{
@@ -3972,7 +3966,7 @@ void FNativeClassHeaderGenerator::ExportGeneratedStructBodyMacros(FOutputDevice&
 			{
 			Out.Logf(TEXT("int32 %s::GetArraySize(const FName& InMemberName, const FRigVMUserDataArray& Context)\r\n"), *StructNameCPP);
 			Out.Log(TEXT("{\r\n"));
-			for (const FRigVMParameter& StructMember : StructRigVMInfo->Members)
+			for (const FRigVMParameter& StructMember : StructRigVMInfo.Members)
 			{
 				if (!StructMember.ArraySize.IsEmpty())
 				{
@@ -5029,16 +5023,13 @@ void FNativeClassHeaderGenerator::ExportFunctionThunk(FUHTStringBuilder& RPCWrap
 	RPCWrappers += TEXT("\t\tP_FINISH;") LINE_TERMINATOR;
 	RPCWrappers += TEXT("\t\tP_NATIVE_BEGIN;") LINE_TERMINATOR;
 
-	ClassDefinitionRange ClassRange;
-	if (ClassDefinitionRanges.Contains(Function->GetOwnerClass()))
-	{
-		ClassRange = ClassDefinitionRanges[Function->GetOwnerClass()];
-		ClassRange.Validate();
-	}
+	UClass* OwnerClass = Function->GetOwnerClass();
+	FUnrealClassDefinitionInfo& OwnerClassDef = GTypeDefinitionInfoMap.FindChecked<FUnrealClassDefinitionInfo>(OwnerClass);
 
-	const TCHAR* ClassStart = ClassRange.Start;
-	const TCHAR* ClassEnd   = ClassRange.End;
-	FString      ClassName  = Function->GetOwnerClass()->GetName();
+	OwnerClassDef.GetDefinitionRange().Validate();
+	const TCHAR* ClassStart = OwnerClassDef.GetDefinitionRange().Start;
+	const TCHAR* ClassEnd = OwnerClassDef.GetDefinitionRange().End;
+	FString      ClassName  = OwnerClass->GetName();
 
 	FString ClassDefinition(UE_PTRDIFF_TO_INT32(ClassEnd - ClassStart), ClassStart);
 
@@ -5047,13 +5038,13 @@ void FNativeClassHeaderGenerator::ExportFunctionThunk(FUHTStringBuilder& RPCWrap
 
 	bool bShouldEnableImplementationDeprecation =
 		// Enable deprecation warnings only if GENERATED_BODY is used inside class or interface (not GENERATED_UCLASS_BODY etc.)
-		ClassRange.bHasGeneratedBody
+		OwnerClassDef.HasGeneratedBody()
 		// and implementation function is called, but not the one declared by user
 		&& (FunctionData.CppImplName != Function->GetName() && !bHasImplementation);
 
 	bool bShouldEnableValidateDeprecation =
 		// Enable deprecation warnings only if GENERATED_BODY is used inside class or interface (not GENERATED_UCLASS_BODY etc.)
-		ClassRange.bHasGeneratedBody
+		OwnerClassDef.HasGeneratedBody()
 		// and validation function is called
 		&& (FunctionData.FunctionFlags & FUNC_NetValidate) && !bHasValidate;
 
@@ -5152,12 +5143,9 @@ void FNativeClassHeaderGenerator::ExportNativeFunctions(FOutputDevice& OutGenera
 	FStructMetaData& StructData = ClassDef.GetStructMetaData();
 	const FString ClassCPPName = FNameLookupCPP::GetNameCPP(Class, Class->HasAnyClassFlags(CLASS_Interface));
 
-	ClassDefinitionRange ClassRange;
-	if (ClassDefinitionRanges.Contains(Class))
-	{
-		ClassRange = ClassDefinitionRanges[Class];
-		ClassRange.Validate();
-	}
+	ClassDef.GetDefinitionRange().Validate();
+	const TCHAR* ClassStart = ClassDef.GetDefinitionRange().Start;
+	const TCHAR* ClassEnd = ClassDef.GetDefinitionRange().End;
 
 	// gather static class data
 	TArray<FString> SparseClassDataTypes;
@@ -5254,8 +5242,6 @@ void FNativeClassHeaderGenerator::ExportNativeFunctions(FOutputDevice& OutGenera
 
 		if (!bWillBeProgrammerTyped)
 		{
-			const TCHAR* ClassStart = ClassRange.Start;
-			const TCHAR* ClassEnd   = ClassRange.End;
 			FString ClassDefinition(UE_PTRDIFF_TO_INT32(ClassEnd - ClassStart), ClassStart);
 
 			FString FunctionName = Function->GetName();
@@ -5302,7 +5288,7 @@ void FNativeClassHeaderGenerator::ExportNativeFunctions(FOutputDevice& OutGenera
 			}
 
 			// Versions that skip function autodeclaration throw an error when a function is missing.
-			if (ClassRange.bHasGeneratedBody && (SourceFile.GetGeneratedCodeVersionForStruct(ClassDef) > EGeneratedCodeVersion::V1))
+			if (ClassDef.HasGeneratedBody() && (ClassDef.GetGeneratedCodeVersion() > EGeneratedCodeVersion::V1))
 			{
 				CheckRPCFunctions(OutReferenceGatherers, *FunctionDef, ClassCPPName, ImplementationPosition, ValidatePosition, SourceFile);
 			}
@@ -5353,7 +5339,7 @@ void FNativeClassHeaderGenerator::ExportNativeFunctions(FOutputDevice& OutGenera
 
 		// Put static checks before RPCWrappers to get proper messages from static asserts before compiler errors.
 		FString NoPureDeclsMacroName = SourceFile.GetGeneratedMacroName(StructData, TEXT("_RPC_WRAPPERS_NO_PURE_DECLS"));
-		if (SourceFile.GetGeneratedCodeVersionForStruct(ClassDef) > EGeneratedCodeVersion::V1)
+		if (ClassDef.GetGeneratedCodeVersion() > EGeneratedCodeVersion::V1)
 		{
 			WriteMacro(OutGeneratedHeaderText, NoPureDeclsMacroName, RuntimeStringBuilders.RPCWrappers);
 		}
@@ -5386,7 +5372,7 @@ void FNativeClassHeaderGenerator::ExportNativeFunctions(FOutputDevice& OutGenera
 
 		// Put static checks before RPCWrappers to get proper messages from static asserts before compiler errors.
 		FString NoPureDeclsMacroName = SourceFile.GetGeneratedMacroName(StructData, TEXT("_EDITOR_ONLY_RPC_WRAPPERS_NO_PURE_DECLS"));
-		if (SourceFile.GetGeneratedCodeVersionForStruct(ClassDef) > EGeneratedCodeVersion::V1)
+		if (ClassDef.GetGeneratedCodeVersion() > EGeneratedCodeVersion::V1)
 		{
 			WriteMacro(OutGeneratedHeaderText, NoPureDeclsMacroName, EditorStringBuilders.RPCWrappers);
 		}
diff --git a/Engine/Source/Programs/UnrealHeaderTool/Private/HeaderParser.cpp b/Engine/Source/Programs/UnrealHeaderTool/Private/HeaderParser.cpp
--- a/Engine/Source/Programs/UnrealHeaderTool/Private/HeaderParser.cpp
+++ b/Engine/Source/Programs/UnrealHeaderTool/Private/HeaderParser.cpp
@@ -107,10 +107,8 @@ const FName FHeaderParserNames::NAME_AutoCollapseCategories(TEXT("AutoCollapseCa
 const FName FHeaderParserNames::NAME_HideFunctions(TEXT("HideFunctions"));
 const FName FHeaderParserNames::NAME_AutoExpandCategories(TEXT("AutoExpandCategories"));
 
-FRigVMStructMap FHeaderParser::StructRigVMMap;
 TArray<FString> FHeaderParser::PropertyCPPTypesRequiringUIRanges = { TEXT("float"), TEXT("double") };
 TArray<FString> FHeaderParser::ReservedTypeNames = { TEXT("none") };
-TMap<UClass*, ClassDefinitionRange> ClassDefinitionRanges;
 
 /*-----------------------------------------------------------------------------
 	Utility functions.
@@ -5126,7 +5124,7 @@ bool FHeaderParser::CompileDeclaration(TArray<FUnrealFunctionDefinitionInfo*>& D
 
 		if (Token.Matches(TEXT("GENERATED_BODY"), ESearchCase::CaseSensitive))
 		{
-			ClassDefinitionRanges[GetCurrentClass()].bHasGeneratedBody = true;
+			GetCurrentClassDef().MarkGeneratedBody();
 		}
 		return true;
 	}
@@ -5166,12 +5164,7 @@ bool FHeaderParser::CompileDeclaration(TArray<FUnrealFunctionDefinitionInfo*>& D
 
 		if (Token.Matches(TEXT("GENERATED_BODY"), ESearchCase::CaseSensitive))
 		{
-			if (!ClassDefinitionRanges.Contains(GetCurrentClass()))
-			{
-				ClassDefinitionRanges.Add(GetCurrentClass(), ClassDefinitionRange());
-			}
-
-			ClassDefinitionRanges[GetCurrentClass()].bHasGeneratedBody = true;
+			GetCurrentClassDef().MarkGeneratedBody();
 
 			StructData.GeneratedBodyMacroAccessSpecifier = CurrentAccessSpecifier;
 		}
@@ -5258,10 +5251,7 @@ bool FHeaderParser::CompileDeclaration(TArray<FUnrealFunctionDefinitionInfo*>& D
 
 	if (bEncounteredNewStyleClass_UnmatchedBrackets && Token.Matches(TEXT('}')))
 	{
-		if (ClassDefinitionRanges.Contains(GetCurrentClass()))
-		{
-			ClassDefinitionRanges[GetCurrentClass()].End = &Input[InputPos];
-		}
+		GetCurrentClassDef().GetDefinitionRange().End = &Input[InputPos];
 		MatchSemi();
 
 		// Closing brace for class declaration
@@ -5733,7 +5723,7 @@ FUnrealClassDefinitionInfo& FHeaderParser::CompileClassDeclaration()
 	FUnrealClassDefinitionInfo& ClassDef = ParseClassNameDeclaration(/*out*/ DeclaredClassName, /*out*/ RequiredAPIMacroIfPresent);
 	UClass* Class = ClassDef.GetClass();
 
-	ClassDefinitionRanges.Add(Class, ClassDefinitionRange(&Input[InputPos], nullptr));
+	ClassDef.GetDefinitionRange().Start = &Input[InputPos];
 
 	check(Class->ClassFlags == 0 || (Class->ClassFlags & ClassDef.ClassFlags) != 0);
 
@@ -5938,7 +5928,7 @@ void FHeaderParser::CompileInterfaceDeclaration()
 	FUnrealClassDefinitionInfo* ClassDef = ParseInterfaceNameDeclaration(/*out*/ DeclaredInterfaceName, /*out*/ RequiredAPIMacroIfPresent);
 	check(ClassDef);
 	UClass* InterfaceClass = ClassDef->GetClass();
-	ClassDefinitionRanges.Add(InterfaceClass, ClassDefinitionRange(&Input[InputPos], nullptr));
+	ClassDef->GetDefinitionRange().Start = &Input[InputPos];
 
 	// Record that this interface is RequiredAPI if the CORE_API style macro was present
 	if (!RequiredAPIMacroIfPresent.IsEmpty())
@@ -6137,7 +6127,8 @@ void FHeaderParser::CompileRigVMMethodDeclaration(FUnrealStructDefinitionInfo& S
 		}
 	}
 
-	FRigVMStructInfo& StructRigVMInfo = StructRigVMMap.FindOrAdd(&StructDef);
+	FRigVMStructInfo& StructRigVMInfo = StructDef.GetRigVMInfo();
+	StructRigVMInfo.bHasRigVM = true;
 	StructRigVMInfo.Name = StructDef.GetStruct()->GetName();
 	StructRigVMInfo.Methods.Add(MethodInfo);
 }
@@ -6160,8 +6151,8 @@ const TCHAR* FHeaderParser::GetDynamicArrayText = TEXT("GetDynamicArray");
 void FHeaderParser::ParseRigVMMethodParameters(FUnrealStructDefinitionInfo& StructDef)
 {
 	UStruct* Struct = StructDef.GetStruct();
-	FRigVMStructInfo* StructRigVMInfo = StructRigVMMap.Find(&StructDef);
-	if (StructRigVMInfo == nullptr)
+	FRigVMStructInfo& StructRigVMInfo = StructDef.GetRigVMInfo();
+	if (!StructRigVMInfo.bHasRigVM)
 	{
 		return;
 	}
@@ -6224,7 +6215,7 @@ void FHeaderParser::ParseRigVMMethodParameters(FUnrealStructDefinitionInfo& Stru
 		if (MemberCPPType.StartsWith(TArrayText, ESearchCase::CaseSensitive))
 		{
 			ExtendedCPPType = FString::Printf(TEXT("<%s>"), *ExtendedCPPType.LeftChop(1).RightChop(1));
-			Parameter.CastName = FString::Printf(TEXT("%s_%d_Array"), *Parameter.Name, StructRigVMInfo->Members.Num());
+			Parameter.CastName = FString::Printf(TEXT("%s_%d_Array"), *Parameter.Name, StructRigVMInfo.Members.Num());
 			if (Parameter.IsConst() || !Parameter.ArraySize.IsEmpty())
 			{
 				Parameter.CastType = FString::Printf(TEXT("%s%s"), FFixedArrayText, *ExtendedCPPType);
@@ -6237,17 +6228,17 @@ void FHeaderParser::ParseRigVMMethodParameters(FUnrealStructDefinitionInfo& Stru
 			}
 		}
 
-		StructRigVMInfo->Members.Add(MoveTemp(Parameter));
+		StructRigVMInfo.Members.Add(MoveTemp(Parameter));
 	}
 
-	if (StructRigVMInfo->Members.Num() == 0)
+	if (StructRigVMInfo.Members.Num() == 0)
 	{
 		UE_LOG_ERROR_UHT(TEXT("RigVM Struct '%s' - has zero members - invalid RIGVM_METHOD."), *Struct->GetName());
 	}
 
-	if (StructRigVMInfo->Members.Num() > 64)
+	if (StructRigVMInfo.Members.Num() > 64)
 	{
-		UE_LOG_ERROR_UHT(TEXT("RigVM Struct '%s' - has %d members (64 is the limit)."), *Struct->GetName(), StructRigVMInfo->Members.Num());
+		UE_LOG_ERROR_UHT(TEXT("RigVM Struct '%s' - has %d members (64 is the limit)."), *Struct->GetName(), StructRigVMInfo.Members.Num());
 	}
 }
 
@@ -9016,15 +9007,14 @@ void FHeaderParser::CompileVersionDeclaration(FUnrealStructDefinitionInfo& Struc
 	if (Token.TokenType == ETokenType::TOKEN_Symbol
 		&& Token.Matches(TEXT(')')))
 	{
-		SourceFile.GetGeneratedCodeVersions().FindOrAdd(&StructDef) = Version;
+		StructDef.SetGeneratedCodeVersion(Version);
 		UngetToken(Token);
 		return;
 	}
 
 	// Overwrite with version specified by macro.
 	Version = ToGeneratedCodeVersion(Token.Identifier);
-
-	SourceFile.GetGeneratedCodeVersions().FindOrAdd(&StructDef) = Version;
+	StructDef.SetGeneratedCodeVersion(Version);
 }
 
 void FHeaderParser::ResetClassData()
diff --git a/Engine/Source/Programs/UnrealHeaderTool/Private/HeaderParser.h b/Engine/Source/Programs/UnrealHeaderTool/Private/HeaderParser.h
--- a/Engine/Source/Programs/UnrealHeaderTool/Private/HeaderParser.h
+++ b/Engine/Source/Programs/UnrealHeaderTool/Private/HeaderParser.h
@@ -121,283 +121,6 @@ struct FIndexRange
 	int32 Count;
 };
 
-/**
- * The FRigVMParameter represents a single parameter of a method
- * marked up with RIGVM_METHOD.
- * Each parameter can be marked with Constant, Input or Output
- * metadata - this struct simplifies access to that information.
- */
-struct FRigVMParameter
-{
-	FRigVMParameter()
-		: Name()
-		, Type()
-		, bConstant(false)
-		, bInput(false)
-		, bOutput(false)
-		, bSingleton(false)
-		, ArraySize()
-		, Getter()
-		, CastName()
-		, CastType()
-		, bEditorOnly(false)
-		, bIsEnum(false)
-	{
-	}
-
-	FString Name;
-	FString Type;
-	bool bConstant;
-	bool bInput;
-	bool bOutput;
-	bool bSingleton;
-	FString ArraySize;
-	FString Getter;
-	FString CastName;
-	FString CastType;
-	bool bEditorOnly;
-	bool bIsEnum;
-
-	const FString& NameOriginal(bool bCastName = false) const
-	{
-		return (bCastName && !CastName.IsEmpty()) ? CastName : Name;
-	}
-
-	const FString& TypeOriginal(bool bCastType = false) const
-	{
-		return (bCastType && !CastType.IsEmpty()) ? CastType : Type;
-	}
-
-	FString Declaration(bool bCastType = false, bool bCastName = false) const
-	{
-		return FString::Printf(TEXT("%s %s"), *TypeOriginal(bCastType), *NameOriginal(bCastName));
-	}
-
-	FString BaseType(bool bCastType = false) const
-	{
-		const FString& String = TypeOriginal(bCastType);
-		int32 LesserPos = 0;
-		if (String.FindChar('<', LesserPos))
-		{
-			return String.Mid(0, LesserPos);
-		}
-		return String;
-	}
-
-	FString ExtendedType(bool bCastType = false) const
-	{
-		const FString& String = TypeOriginal(bCastType);
-		int32 LesserPos = 0;
-		if(String.FindChar('<', LesserPos))
-		{
-			return String.Mid(LesserPos);
-		}
-		return String;
-	}
-
-	FString TypeConstRef(bool bCastType = false) const
-	{
-		const FString& String = TypeNoRef(bCastType);
-		if (String.StartsWith(TEXT("T"), ESearchCase::CaseSensitive) || String.StartsWith(TEXT("F"), ESearchCase::CaseSensitive))
-		{
-			return FString::Printf(TEXT("const %s&"), *String);
-		}
-		return FString::Printf(TEXT("const %s"), *String);
-	}
-
-	FString TypeRef(bool bCastType = false) const
-	{
-		const FString& String = TypeNoRef(bCastType);
-		return FString::Printf(TEXT("%s&"), *String);
-	}
-
-	FString TypeNoRef(bool bCastType = false) const
-	{
-		const FString& String = TypeOriginal(bCastType);
-		if (String.EndsWith(TEXT("&")))
-		{
-			return String.LeftChop(1);
-		}
-		return String;
-	}
-
-	FString TypeVariableRef(bool bCastType = false) const
-	{
-		return IsConst() ? TypeConstRef(bCastType) : TypeRef(bCastType);
-	}
-
-	FString Variable(bool bCastType = false, bool bCastName = false) const
-	{
-		return FString::Printf(TEXT("%s %s"), *TypeVariableRef(bCastType), *NameOriginal(bCastName));
-	}
-
-	bool IsConst() const
-	{
-		return bConstant || (bInput && !bOutput);
-	}
-
-	bool IsArray() const
-	{
-		return BaseType().Equals(TEXT("TArray"));
-	}
-
-	bool IsDynamic() const
-	{
-		return ArraySize.IsEmpty() && !bInput && !bOutput && !bSingleton;
-	}
-
-	bool IsDynamicArray() const
-	{
-		return IsArray() && IsDynamic();
-	}
-
-	bool RequiresCast() const
-	{
-		return !CastType.IsEmpty() && !CastName.IsEmpty();
-	}
-};
-
-/**
- * The FRigVMParameterArray represents the parameters in a notation
- * of a function marked with RIGVM_METHOD. The parameter array can 
- * produce a comma separated list of names or parameter declarations.
- */
-struct FRigVMParameterArray
-{
-public:
-	int32 Num() const { return Parameters.Num(); }
-	const FRigVMParameter& operator[](int32 InIndex) const { return Parameters[InIndex]; }
-	FRigVMParameter& operator[](int32 InIndex) { return Parameters[InIndex]; }
-	TArray<FRigVMParameter>::RangedForConstIteratorType begin() const { return Parameters.begin(); }
-	TArray<FRigVMParameter>::RangedForConstIteratorType end() const { return Parameters.end(); }
-	TArray<FRigVMParameter>::RangedForIteratorType begin() { return Parameters.begin(); }
-	TArray<FRigVMParameter>::RangedForIteratorType end() { return Parameters.end(); }
-
-	int32 Add(const FRigVMParameter& InParameter)
-	{
-		return Parameters.Add(InParameter);
-	}
-
-	FString Names(bool bLeadingSeparator = false, const TCHAR* Separator = TEXT(", "), bool bCastType = false, bool bIncludeEditorOnly = true) const
-	{
-		if (Parameters.Num() == 0)
-		{
-			return FString();
-		}
-		TArray<FString> NameArray;
-		for (const FRigVMParameter& Parameter : Parameters)
-		{
-			if (!bIncludeEditorOnly && Parameter.bEditorOnly)
-			{
-				continue;
-			}
-			NameArray.Add(Parameter.NameOriginal(bCastType));
-		}
-
-		if (NameArray.Num() == 0)
-		{
-			return FString();
-		}
-
-		FString Joined = FString::Join(NameArray, Separator);
-		if (bLeadingSeparator)
-		{
-			return FString::Printf(TEXT("%s%s"), Separator, *Joined);
-		}
-		return Joined;
-	}
-
-	FString Declarations(bool bLeadingSeparator = false, const TCHAR* Separator = TEXT(", "), bool bCastType = false, bool bCastName = false, bool bIncludeEditorOnly = true) const
-	{
-		if (Parameters.Num() == 0)
-		{
-			return FString();
-		}
-		TArray<FString> DeclarationArray;
-		for (const FRigVMParameter& Parameter : Parameters)
-		{
-			if (!bIncludeEditorOnly && Parameter.bEditorOnly)
-			{
-				continue;
-			}
-			DeclarationArray.Add(Parameter.Variable(bCastType, bCastName));
-		}
-
-		if (DeclarationArray.Num() == 0)
-		{
-			return FString();
-		}
-
-		FString Joined = FString::Join(DeclarationArray, Separator);
-		if (bLeadingSeparator)
-		{
-			return FString::Printf(TEXT("%s%s"), Separator, *Joined);
-		}
-		return Joined;
-	}
-
-private:
-	TArray<FRigVMParameter> Parameters;
-};
-
-/**
- * A single info dataset for a function marked with RIGVM_METHOD.
- * This struct provides access to its name, the return type and all parameters.
- */
-struct FRigVMMethodInfo
-{
-	FString ReturnType;
-	FString Name;
-	FRigVMParameterArray Parameters;
-
-	FString ReturnPrefix() const
-	{
-		return (ReturnType.IsEmpty() || (ReturnType == TEXT("void"))) ? TEXT("") : TEXT("return ");
-	}
-};
-
-/**
- * An info dataset providing access to all functions marked with RIGVM_METHOD
- * for each struct.
- */
-struct FRigVMStructInfo
-{
-	FString Name;
-	FRigVMParameterArray Members;
-	TArray<FRigVMMethodInfo> Methods;
-};
-
-typedef TMap<FUnrealStructDefinitionInfo*, FRigVMStructInfo> FRigVMStructMap;
-
-struct ClassDefinitionRange
-{
-	ClassDefinitionRange(const TCHAR* InStart, const TCHAR* InEnd)
-		: Start(InStart)
-		, End(InEnd)
-		, bHasGeneratedBody(false)
-	{ }
-
-	ClassDefinitionRange()
-		: Start(nullptr)
-		, End(nullptr)
-		, bHasGeneratedBody(false)
-	{ }
-
-	void Validate()
-	{
-		if (End <= Start)
-		{
-			FError::Throwf(TEXT("The class definition range is invalid. Most probably caused by previous parsing error."));
-		}
-	}
-
-	const TCHAR* Start;
-	const TCHAR* End;
-	bool bHasGeneratedBody;
-};
-
-extern TMap<UClass*, ClassDefinitionRange> ClassDefinitionRanges;
-
 #ifndef UHT_DOCUMENTATION_POLICY_DEFAULT
 #define UHT_DOCUMENTATION_POLICY_DEFAULT false
 #endif
@@ -619,9 +342,6 @@ protected:
 	// List of all net service functions with undeclared response functions 
 	TMap<int32, FString> RPCsNeedingHookup;
 
-	// List of all multiplex methods defined on structs
-	static FRigVMStructMap StructRigVMMap;
-
 	// Constructor.
 	explicit FHeaderParser(FUnrealPackageDefinitionInfo& InPackageDef, FFeedbackContext* InWarn, FUnrealSourceFile& InSourceFile);
 
diff --git a/Engine/Source/Programs/UnrealHeaderTool/Private/ParserHelper.h b/Engine/Source/Programs/UnrealHeaderTool/Private/ParserHelper.h
--- a/Engine/Source/Programs/UnrealHeaderTool/Private/ParserHelper.h
+++ b/Engine/Source/Programs/UnrealHeaderTool/Private/ParserHelper.h
@@ -859,3 +859,250 @@ public:
 	/** return if more parameters can be marked */
 	bool CanMarkMore() const;
 };
+
+/**
+ * The FRigVMParameter represents a single parameter of a method
+ * marked up with RIGVM_METHOD.
+ * Each parameter can be marked with Constant, Input or Output
+ * metadata - this struct simplifies access to that information.
+ */
+struct FRigVMParameter
+{
+	FRigVMParameter()
+		: Name()
+		, Type()
+		, bConstant(false)
+		, bInput(false)
+		, bOutput(false)
+		, bSingleton(false)
+		, ArraySize()
+		, Getter()
+		, CastName()
+		, CastType()
+		, bEditorOnly(false)
+		, bIsEnum(false)
+	{
+	}
+
+	FString Name;
+	FString Type;
+	bool bConstant;
+	bool bInput;
+	bool bOutput;
+	bool bSingleton;
+	FString ArraySize;
+	FString Getter;
+	FString CastName;
+	FString CastType;
+	bool bEditorOnly;
+	bool bIsEnum;
+
+	const FString& NameOriginal(bool bCastName = false) const
+	{
+		return (bCastName && !CastName.IsEmpty()) ? CastName : Name;
+	}
+
+	const FString& TypeOriginal(bool bCastType = false) const
+	{
+		return (bCastType && !CastType.IsEmpty()) ? CastType : Type;
+	}
+
+	FString Declaration(bool bCastType = false, bool bCastName = false) const
+	{
+		return FString::Printf(TEXT("%s %s"), *TypeOriginal(bCastType), *NameOriginal(bCastName));
+	}
+
+	FString BaseType(bool bCastType = false) const
+	{
+		const FString& String = TypeOriginal(bCastType);
+		int32 LesserPos = 0;
+		if (String.FindChar('<', LesserPos))
+		{
+			return String.Mid(0, LesserPos);
+		}
+		return String;
+	}
+
+	FString ExtendedType(bool bCastType = false) const
+	{
+		const FString& String = TypeOriginal(bCastType);
+		int32 LesserPos = 0;
+		if (String.FindChar('<', LesserPos))
+		{
+			return String.Mid(LesserPos);
+		}
+		return String;
+	}
+
+	FString TypeConstRef(bool bCastType = false) const
+	{
+		const FString& String = TypeNoRef(bCastType);
+		if (String.StartsWith(TEXT("T"), ESearchCase::CaseSensitive) || String.StartsWith(TEXT("F"), ESearchCase::CaseSensitive))
+		{
+			return FString::Printf(TEXT("const %s&"), *String);
+		}
+		return FString::Printf(TEXT("const %s"), *String);
+	}
+
+	FString TypeRef(bool bCastType = false) const
+	{
+		const FString& String = TypeNoRef(bCastType);
+		return FString::Printf(TEXT("%s&"), *String);
+	}
+
+	FString TypeNoRef(bool bCastType = false) const
+	{
+		const FString& String = TypeOriginal(bCastType);
+		if (String.EndsWith(TEXT("&")))
+		{
+			return String.LeftChop(1);
+		}
+		return String;
+	}
+
+	FString TypeVariableRef(bool bCastType = false) const
+	{
+		return IsConst() ? TypeConstRef(bCastType) : TypeRef(bCastType);
+	}
+
+	FString Variable(bool bCastType = false, bool bCastName = false) const
+	{
+		return FString::Printf(TEXT("%s %s"), *TypeVariableRef(bCastType), *NameOriginal(bCastName));
+	}
+
+	bool IsConst() const
+	{
+		return bConstant || (bInput && !bOutput);
+	}
+
+	bool IsArray() const
+	{
+		return BaseType().Equals(TEXT("TArray"));
+	}
+
+	bool IsDynamic() const
+	{
+		return ArraySize.IsEmpty() && !bInput && !bOutput && !bSingleton;
+	}
+
+	bool IsDynamicArray() const
+	{
+		return IsArray() && IsDynamic();
+	}
+
+	bool RequiresCast() const
+	{
+		return !CastType.IsEmpty() && !CastName.IsEmpty();
+	}
+};
+
+/**
+ * The FRigVMParameterArray represents the parameters in a notation
+ * of a function marked with RIGVM_METHOD. The parameter array can
+ * produce a comma separated list of names or parameter declarations.
+ */
+struct FRigVMParameterArray
+{
+public:
+	int32 Num() const { return Parameters.Num(); }
+	const FRigVMParameter& operator[](int32 InIndex) const { return Parameters[InIndex]; }
+	FRigVMParameter& operator[](int32 InIndex) { return Parameters[InIndex]; }
+	TArray<FRigVMParameter>::RangedForConstIteratorType begin() const { return Parameters.begin(); }
+	TArray<FRigVMParameter>::RangedForConstIteratorType end() const { return Parameters.end(); }
+	TArray<FRigVMParameter>::RangedForIteratorType begin() { return Parameters.begin(); }
+	TArray<FRigVMParameter>::RangedForIteratorType end() { return Parameters.end(); }
+
+	int32 Add(const FRigVMParameter& InParameter)
+	{
+		return Parameters.Add(InParameter);
+	}
+
+	FString Names(bool bLeadingSeparator = false, const TCHAR* Separator = TEXT(", "), bool bCastType = false, bool bIncludeEditorOnly = true) const
+	{
+		if (Parameters.Num() == 0)
+		{
+			return FString();
+		}
+		TArray<FString> NameArray;
+		for (const FRigVMParameter& Parameter : Parameters)
+		{
+			if (!bIncludeEditorOnly && Parameter.bEditorOnly)
+			{
+				continue;
+			}
+			NameArray.Add(Parameter.NameOriginal(bCastType));
+		}
+
+		if (NameArray.Num() == 0)
+		{
+			return FString();
+		}
+
+		FString Joined = FString::Join(NameArray, Separator);
+		if (bLeadingSeparator)
+		{
+			return FString::Printf(TEXT("%s%s"), Separator, *Joined);
+		}
+		return Joined;
+	}
+
+	FString Declarations(bool bLeadingSeparator = false, const TCHAR* Separator = TEXT(", "), bool bCastType = false, bool bCastName = false, bool bIncludeEditorOnly = true) const
+	{
+		if (Parameters.Num() == 0)
+		{
+			return FString();
+		}
+		TArray<FString> DeclarationArray;
+		for (const FRigVMParameter& Parameter : Parameters)
+		{
+			if (!bIncludeEditorOnly && Parameter.bEditorOnly)
+			{
+				continue;
+			}
+			DeclarationArray.Add(Parameter.Variable(bCastType, bCastName));
+		}
+
+		if (DeclarationArray.Num() == 0)
+		{
+			return FString();
+		}
+
+		FString Joined = FString::Join(DeclarationArray, Separator);
+		if (bLeadingSeparator)
+		{
+			return FString::Printf(TEXT("%s%s"), Separator, *Joined);
+		}
+		return Joined;
+	}
+
+private:
+	TArray<FRigVMParameter> Parameters;
+};
+
+/**
+ * A single info dataset for a function marked with RIGVM_METHOD.
+ * This struct provides access to its name, the return type and all parameters.
+ */
+struct FRigVMMethodInfo
+{
+	FString ReturnType;
+	FString Name;
+	FRigVMParameterArray Parameters;
+
+	FString ReturnPrefix() const
+	{
+		return (ReturnType.IsEmpty() || (ReturnType == TEXT("void"))) ? TEXT("") : TEXT("return ");
+	}
+};
+
+/**
+ * An info dataset providing access to all functions marked with RIGVM_METHOD
+ * for each struct.
+ */
+struct FRigVMStructInfo
+{
+	bool bHasRigVM = false;
+	FString Name;
+	FRigVMParameterArray Members;
+	TArray<FRigVMMethodInfo> Methods;
+};
diff --git a/Engine/Source/Programs/UnrealHeaderTool/Private/UnrealSourceFile.cpp b/Engine/Source/Programs/UnrealHeaderTool/Private/UnrealSourceFile.cpp
--- a/Engine/Source/Programs/UnrealHeaderTool/Private/UnrealSourceFile.cpp
+++ b/Engine/Source/Programs/UnrealHeaderTool/Private/UnrealSourceFile.cpp
@@ -41,7 +41,7 @@ const FString& FUnrealSourceFile::GetFileId() const
 			bRelativePath = FPaths::MakePathRelativeTo(StdFilename, *FPaths::GetPath(FPaths::GetProjectFilePath()));
 		}
 
-		// If the path has passed either MakeStandardFilename or MakePathRelativeTo it should be using internal path seperators
+		// If the path has passed either MakeStandardFilename or MakePathRelativeTo it should be using internal path separators
 		if (bRelativePath)
 		{
 			// Remove any preceding parent directory paths
@@ -128,16 +128,6 @@ const FString& FUnrealSourceFile::GetContent() const
 	return Content;
 }
 
-EGeneratedCodeVersion FUnrealSourceFile::GetGeneratedCodeVersionForStruct(FUnrealStructDefinitionInfo& StructDef) const
-{
-	if (const EGeneratedCodeVersion* Version = GeneratedCodeVersions.Find(&StructDef))
-	{
-		return *Version;
-	}
-
-	return FUHTConfig::Get().DefaultGeneratedCodeVersion;
-}
-
 bool FUnrealSourceFile::HasChanged() const
 {
 	return bHasChanged;
diff --git a/Engine/Source/Programs/UnrealHeaderTool/Private/UnrealSourceFile.h b/Engine/Source/Programs/UnrealHeaderTool/Private/UnrealSourceFile.h
--- a/Engine/Source/Programs/UnrealHeaderTool/Private/UnrealSourceFile.h
+++ b/Engine/Source/Programs/UnrealHeaderTool/Private/UnrealSourceFile.h
@@ -10,8 +10,6 @@
 #include "GeneratedCodeVersion.h"
 
 class UPackage;
-class UClass;
-class UStruct;
 class FArchive;
 class FStructMetaData;
 
@@ -288,27 +286,6 @@ public:
 	 */
 	void AddTypeDefIncludeIfNeeded(UField* Field);
 
-	/**
-	 * Gets generated code version for given UStruct.
-	 */
-	EGeneratedCodeVersion GetGeneratedCodeVersionForStruct(FUnrealStructDefinitionInfo& StructDef) const;
-
-	/**
-	 * Gets generated code versions.
-	 */
-	TMap<FUnrealStructDefinitionInfo*, EGeneratedCodeVersion>& GetGeneratedCodeVersions()
-	{
-		return GeneratedCodeVersions;
-	}
-
-	/**
-	 * Gets generated code versions. Const version.
-	 */
-	const TMap<FUnrealStructDefinitionInfo*, EGeneratedCodeVersion>& GetGeneratedCodeVersions() const
-	{
-		return GeneratedCodeVersions;
-	}
-
 	/**
 	 * Sets generated filename.
 	 */
@@ -519,9 +496,6 @@ private:
 	// List of structs defined in this source file along with parsing info.
 	TArray<TSharedRef<FUnrealTypeDefinitionInfo>> DefinedStructs;
 
-	// Mapping of FUnrealStructDefinitionInfo to versions, according to which their code should be generated.
-	TMap<FUnrealStructDefinitionInfo*, EGeneratedCodeVersion> GeneratedCodeVersions;
-
 	// Collection of all singletons found during code generation
 	TArray<FUnrealFieldDefinitionInfo*> Singletons;
 };
diff --git a/Engine/Source/Programs/UnrealHeaderTool/Private/UnrealTypeDefinitionInfo.h b/Engine/Source/Programs/UnrealHeaderTool/Private/UnrealTypeDefinitionInfo.h
--- a/Engine/Source/Programs/UnrealHeaderTool/Private/UnrealTypeDefinitionInfo.h
+++ b/Engine/Source/Programs/UnrealHeaderTool/Private/UnrealTypeDefinitionInfo.h
@@ -921,6 +921,21 @@ public:
 		FUnrealStructDefinitionInfo* Struct = nullptr;
 	};
 
+	struct FDefinitionRange
+	{
+		void Validate()
+		{
+			if (End <= Start)
+			{
+				FError::Throwf(TEXT("The class definition range is invalid. Most probably caused by previous parsing error."));
+			}
+		}
+
+		const TCHAR* Start = nullptr;
+		const TCHAR* End = nullptr;
+	};
+
+
 public:
 	using FUnrealFieldDefinitionInfo::FUnrealFieldDefinitionInfo;
 
@@ -1037,6 +1052,54 @@ public:
 		return false;
 	}
 
+	/**
+	 * Get the generated code version
+	 */
+	EGeneratedCodeVersion GetGeneratedCodeVersion() const
+	{
+		return GeneratedCodeVersion;
+	}
+
+	/**
+	 * Set the generated code version
+	 */
+	void SetGeneratedCodeVersion(EGeneratedCodeVersion InGeneratedCodeVersion)
+	{
+		GeneratedCodeVersion = InGeneratedCodeVersion;
+	}
+
+	/**
+	 * Get if we have a generated body
+	 */
+	bool HasGeneratedBody() const
+	{
+		return bHasGeneratedBody;
+	}
+
+	/**
+	 * Mark that we have a generated body
+	 */
+	void MarkGeneratedBody()
+	{
+		bHasGeneratedBody = true;
+	}
+
+	/**
+	 * Return the definition range of the structure
+	 */
+	FDefinitionRange& GetDefinitionRange()
+	{
+		return DefinitionRange;
+	}
+
+	/**
+	 * Get the RigVM information
+	 */
+	FRigVMStructInfo& GetRigVMInfo()
+	{
+		return RigVMInfo;
+	}
+
 private:
 	TSharedPtr<FScope> StructScope;
 
@@ -1050,8 +1113,17 @@ private:
 	FBaseClassInfo SuperClassInfo;
 	TArray<FBaseClassInfo> BaseClassInfo;
 
+	FDefinitionRange DefinitionRange;
+
+	FRigVMStructInfo RigVMInfo;
+
+	EGeneratedCodeVersion GeneratedCodeVersion = FUHTConfig::Get().DefaultGeneratedCodeVersion;
+
 	/** whether this struct declares delegate functions or properties */
 	bool bContainsDelegates = false;
+
+	/** If true, this struct contains the generated body macro */
+	bool bHasGeneratedBody = false;
 };
 
 /**
