Removed most calls to GetFunction (parser still uses it a lot)

#rnx

diff --git a/Engine/Source/Programs/UnrealHeaderTool/Private/CodeGenerator.cpp b/Engine/Source/Programs/UnrealHeaderTool/Private/CodeGenerator.cpp
--- a/Engine/Source/Programs/UnrealHeaderTool/Private/CodeGenerator.cpp
+++ b/Engine/Source/Programs/UnrealHeaderTool/Private/CodeGenerator.cpp
@@ -796,19 +796,6 @@ const FString& FNativeClassHeaderGenerator::GetSingletonNameFuncAddr(FUnrealFiel
 	}
 }
 
-const FString& FNativeClassHeaderGenerator::GetSingletonNameFuncAddr(UField* Item, TSet<FString>* UniqueCrossModuleReferences, bool bRequiresValidObject)
-{
-	if (!Item)
-	{
-		return GNullPtr;
-	}
-	else
-	{
-		FUnrealFieldDefinitionInfo& Field = GTypeDefinitionInfoMap.FindChecked<FUnrealFieldDefinitionInfo>(Item);
-		return GetSingletonNameFuncAddr(&Field, UniqueCrossModuleReferences, bRequiresValidObject);
-	}
-}
-
 void FNativeClassHeaderGenerator::GetPropertyTag(FUHTStringBuilder& Out, FUnrealPropertyDefinitionInfo& PropDef)
 {
 	const FPropertyBase& PropertyBase = PropDef.GetPropertyBase();
@@ -1895,7 +1882,7 @@ static bool IsAlwaysAccessible(UScriptStruct* Script)
 
 static void FindNoExportStructsRecursive(TArray<FUnrealScriptStructDefinitionInfo*>& StructDefs, FUnrealStructDefinitionInfo* StartDef)
 {
-	for (; StartDef; StartDef = StartDef->GetSuperClassInfo().Struct)
+	for (; StartDef; StartDef = StartDef->GetSuperStructInfo().Struct)
 	{
 		UStruct* Start = StartDef->GetStruct();
 		if (UScriptStruct* StartScript = Cast<UScriptStruct>(Start))
@@ -2432,7 +2419,8 @@ void FNativeClassHeaderGenerator::ExportFunction(FOutputDevice& Out, FReferenceG
 	UFunction* Function = FunctionDef.GetFunction();
 	FunctionDef.AddCrossModuleReference(OutReferenceGatherers.UniqueCrossModuleReferences, true);
 
-	UFunction* SuperFunction = Function->GetSuperFunction();
+	FUnrealFunctionDefinitionInfo* SuperFunctionDef = FunctionDef.GetSuperFunction();
+	UFunction* SuperFunction = SuperFunctionDef ? SuperFunctionDef->GetFunction() : nullptr;
 
 	const bool bIsEditorOnlyFunction = Function->HasAnyFunctionFlags(FUNC_EditorOnly);
 
@@ -2472,7 +2460,7 @@ void FNativeClassHeaderGenerator::ExportFunction(FOutputDevice& Out, FReferenceG
 	if (FUnrealObjectDefinitionInfo* OuterDef = FunctionDef.GetOuter())
 	{
 		FUnrealPackageDefinitionInfo* OuterPackageDef = UHTCast<FUnrealPackageDefinitionInfo>(OuterDef);
-		OuterFunc = OuterPackageDef ? GetPackageSingletonNameFuncAddr(*OuterPackageDef, OutReferenceGatherers.UniqueCrossModuleReferences) : GetSingletonNameFuncAddr(Function->GetOwnerClass(), OutReferenceGatherers.UniqueCrossModuleReferences);
+		OuterFunc = OuterPackageDef ? GetPackageSingletonNameFuncAddr(*OuterPackageDef, OutReferenceGatherers.UniqueCrossModuleReferences) : GetSingletonNameFuncAddr(FunctionDef.GetOwnerClass(), OutReferenceGatherers.UniqueCrossModuleReferences);
 	}
 	else
 	{
@@ -2523,7 +2511,7 @@ void FNativeClassHeaderGenerator::ExportFunction(FOutputDevice& Out, FReferenceG
 		TEXT("\tconst UECodeGen_Private::FFunctionParams %s::FuncParams = { (UObject*(*)())%s, %s, %s, %s, %s, %s, %s, %s, %s, (EFunctionFlags)0x%08X, %d, %d, %s };\r\n"),
 		*StaticsStructName,
 		*OuterFunc,
-		*GetSingletonNameFuncAddr(SuperFunction, OutReferenceGatherers.UniqueCrossModuleReferences),
+		*GetSingletonNameFuncAddr(SuperFunctionDef, OutReferenceGatherers.UniqueCrossModuleReferences),
 		*CreateUTF8LiteralString(FNativeClassHeaderGenerator::GetOverriddenName(Function)),
 		(SparseDelegateFunction ? *CreateUTF8LiteralString(SparseDelegateFunction->OwningClassName.ToString()) : TEXT("nullptr")),
 		(SparseDelegateFunction ? *CreateUTF8LiteralString(SparseDelegateFunction->DelegateName.ToString()) : TEXT("nullptr")),
@@ -3463,7 +3451,7 @@ void FNativeClassHeaderGenerator::ExportGeneratedStructBodyMacros(FOutputDevice&
 	const FString ActualStructName = FNativeClassHeaderGenerator::GetOverriddenName(Struct);
 	const FString& FriendApiString  = GetAPIString();
 
-	FUnrealStructDefinitionInfo* BaseStructDef = ScriptStructDef.GetSuperClassInfo().Struct;
+	FUnrealStructDefinitionInfo* BaseStructDef = ScriptStructDef.GetSuperStructInfo().Struct;
 	UStruct* BaseStruct = Struct->GetSuperStruct();
 
 	const FString StructNameCPP = FNameLookupCPP::GetNameCPP(Struct);
@@ -5023,12 +5011,13 @@ void FNativeClassHeaderGenerator::ExportFunctionThunk(FUHTStringBuilder& RPCWrap
 	RPCWrappers += TEXT("\t\tP_FINISH;") LINE_TERMINATOR;
 	RPCWrappers += TEXT("\t\tP_NATIVE_BEGIN;") LINE_TERMINATOR;
 
-	UClass* OwnerClass = Function->GetOwnerClass();
-	FUnrealClassDefinitionInfo& OwnerClassDef = GTypeDefinitionInfoMap.FindChecked<FUnrealClassDefinitionInfo>(OwnerClass);
+	FUnrealClassDefinitionInfo* OwnerClassDef = FunctionDef.GetOwnerClass();
+	check(OwnerClassDef);
+	UClass* OwnerClass = OwnerClassDef->GetClass();
 
-	OwnerClassDef.GetDefinitionRange().Validate();
-	const TCHAR* ClassStart = OwnerClassDef.GetDefinitionRange().Start;
-	const TCHAR* ClassEnd = OwnerClassDef.GetDefinitionRange().End;
+	OwnerClassDef->GetDefinitionRange().Validate();
+	const TCHAR* ClassStart = OwnerClassDef->GetDefinitionRange().Start;
+	const TCHAR* ClassEnd = OwnerClassDef->GetDefinitionRange().End;
 	FString      ClassName  = OwnerClass->GetName();
 
 	FString ClassDefinition(UE_PTRDIFF_TO_INT32(ClassEnd - ClassStart), ClassStart);
@@ -5038,13 +5027,13 @@ void FNativeClassHeaderGenerator::ExportFunctionThunk(FUHTStringBuilder& RPCWrap
 
 	bool bShouldEnableImplementationDeprecation =
 		// Enable deprecation warnings only if GENERATED_BODY is used inside class or interface (not GENERATED_UCLASS_BODY etc.)
-		OwnerClassDef.HasGeneratedBody()
+		OwnerClassDef->HasGeneratedBody()
 		// and implementation function is called, but not the one declared by user
 		&& (FunctionData.CppImplName != Function->GetName() && !bHasImplementation);
 
 	bool bShouldEnableValidateDeprecation =
 		// Enable deprecation warnings only if GENERATED_BODY is used inside class or interface (not GENERATED_UCLASS_BODY etc.)
-		OwnerClassDef.HasGeneratedBody()
+		OwnerClassDef->HasGeneratedBody()
 		// and validation function is called
 		&& (FunctionData.FunctionFlags & FUNC_NetValidate) && !bHasValidate;
 
@@ -5210,7 +5199,7 @@ void FNativeClassHeaderGenerator::ExportNativeFunctions(FOutputDevice& OutGenera
 				}
 			}
 
-			SparseClassDataStructDef = UHTCast<FUnrealScriptStructDefinitionInfo>(SparseClassDataStructDef->GetSuperClassInfo().Struct);
+			SparseClassDataStructDef = UHTCast<FUnrealScriptStructDefinitionInfo>(SparseClassDataStructDef->GetSuperStructInfo().Struct);
 		}
 	}
 
@@ -6385,7 +6374,7 @@ void ResolveSuperClasses(UPackage* Package)
 
 		// Resolve the base class
 		{
-			FUnrealStructDefinitionInfo::FBaseClassInfo& SuperClassInfo = ClassDef.GetSuperClassInfo();
+			FUnrealStructDefinitionInfo::FBaseStructInfo& SuperClassInfo = ClassDef.GetSuperStructInfo();
 			const FString& BaseClassName = SuperClassInfo.Name;
 			const FString& BaseClassNameStripped = GetClassNameWithPrefixRemoved(BaseClassName);
 
@@ -6411,7 +6400,7 @@ void ResolveSuperClasses(UPackage* Package)
 
 		// Resolve the inherited classes
 		{
-			for (FUnrealStructDefinitionInfo::FBaseClassInfo& BaseClassInfo : ClassDef.GetBaseClassInfo())
+			for (FUnrealStructDefinitionInfo::FBaseStructInfo& BaseClassInfo : ClassDef.GetBaseStructInfo())
 			{
 				BaseClassInfo.Struct = FUnrealClassDefinitionInfo::FindScriptClass(BaseClassInfo.Name);
 			}
@@ -6667,7 +6656,7 @@ void DefineTypes(TArray<FUnrealPackageDefinitionInfo*>& PackageDefs)
 				for (TSharedRef<FUnrealTypeDefinitionInfo>& TypeDef : SourceFile.GetDefinedClasses())
 				{
 					FUnrealClassDefinitionInfo& ClassDef = TypeDef->AsClassChecked();
-					UClass* ResultClass = ProcessParsedClass(ClassDef.IsInterface(), ClassDef.GetNameCPP(), ClassDef.GetSuperClassInfo().Name, Package, RF_Public | RF_Standalone);
+					UClass* ResultClass = ProcessParsedClass(ClassDef.IsInterface(), ClassDef.GetNameCPP(), ClassDef.GetSuperStructInfo().Name, Package, RF_Public | RF_Standalone);
 					GTypeDefinitionInfoMap.Add(ResultClass, TypeDef);
 					ClassDef.SetObject(ResultClass);
 					AllClasses.Add(ResultClass);
diff --git a/Engine/Source/Programs/UnrealHeaderTool/Private/HeaderParser.cpp b/Engine/Source/Programs/UnrealHeaderTool/Private/HeaderParser.cpp
--- a/Engine/Source/Programs/UnrealHeaderTool/Private/HeaderParser.cpp
+++ b/Engine/Source/Programs/UnrealHeaderTool/Private/HeaderParser.cpp
@@ -2081,9 +2081,9 @@ FUnrealScriptStructDefinitionInfo& FHeaderParser::CompileStructDeclaration()
 
 	// if we have a base struct, propagate inherited struct flags now
 	FUnrealFieldDefinitionInfo* BaseStructDef = nullptr;
-	if (!StructDef.GetSuperClassInfo().Name.IsEmpty())
+	if (!StructDef.GetSuperStructInfo().Name.IsEmpty())
 	{
-		const FString& ParentStructNameInScript = StructDef.GetSuperClassInfo().Name;
+		const FString& ParentStructNameInScript = StructDef.GetSuperStructInfo().Name;
 		FString ParentStructNameStripped;
 		bool bOverrideParentStructName = false;
 
@@ -2145,7 +2145,7 @@ FUnrealScriptStructDefinitionInfo& FHeaderParser::CompileStructDeclaration()
 		UScriptStruct* BaseStruct = BaseStructDef->AsScriptStructChecked().GetScriptStruct();
 		Struct->SetSuperStruct(BaseStruct);
 		StructFlags |= (BaseStruct->StructFlags & STRUCT_Inherit);
-		StructDef.GetSuperClassInfo().Struct = BaseStructDef->AsStruct();
+		StructDef.GetSuperStructInfo().Struct = BaseStructDef->AsStruct();
 	}
 
 	Scope->AddType(StructDef);
@@ -2965,11 +2965,10 @@ void FHeaderParser::VerifyPropertyMarkups(FUnrealClassDefinitionInfo& TargetClas
 			return static_cast<FUnrealFunctionDefinitionInfo*>(nullptr);
 		};
 
-		FUnrealPropertyDefinitionInfo& PropDef = GTypeDefinitionInfoMap.FindChecked<FUnrealPropertyDefinitionInfo>(Prop);
-		FPropertyBase& PropertyToken = PropDef.GetPropertyBase();
+		FPropertyBase& PropertyToken = PropertyDef->GetPropertyBase();
 
-		TGuardValue<int32> GuardedInputPos(InputPos, PropDef.GetParsePosition());
-		TGuardValue<int32> GuardedInputLine(InputLine, PropDef.GetLineNumber());
+		TGuardValue<int32> GuardedInputPos(InputPos, PropertyDef->GetParsePosition());
+		TGuardValue<int32> GuardedInputLine(InputLine, PropertyDef->GetLineNumber());
 
 		if (Prop->HasAnyPropertyFlags(CPF_RepNotify))
 		{
@@ -4742,7 +4741,7 @@ void FHeaderParser::GetVarType(
 
 FUnrealPropertyDefinitionInfo& FHeaderParser::GetVarNameAndDim
 (
-	FUnrealStructDefinitionInfo& ParentStruct, 
+	FUnrealStructDefinitionInfo& ParentStruct,
 	FPropertyBase& VarProperty,
 	EVariableCategory::Type VariableCategory,
 	ELayoutMacroType        LayoutMacroType
@@ -5023,7 +5022,7 @@ FUnrealPropertyDefinitionInfo& FHeaderParser::GetVarNameAndDim
 	FName PropertyName(Identifier);
 	//VarProperty.StartLine = InputLine;
 	//VarProperty.StartPos = InputPos;
-	FUnrealPropertyDefinitionInfo& PropDef = FPropertyTraits::CreateProperty(VarProperty, Scope, PropertyName, ObjectFlags, VariableCategory, Dimensions.String, SourceFile, InputLine, InputPos);
+	FUnrealPropertyDefinitionInfo& PropDef = FPropertyTraits::CreateProperty(VarProperty, ParentStruct, PropertyName, ObjectFlags, VariableCategory, Dimensions.String, SourceFile, InputLine, InputPos);
 	FProperty* Property = PropDef.GetProperty();
 
 	// Add the property to the parent
@@ -5620,11 +5619,11 @@ FUnrealClassDefinitionInfo& FHeaderParser::ParseClassNameDeclaration(FString& De
 	ParseInheritance(TEXT("class"), [](const TCHAR* ClassName, bool bIsSuperClass) {}); // Eat the results, already been parsed
 
 	// Collect data about the super and base classes
-	if (FUnrealClassDefinitionInfo* SuperClassDef = UHTCast<FUnrealClassDefinitionInfo>(ClassDef->GetSuperClassInfo().Struct))
+	if (FUnrealClassDefinitionInfo* SuperClassDef = UHTCast<FUnrealClassDefinitionInfo>(ClassDef->GetSuperStructInfo().Struct))
 	{
 		FoundClass->ClassCastFlags |= SuperClassDef->GetClass()->ClassCastFlags;
 	}
-	for (FUnrealStructDefinitionInfo::FBaseClassInfo& BaseClassInfo : ClassDef->GetBaseClassInfo())
+	for (FUnrealStructDefinitionInfo::FBaseStructInfo& BaseClassInfo : ClassDef->GetBaseStructInfo())
 	{
 		if (FUnrealClassDefinitionInfo* BaseClassDef = UHTCast<FUnrealClassDefinitionInfo>(BaseClassInfo.Struct))
 		{
@@ -5786,7 +5785,7 @@ FUnrealClassDefinitionInfo& FHeaderParser::CompileClassDeclaration()
 
 	// auto-create properties for all of the VFTables needed for the multiple inheritances
 	// get the inheritance parents
-	const TArray<FUnrealStructDefinitionInfo::FBaseClassInfo>& BaseClassInfos = ClassDef.GetBaseClassInfo();
+	const TArray<FUnrealStructDefinitionInfo::FBaseStructInfo>& BaseClassInfos = ClassDef.GetBaseStructInfo();
 
 	// for all base class types, make a VfTable property
 	for (int32 ParentIndex = BaseClassInfos.Num() - 1; ParentIndex >= 0; ParentIndex--)
@@ -5855,6 +5854,8 @@ FUnrealClassDefinitionInfo* FHeaderParser::ParseInterfaceNameDeclaration(FString
 	if (SuperClass == NULL)
 	{
 		FoundClass->SetSuperStruct(TempClass);
+		ClassDef->GetSuperStructInfo().Struct = TempClassDef;
+		ClassDef->GetSuperStructInfo().Name = TempClass->GetName();
 	}
 	else if (SuperClass != TempClass)
 	{
@@ -7313,7 +7314,7 @@ bool FHeaderParser::ValidateScriptStructOkForNet(const FString& OriginStructName
 
 	bool bIsStructValid = true;
 
-	if (FUnrealScriptStructDefinitionInfo* SuperScriptStructDef = UHTCast<FUnrealScriptStructDefinitionInfo>(InStructDef.GetSuperClassInfo().Struct))
+	if (FUnrealScriptStructDefinitionInfo* SuperScriptStructDef = UHTCast<FUnrealScriptStructDefinitionInfo>(InStructDef.GetSuperStructInfo().Struct))
 	{
 		if (!ValidateScriptStructOkForNet(OriginStructName, *SuperScriptStructDef))
 		{
@@ -8735,11 +8736,11 @@ TSharedRef<FUnrealTypeDefinitionInfo> FHeaderPreParser::ParseClassDeclaration(FU
 		SourceFile.AddClassIncludeIfNeeded(InputLine, ClassNameWithoutPrefixStr, ClassNameStr);
 		if (bIsSuperClass)
 		{
-			ClassDef.GetSuperClassInfo().Name = MoveTemp(ClassNameStr);
+			ClassDef.GetSuperStructInfo().Name = MoveTemp(ClassNameStr);
 		}
 		else
 		{
-			ClassDef.GetBaseClassInfo().Emplace(FUnrealStructDefinitionInfo::FBaseClassInfo{ MoveTemp(ClassNameStr) });
+			ClassDef.GetBaseStructInfo().Emplace(FUnrealStructDefinitionInfo::FBaseStructInfo{ MoveTemp(ClassNameStr) });
 		}
 	}
 	);
@@ -8861,11 +8862,11 @@ TSharedRef<FUnrealTypeDefinitionInfo> FHeaderPreParser::ParseStructDeclaration(F
 		SourceFile.AddScriptStructIncludeIfNeeded(InputLine, StructNameStripped, StructNameStr);
 		if (bIsSuperClass)
 		{
-			StructDef.GetSuperClassInfo().Name = MoveTemp(StructNameStr);
+			StructDef.GetSuperStructInfo().Name = MoveTemp(StructNameStr);
 		}
 		else
 		{
-			StructDef.GetBaseClassInfo().Emplace(FUnrealStructDefinitionInfo::FBaseClassInfo{ MoveTemp(StructNameStr) });
+			StructDef.GetBaseStructInfo().Emplace(FUnrealStructDefinitionInfo::FBaseStructInfo{ MoveTemp(StructNameStr) });
 		}
 	}
 	);
@@ -9076,7 +9077,7 @@ void FHeaderParser::PostPopNestClass(FUnrealClassDefinitionInfo& CurrentClassDef
 	VerifyPropertyMarkups(CurrentClassDef);
 
 	// Iterate over all the interfaces we claim to implement
-	for (const FUnrealStructDefinitionInfo::FBaseClassInfo& BaseClassInfo : CurrentClassDef.GetBaseClassInfo())
+	for (const FUnrealStructDefinitionInfo::FBaseStructInfo& BaseClassInfo : CurrentClassDef.GetBaseStructInfo())
 	{
 
 		// Skip unknown base classes or things that aren't interfaces
diff --git a/Engine/Source/Programs/UnrealHeaderTool/Private/NativeClassExporter.h b/Engine/Source/Programs/UnrealHeaderTool/Private/NativeClassExporter.h
--- a/Engine/Source/Programs/UnrealHeaderTool/Private/NativeClassExporter.h
+++ b/Engine/Source/Programs/UnrealHeaderTool/Private/NativeClassExporter.h
@@ -6,7 +6,6 @@
 #include "Async/TaskGraphInterfaces.h"
 
 class FUnrealSourceFile;
-class UField;
 class FOutputDevice;
 class FUnrealPackageDefinitionInfo;
 class FUnrealPropertyDefinitionInfo;
@@ -327,9 +326,6 @@ private:
 	/** Return the address of the singleton function - handles nullptr */
 	static const FString& GetSingletonNameFuncAddr(FUnrealFieldDefinitionInfo* FieldDef, TSet<FString>* UniqueCrossModuleReferences, bool bRequiresValidObject = true);
 
-	/** Return the address of the singleton function - handles nullptr */
-	static const FString& GetSingletonNameFuncAddr(UField* Item, TSet<FString>* UniqueCrossModuleReferences, bool bRequiresValidObject = true);
-
 	/**
 	 * Returns the name (overridden if marked up) or "" wrappers for use in a string literal.
 	 */
diff --git a/Engine/Source/Programs/UnrealHeaderTool/Private/PropertyTypes.cpp b/Engine/Source/Programs/UnrealHeaderTool/Private/PropertyTypes.cpp
--- a/Engine/Source/Programs/UnrealHeaderTool/Private/PropertyTypes.cpp
+++ b/Engine/Source/Programs/UnrealHeaderTool/Private/PropertyTypes.cpp
@@ -665,7 +665,8 @@ struct FPropertyTypeTraitsEnum : public FPropertyTypeTraitsBase
 		}
 
 		FEnumProperty* Result = new FEnumProperty(Scope, Name, ObjectFlags);
-		FUnrealPropertyDefinitionInfo& SubProp = FPropertyTraits::CreateProperty(UnderlyingProperty, Result, TEXT("UnderlyingType"), ObjectFlags, VariableCategory, Dimensions, PropDef.GetUnrealSourceFile(), PropDef.GetLineNumber(), PropDef.GetParsePosition());
+		PropDef.SetProperty(Result);
+		FUnrealPropertyDefinitionInfo& SubProp = FPropertyTraits::CreateProperty(UnderlyingProperty, PropDef, TEXT("UnderlyingType"), ObjectFlags, VariableCategory, Dimensions, PropDef.GetUnrealSourceFile(), PropDef.GetLineNumber(), PropDef.GetParsePosition());
 		PropDef.SetValuePropDef(SubProp);
 		Result->UnderlyingProp = CastFieldChecked<FNumericProperty>(SubProp.GetProperty());
 		Result->Enum = Enum;
@@ -1298,10 +1299,11 @@ struct FPropertyTypeTraitsDynamicArray : public FPropertyTypeTraitsBase
 		FPropertyBase& VarProperty = PropDef.GetPropertyBase();
 
 		FArrayProperty* Array = new FArrayProperty(Scope, Name, ObjectFlags);
+		PropDef.SetProperty(Array);
 
 		FPropertyBase InnerVarProperty = PropDef.GetPropertyBase();
 		InnerVarProperty.ArrayType = EArrayType::None;
-		FUnrealPropertyDefinitionInfo& InnerPropDef = FPropertyTraits::CreateProperty(InnerVarProperty, Array, Name, RF_Public, VariableCategory, Dimensions, PropDef.GetUnrealSourceFile(), PropDef.GetLineNumber(), PropDef.GetParsePosition());
+		FUnrealPropertyDefinitionInfo& InnerPropDef = FPropertyTraits::CreateProperty(InnerVarProperty, PropDef, Name, RF_Public, VariableCategory, Dimensions, PropDef.GetUnrealSourceFile(), PropDef.GetLineNumber(), PropDef.GetParsePosition());
 		FProperty* InnerProp = InnerPropDef.GetProperty();
 
 		Array->Inner = InnerProp;
@@ -1330,10 +1332,11 @@ struct FPropertyTypeTraitsSet : public FPropertyTypeTraitsBase
 		FPropertyBase& VarProperty = PropDef.GetPropertyBase();
 
 		FSetProperty* Set = new FSetProperty(Scope, Name, ObjectFlags);
+		PropDef.SetProperty(Set);
 
 		FPropertyBase InnerVarProperty = PropDef.GetPropertyBase();
 		InnerVarProperty.ArrayType = EArrayType::None;
-		FUnrealPropertyDefinitionInfo& InnerPropDef = FPropertyTraits::CreateProperty(InnerVarProperty, Set, Name, RF_Public, VariableCategory, Dimensions, PropDef.GetUnrealSourceFile(), PropDef.GetLineNumber(), PropDef.GetParsePosition());
+		FUnrealPropertyDefinitionInfo& InnerPropDef = FPropertyTraits::CreateProperty(InnerVarProperty, PropDef, Name, RF_Public, VariableCategory, Dimensions, PropDef.GetUnrealSourceFile(), PropDef.GetLineNumber(), PropDef.GetParsePosition());
 		FProperty* InnerProp = InnerPropDef.GetProperty();
 
 		Set->ElementProp = InnerProp;
@@ -1361,13 +1364,14 @@ struct FPropertyTypeTraitsMap : public FPropertyTypeTraitsBase
 		FPropertyBase& VarProperty = PropDef.GetPropertyBase();
 
 		FMapProperty* Map = new FMapProperty(Scope, Name, ObjectFlags);
+		PropDef.SetProperty(Map);
 
-		FUnrealPropertyDefinitionInfo& KeyPropDef = FPropertyTraits::CreateProperty(*PropDef.GetPropertyBase().MapKeyProp, Map, *(Name.ToString() + TEXT("_Key")), RF_Public, VariableCategory, Dimensions, PropDef.GetUnrealSourceFile(), PropDef.GetLineNumber(), PropDef.GetParsePosition());
+		FUnrealPropertyDefinitionInfo& KeyPropDef = FPropertyTraits::CreateProperty(*PropDef.GetPropertyBase().MapKeyProp, PropDef, *(Name.ToString() + TEXT("_Key")), RF_Public, VariableCategory, Dimensions, PropDef.GetUnrealSourceFile(), PropDef.GetLineNumber(), PropDef.GetParsePosition());
 		FProperty* KeyProp = KeyPropDef.GetProperty();
 		FPropertyBase ValueVarProperty = PropDef.GetPropertyBase();
 		ValueVarProperty.ArrayType = EArrayType::None;
 		ValueVarProperty.MapKeyProp = nullptr;
-		FUnrealPropertyDefinitionInfo& ValuePropDef = FPropertyTraits::CreateProperty(ValueVarProperty, Map, Name, RF_Public, VariableCategory, Dimensions, PropDef.GetUnrealSourceFile(), PropDef.GetLineNumber(), PropDef.GetParsePosition());
+		FUnrealPropertyDefinitionInfo& ValuePropDef = FPropertyTraits::CreateProperty(ValueVarProperty, PropDef, Name, RF_Public, VariableCategory, Dimensions, PropDef.GetUnrealSourceFile(), PropDef.GetLineNumber(), PropDef.GetParsePosition());
 		FProperty* ValueProp = ValuePropDef.GetProperty();
 
 		Map->KeyProp = KeyProp;
@@ -1417,12 +1421,13 @@ bool FPropertyTraits::IsObject(EPropertyType PropertyType)
 	return PropertyTypeDispatch<IsObjectDispatch, bool>(PropertyType);
 }
 
-FUnrealPropertyDefinitionInfo& FPropertyTraits::CreateProperty(const FPropertyBase& VarProperty, FFieldVariant Scope, const FName& Name, EObjectFlags ObjectFlags, EVariableCategory::Type VariableCategory,
+FUnrealPropertyDefinitionInfo& FPropertyTraits::CreateProperty(const FPropertyBase& VarProperty, FUnrealTypeDefinitionInfo& Outer, const FName& Name, EObjectFlags ObjectFlags, EVariableCategory::Type VariableCategory,
 	const TCHAR* Dimensions, FUnrealSourceFile& SourceFile, int LineNumber, int ParsePosition)
 {
-	TSharedRef<FUnrealPropertyDefinitionInfo> PropDefRef = MakeShared<FUnrealPropertyDefinitionInfo>(SourceFile, LineNumber, ParsePosition, VarProperty, Name.ToString());
+	TSharedRef<FUnrealPropertyDefinitionInfo> PropDefRef = MakeShared<FUnrealPropertyDefinitionInfo>(SourceFile, LineNumber, ParsePosition, VarProperty, Name.ToString(), Outer);
 	FUnrealPropertyDefinitionInfo& PropDef = *PropDefRef;
 
+	FFieldVariant Scope = Outer.AsProperty() ? FFieldVariant(Outer.AsProperty()->GetProperty()) : FFieldVariant(Outer.AsObject()->GetObject());
 	FProperty* Property = CreatePropertyHelper<true>(PropDef, Scope, Name, ObjectFlags, VariableCategory, Dimensions);
 
 	PropDef.SetProperty(Property);
diff --git a/Engine/Source/Programs/UnrealHeaderTool/Private/PropertyTypes.h b/Engine/Source/Programs/UnrealHeaderTool/Private/PropertyTypes.h
--- a/Engine/Source/Programs/UnrealHeaderTool/Private/PropertyTypes.h
+++ b/Engine/Source/Programs/UnrealHeaderTool/Private/PropertyTypes.h
@@ -33,7 +33,7 @@ struct FPropertyTraits
 	 * Given a property definition token, create the property definition and then underlying engine FProperty
 	 *
 	 * @param Token The definition of the property
-	 * @param Scope The parent object owning the property
+	 * @param Outer The parent object owning the property
 	 * @param Name The name of the property
 	 * @param ObjectFlags The flags associated with the property
 	 * @param VariableCategory The parsing context of the property
@@ -43,7 +43,7 @@ struct FPropertyTraits
 	 * @param ParsePosition Character position of the property in the header
 	 * @return The pointer to the newly created property.  It will be attached to the definition by the caller
 	 */
-	static FUnrealPropertyDefinitionInfo& CreateProperty(const FPropertyBase& VarProperty, FFieldVariant Scope, const FName& Name, EObjectFlags ObjectFlags, EVariableCategory::Type VariableCategory, const TCHAR* Dimensions, FUnrealSourceFile& SourceFile, int LineNumber, int ParsePosition);
+	static FUnrealPropertyDefinitionInfo& CreateProperty(const FPropertyBase& VarProperty, FUnrealTypeDefinitionInfo& Outer, const FName& Name, EObjectFlags ObjectFlags, EVariableCategory::Type VariableCategory, const TCHAR* Dimensions, FUnrealSourceFile& SourceFile, int LineNumber, int ParsePosition);
 
 	/**
 	 * Test to see if the property can be used in a blueprint
diff --git a/Engine/Source/Programs/UnrealHeaderTool/Private/UnrealTypeDefinitionInfo.cpp b/Engine/Source/Programs/UnrealHeaderTool/Private/UnrealTypeDefinitionInfo.cpp
--- a/Engine/Source/Programs/UnrealHeaderTool/Private/UnrealTypeDefinitionInfo.cpp
+++ b/Engine/Source/Programs/UnrealHeaderTool/Private/UnrealTypeDefinitionInfo.cpp
@@ -312,15 +312,15 @@ bool FUnrealPropertyDefinitionInfo::IsDynamic() const
 
 bool FUnrealPropertyDefinitionInfo::IsOwnedByDynamicType() const
 {
-	for (FFieldVariant Owner = GetProperty()->GetOwnerVariant(); Owner.IsValid(); Owner = Owner.GetOwnerVariant())
+	for (FUnrealTypeDefinitionInfo* Owner = GetOuter(); Owner; Owner = Owner->GetOuter())
 	{
-		if (Owner.IsUObject())
+		if (FUnrealPropertyDefinitionInfo* PropDef = UHTCast<FUnrealPropertyDefinitionInfo>(Owner))
 		{
-			return GTypeDefinitionInfoMap.FindChecked<FUnrealObjectDefinitionInfo>(Owner.ToUObject()).IsOwnedByDynamicType();
+			return PropDef->IsOwnedByDynamicType();
 		}
-		else
+		else if (FUnrealObjectDefinitionInfo* ObjectDef = UHTCast<FUnrealObjectDefinitionInfo>(Owner))
 		{
-			return GTypeDefinitionInfoMap.FindChecked(Owner.ToField()).IsOwnedByDynamicType();
+			return ObjectDef->IsOwnedByDynamicType();
 		}
 	}
 	return false;
@@ -434,6 +434,18 @@ bool FUnrealFieldDefinitionInfo::IsOwnedByDynamicType() const
 	return false;
 }
 
+FUnrealClassDefinitionInfo* FUnrealFieldDefinitionInfo::GetOwnerClass() const
+{
+	for (FUnrealTypeDefinitionInfo* TypeDef = GetOuter(); TypeDef; TypeDef = TypeDef->GetOuter())
+	{
+		if (FUnrealClassDefinitionInfo* ClassDef = UHTCast<FUnrealClassDefinitionInfo>(TypeDef))
+		{
+			return ClassDef;
+		}
+	}
+	return nullptr;
+}
+
 FUnrealEnumDefinitionInfo::FUnrealEnumDefinitionInfo(FUnrealSourceFile& InSourceFile, int32 InLineNumber, FString&& InNameCPP)
 	: FUnrealFieldDefinitionInfo(InSourceFile, InLineNumber, MoveTemp(InNameCPP), InSourceFile.GetPackageDef())
 { }
@@ -1201,3 +1213,8 @@ void FUnrealFunctionDefinitionInfo::AddProperty(FUnrealPropertyDefinitionInfo& P
 	}
 	FUnrealStructDefinitionInfo::AddProperty(PropertyDef);
 }
+
+FUnrealFunctionDefinitionInfo* FUnrealFunctionDefinitionInfo::GetSuperFunction() const
+{
+	return UHTCast<FUnrealFunctionDefinitionInfo>(GetSuperStruct());
+}
diff --git a/Engine/Source/Programs/UnrealHeaderTool/Private/UnrealTypeDefinitionInfo.h b/Engine/Source/Programs/UnrealHeaderTool/Private/UnrealTypeDefinitionInfo.h
--- a/Engine/Source/Programs/UnrealHeaderTool/Private/UnrealTypeDefinitionInfo.h
+++ b/Engine/Source/Programs/UnrealHeaderTool/Private/UnrealTypeDefinitionInfo.h
@@ -335,6 +335,14 @@ public:
 		return false;
 	}
 
+	/**
+	 * Return the "outer" object that contains the definition for this object
+	 */
+	FUnrealTypeDefinitionInfo* GetOuter() const
+	{
+		return Outer;
+	}
+
 	static FString GetNameWithPrefix(const UClass* Class, EEnforceInterfacePrefix EnforceInterfacePrefix = EEnforceInterfacePrefix::None);
 
 protected:
@@ -342,14 +350,16 @@ protected:
 		: NameCPP(MoveTemp(InNameCPP))
 	{ }
 
-	FUnrealTypeDefinitionInfo(FUnrealSourceFile& InSourceFile, int32 InLineNumber, FString&& InNameCPP)
+	FUnrealTypeDefinitionInfo(FUnrealSourceFile& InSourceFile, int32 InLineNumber, FString&& InNameCPP, FUnrealTypeDefinitionInfo& InOuter)
 		: NameCPP(MoveTemp(InNameCPP))
+		, Outer(&InOuter)
 		, SourceFile(&InSourceFile)
 		, LineNumber(InLineNumber)
 	{ }
 
 private:
 	FString NameCPP;
+	FUnrealTypeDefinitionInfo* Outer = nullptr;
 	FUnrealSourceFile* SourceFile = nullptr;
 	int32 LineNumber = 0;
 	std::atomic<uint32> Hash = 0;
@@ -362,8 +372,8 @@ class FUnrealPropertyDefinitionInfo
 	: public FUnrealTypeDefinitionInfo
 {
 public:
-	FUnrealPropertyDefinitionInfo(FUnrealSourceFile& InSourceFile, int32 InLineNumber, int32 InParsePosition, const FPropertyBase& InVarProperty, FString&& InNameCPP)
-		: FUnrealTypeDefinitionInfo(InSourceFile, InLineNumber, MoveTemp(InNameCPP))
+	FUnrealPropertyDefinitionInfo(FUnrealSourceFile& InSourceFile, int32 InLineNumber, int32 InParsePosition, const FPropertyBase& InVarProperty, FString&& InNameCPP, FUnrealTypeDefinitionInfo& InOuter)
+		: FUnrealTypeDefinitionInfo(InSourceFile, InLineNumber, MoveTemp(InNameCPP), InOuter)
 		, PropertyBase(InVarProperty)
 		, ParsePosition(InParsePosition)
 	{ }
@@ -400,6 +410,7 @@ public:
 	 */
 	void SetProperty(FProperty* InProperty)
 	{
+		check(Property == nullptr || Property == InProperty);
 		Property = InProperty;
 	}
 
@@ -603,7 +614,7 @@ public:
 	 */
 	FUnrealObjectDefinitionInfo* GetOuter() const
 	{
-		return Outer;
+		return static_cast<FUnrealObjectDefinitionInfo*>(FUnrealTypeDefinitionInfo::GetOuter());
 	}
 
 protected:
@@ -613,12 +624,10 @@ protected:
 	{ }
 
 	FUnrealObjectDefinitionInfo(FUnrealSourceFile& InSourceFile, int32 InLineNumber, FString&& InNameCPP, FUnrealObjectDefinitionInfo& InOuter)
-		: FUnrealTypeDefinitionInfo(InSourceFile, InLineNumber, MoveTemp(InNameCPP))
-		, Outer(&InOuter)
+		: FUnrealTypeDefinitionInfo(InSourceFile, InLineNumber, MoveTemp(InNameCPP), InOuter)
 	{ }
 
 private:
-	FUnrealObjectDefinitionInfo* Outer = nullptr;
 	UObject* Object = nullptr;
 };
 
@@ -836,6 +845,11 @@ public:
 	 */
 	virtual bool IsOwnedByDynamicType() const override;
 
+	/**
+	 * Return the owning class
+	 */
+	FUnrealClassDefinitionInfo* GetOwnerClass() const;
+
 private:
 	FString SingletonName[2];
 	FString SingletonNameChopped[2];
@@ -915,7 +929,7 @@ private:
 class FUnrealStructDefinitionInfo : public FUnrealFieldDefinitionInfo
 {
 public:
-	struct FBaseClassInfo
+	struct FBaseStructInfo
 	{
 		FString Name;
 		FUnrealStructDefinitionInfo* Struct = nullptr;
@@ -999,26 +1013,37 @@ public:
 	}
 
 	/**
-	 * Return the super class information
+	 * Return the super struct
 	 */
-	const FBaseClassInfo& GetSuperClassInfo() const
+	FUnrealStructDefinitionInfo* GetSuperStruct() const
 	{
-		return SuperClassInfo;
+		return SuperStructInfo.Struct;
 	}
 
-	FBaseClassInfo& GetSuperClassInfo()
+	/**
+	 * Return the super struct information
+	 */
+	const FBaseStructInfo& GetSuperStructInfo() const
 	{
-		return SuperClassInfo;
+		return SuperStructInfo;
 	}
 
-	const TArray<FBaseClassInfo>& GetBaseClassInfo() const
+	FBaseStructInfo& GetSuperStructInfo()
 	{
-		return BaseClassInfo;
+		return SuperStructInfo;
 	}
 
-	TArray<FBaseClassInfo>& GetBaseClassInfo()
+	/**
+	 * Return the base struct information
+	 */
+	const TArray<FBaseStructInfo>& GetBaseStructInfo() const
 	{
-		return BaseClassInfo;
+		return BaseStructInfo;
+	}
+
+	TArray<FBaseStructInfo>& GetBaseStructInfo()
+	{
+		return BaseStructInfo;
 	}
 
 	/**
@@ -1042,7 +1067,7 @@ public:
 	 */
 	bool IsChildOf(FUnrealStructDefinitionInfo& ParentStruct) const
 	{
-		for (const FUnrealStructDefinitionInfo* Current = this; Current; Current = Current->GetSuperClassInfo().Struct)
+		for (const FUnrealStructDefinitionInfo* Current = this; Current; Current = Current->GetSuperStructInfo().Struct)
 		{
 			if (Current == &ParentStruct)
 			{
@@ -1110,8 +1135,8 @@ private:
 	TArray<FUnrealFunctionDefinitionInfo*> Functions;
 
 	FStructMetaData StructMetaData;
-	FBaseClassInfo SuperClassInfo;
-	TArray<FBaseClassInfo> BaseClassInfo;
+	FBaseStructInfo SuperStructInfo;
+	TArray<FBaseStructInfo> BaseStructInfo;
 
 	FDefinitionRange DefinitionRange;
 
@@ -1220,6 +1245,11 @@ public:
 		FunctionData.FunctionExportFlags &= ~ClearFlags;
 	}
 
+	/**
+	 * Return the super function
+	 */
+	FUnrealFunctionDefinitionInfo* GetSuperFunction() const;
+
 private:
 
 	/** info about the function associated with this FFunctionData */
@@ -1297,7 +1327,7 @@ public:
 
 	FUnrealClassDefinitionInfo* GetSuperClass() const
 	{
-		if (FUnrealStructDefinitionInfo* SuperClass = GetSuperClassInfo().Struct)
+		if (FUnrealStructDefinitionInfo* SuperClass = GetSuperStructInfo().Struct)
 		{
 			return &SuperClass->AsClassChecked();
 		}
@@ -1739,7 +1769,7 @@ protected:
 
 			if (bIncludeSuper)
 			{
-				CurrentStructDef = CurrentStructDef->GetSuperClassInfo().Struct;
+				CurrentStructDef = CurrentStructDef->GetSuperStructInfo().Struct;
 				if (CurrentStructDef)
 				{
 					CurrentField = UpdateRange(CurrentStructDef);
