Changed more routines to not use UFunction type.
Removed GetEnum as much as possible.

#rnx

diff --git a/Engine/Source/Programs/UnrealHeaderTool/Private/CodeGenerator.cpp b/Engine/Source/Programs/UnrealHeaderTool/Private/CodeGenerator.cpp
--- a/Engine/Source/Programs/UnrealHeaderTool/Private/CodeGenerator.cpp
+++ b/Engine/Source/Programs/UnrealHeaderTool/Private/CodeGenerator.cpp
@@ -535,26 +535,26 @@ FParmsAndReturnProperties GetFunctionParmsAndReturn(FUnrealFunctionDefinitionInf
  * @param	Function	the function to check
  * @return	true if the glue version of the function should be exported.
  */
-bool ShouldExportUFunction(UFunction* Function)
+bool ShouldExportFunction(FUnrealFunctionDefinitionInfo& FunctionDef)
 {
 	// export any script stubs for native functions declared in interface classes
-	bool bIsBlueprintNativeEvent = (Function->FunctionFlags & FUNC_BlueprintEvent) && (Function->FunctionFlags & FUNC_Native);
-	if (Function->GetOwnerClass()->HasAnyClassFlags(CLASS_Interface) && !bIsBlueprintNativeEvent)
+	bool bIsBlueprintNativeEvent = FunctionDef.HasAllFunctionFlags(FUNC_BlueprintEvent | FUNC_Native);
+	if (FunctionDef.GetOwnerClass()->HasAnyClassFlags(CLASS_Interface) && !bIsBlueprintNativeEvent)
 	{
 		return true;
 	}
 
 	// always export if the function is static
-	if (Function->FunctionFlags & FUNC_Static)
+	if (FunctionDef.HasAnyFunctionFlags(FUNC_Static))
 	{
 		return true;
 	}
 
 	// don't export the function if this is not the original declaration and there is
 	// at least one parent version of the function that is declared native
-	for (UFunction* ParentFunction = Function->GetSuperFunction(); ParentFunction; ParentFunction = ParentFunction->GetSuperFunction())
+	for (FUnrealFunctionDefinitionInfo* ParentFunctionDef = FunctionDef.GetSuperFunction(); ParentFunctionDef; ParentFunctionDef = ParentFunctionDef->GetSuperFunction())
 	{
-		if (ParentFunction->FunctionFlags & FUNC_Native)
+		if (ParentFunctionDef->HasAnyFunctionFlags(FUNC_Native))
 		{
 			return false;
 		}
@@ -687,17 +687,21 @@ TMap<FName, FString> GenerateMetadataMapForField(const FField* Field)
 }
 
 // Returns the METADATA_PARAMS for this output
-static FString OutputMetaDataCodeForObject(FOutputDevice& OutDeclaration, FOutputDevice& Out, FFieldVariant Object, const TCHAR* MetaDataBlockName, const TCHAR* DeclSpaces, const TCHAR* Spaces)
+static FString OutputMetaDataCodeForObject(FOutputDevice& OutDeclaration, FOutputDevice& Out, FUnrealTypeDefinitionInfo& TypeDef, const TCHAR* MetaDataBlockName, const TCHAR* DeclSpaces, const TCHAR* Spaces)
 {
 	TMap<FName, FString> MetaData;
 	
-	if (Object.IsUObject())
+	if (FUnrealObjectDefinitionInfo* ObjectDef = UHTCast<FUnrealObjectDefinitionInfo>(TypeDef))
 	{
-		MetaData = GenerateMetadataMapForObject(Object.ToUObject());
+		MetaData = GenerateMetadataMapForObject(ObjectDef->GetObject());
+	}
+	else if (FUnrealPropertyDefinitionInfo* PropertyDef = UHTCast<FUnrealPropertyDefinitionInfo>(TypeDef))
+	{
+		MetaData = GenerateMetadataMapForField(PropertyDef->GetProperty());
 	}
 	else
 	{
-		MetaData = GenerateMetadataMapForField(Object.ToField());
+		FError::Throwf(TEXT("Unknown type definition"));
 	}
 
 	FString Result;
@@ -916,7 +920,7 @@ void FNativeClassHeaderGenerator::PropertyNew(FOutputDevice& DeclOut, FOutputDev
 
 	FString ArrayDim = (Prop->ArrayDim != 1) ? FString::Printf(TEXT("CPP_ARRAY_DIM(%s, %s)"), *PropNameDep, SourceStruct) : TEXT("1");
 
-	FString MetaDataParams = OutputMetaDataCodeForObject(DeclOut, Out, Prop, *FString::Printf(TEXT("%s_MetaData"), Name), DeclSpaces, Spaces);
+	FString MetaDataParams = OutputMetaDataCodeForObject(DeclOut, Out, PropertyDef, *FString::Printf(TEXT("%s_MetaData"), Name), DeclSpaces, Spaces);
 
 	FString NameWithoutScope = Name;
 	FString Scope;
@@ -1730,16 +1734,16 @@ TTuple<FString, FString> FNativeClassHeaderGenerator::OutputProperties(FOutputDe
 	}
 }
 
-inline FString GetEventStructParamsName(UObject* Outer, const TCHAR* FunctionName)
+inline FString GetEventStructParamsName(FUnrealObjectDefinitionInfo& OuterDef, const TCHAR* FunctionName)
 {
 	FString OuterName;
-	if (Outer->IsA<UClass>())
+	if (FUnrealClassDefinitionInfo* ClassDef = UHTCast<FUnrealClassDefinitionInfo>(OuterDef))
 	{
-		OuterName = ((UClass*)Outer)->GetName();
+		OuterName = ClassDef->GetClass()->GetName();
 	}
-	else if (Outer->IsA<UPackage>())
+	else if (FUnrealPackageDefinitionInfo* PackageDef = UHTCast<FUnrealPackageDefinitionInfo>(OuterDef))
 	{
-		OuterName = ((UPackage*)Outer)->GetName();
+		OuterName = PackageDef->GetPackage()->GetName();
 		OuterName.ReplaceInline(TEXT("/"), TEXT("_"), ESearchCase::CaseSensitive);
 	}
 	else
@@ -1826,19 +1830,19 @@ void FNativeClassHeaderGenerator::OutputProperty(FOutputDevice& DeclOut, FOutput
 
 	{
 		FString SourceStruct;
-		if (UFunction* Function = Prop->GetOwner<UFunction>())
+		if (FUnrealFunctionDefinitionInfo* FunctionDef = UHTCast<FUnrealFunctionDefinitionInfo>(PropertyDef.GetOuter()))
 		{
-			while (Function->GetSuperFunction())
+			while (FunctionDef->GetSuperFunction())
 			{
-				Function = Function->GetSuperFunction();
+				FunctionDef = FunctionDef->GetSuperFunction();
 			}
-			FString FunctionName = Function->GetName();
-			if (Function->HasAnyFunctionFlags(FUNC_Delegate))
+			FString FunctionName = FunctionDef->GetName();
+			if (FunctionDef->HasAnyFunctionFlags(FUNC_Delegate))
 			{
 				FunctionName.LeftChopInline(HEADER_GENERATED_DELEGATE_SIGNATURE_SUFFIX_LENGTH, false);
 			}
 
-			SourceStruct = GetEventStructParamsName(Function->GetOuter(), *FunctionName);
+			SourceStruct = GetEventStructParamsName(*FunctionDef->GetOuter(), *FunctionName);
 		}
 		else
 		{
@@ -1939,8 +1943,8 @@ void FNativeClassHeaderGenerator::ExportGeneratedPackageInitCode(FOutputDevice&
 
 	Algo::Sort(Singletons, [](FUnrealFieldDefinitionInfo* A, FUnrealFieldDefinitionInfo* B)
 	{
-		bool bADel = A->GetField()->IsA<UDelegateFunction>();
-		bool bBDel = B->GetField()->IsA<UDelegateFunction>();
+		bool bADel = A->IsADelegateFunction();
+		bool bBDel = B->IsADelegateFunction();
 		if (bADel != bBDel)
 		{
 			return !bADel;
@@ -1956,7 +1960,7 @@ void FNativeClassHeaderGenerator::ExportGeneratedPackageInitCode(FOutputDevice&
 	}
 
 	FOutputDeviceNull OutputDeviceNull;
-	FString MetaDataParams = OutputMetaDataCodeForObject(OutputDeviceNull, Out, Package, TEXT("Package_MetaDataParams"), TEXT(""), TEXT("\t\t\t"));
+	FString MetaDataParams = OutputMetaDataCodeForObject(OutputDeviceNull, Out, PackageDef, TEXT("Package_MetaDataParams"), TEXT(""), TEXT("\t\t\t"));
 
 	Out.Logf(TEXT("\tUPackage* %s\r\n"), *SingletonName);
 	Out.Logf(TEXT("\t{\r\n"));
@@ -2018,28 +2022,27 @@ void FNativeClassHeaderGenerator::ExportNativeGeneratedInitCode(FOutputDevice& O
 	bool bAllEditorOnlyFunctions = true;
 	for (FUnrealFunctionDefinitionInfo* LocalFuncDef : ClassDef.GetFunctions())
 	{
-		UFunction* LocalFunc = LocalFuncDef->GetFunction();
-		FName TrueName = FNativeClassHeaderGenerator::GetOverriddenFName(LocalFunc);
+		FName TrueName = FNativeClassHeaderGenerator::GetOverriddenFName(LocalFuncDef);
 		bool bAlreadyIncluded = false;
 		AlreadyIncludedNames.Add(TrueName, &bAlreadyIncluded);
 		if (bAlreadyIncluded)
 		{
 			// In a dynamic class the same function signature may be used for a Multi- and a Single-cast delegate.
-			if (!LocalFunc->IsA<UDelegateFunction>() || !bIsDynamic)
+			if (!LocalFuncDef->IsADelegateFunction() || !bIsDynamic)
 			{
-				FError::Throwf(TEXT("The same function linked twice. Function: %s Class: %s"), *LocalFunc->GetName(), *Class->GetName());
+				FError::Throwf(TEXT("The same function linked twice. Function: %s Class: %s"), *LocalFuncDef->GetName(), *Class->GetName());
 			}
 			continue;
 		}
-		if (!LocalFunc->IsA<UDelegateFunction>())
+		if (!LocalFuncDef->IsADelegateFunction())
 		{
-			bAllEditorOnlyFunctions &= LocalFunc->HasAnyFunctionFlags(FUNC_EditorOnly);
+			bAllEditorOnlyFunctions &= LocalFuncDef->HasAnyFunctionFlags(FUNC_EditorOnly);
 		}
 		FunctionsToExport.Add(LocalFuncDef);
 	}
 
 	// Sort the list of functions
-	Algo::SortBy(FunctionsToExport, [](FUnrealFunctionDefinitionInfo* Obj) { return Obj->GetFunction()->GetName(); });
+	Algo::SortBy(FunctionsToExport, [](FUnrealFunctionDefinitionInfo* Obj) { return Obj->GetName(); });
 
 	FUHTStringBuilder GeneratedClassRegisterFunctionText;
 
@@ -2121,10 +2124,9 @@ void FNativeClassHeaderGenerator::ExportNativeGeneratedInitCode(FOutputDevice& O
 
 			for (FUnrealFunctionDefinitionInfo* FunctionDef : FunctionsToExport)
 			{
-				UFunction* Function = FunctionDef->GetFunction();
-				const bool bIsEditorOnlyFunction = Function->HasAnyFunctionFlags(FUNC_EditorOnly);
+				const bool bIsEditorOnlyFunction = FunctionDef->HasAnyFunctionFlags(FUNC_EditorOnly);
 
-				if (!Function->IsA<UDelegateFunction>())
+				if (!FunctionDef->IsADelegateFunction())
 				{
 					ExportFunction(Out, OutReferenceGatherers, SourceFile, *FunctionDef, bIsNoExport);
 				}
@@ -2136,7 +2138,7 @@ void FNativeClassHeaderGenerator::ExportNativeGeneratedInitCode(FOutputDevice& O
 					TEXT("%s\t\t{ &%s, %s },%s\r\n%s"),
 					BEGIN_WRAP_EDITOR_ONLY(bIsEditorOnlyFunction),
 					*GetSingletonNameFuncAddr(FunctionDef, OutReferenceGatherers.UniqueCrossModuleReferences),
-					*FNativeClassHeaderGenerator::GetUTF8OverriddenNameForLiteral(Function),
+					*FNativeClassHeaderGenerator::GetUTF8OverriddenNameForLiteral(FunctionDef),
 					*FuncHashTag,
 					END_WRAP_EDITOR_ONLY(bIsEditorOnlyFunction)
 				);
@@ -2171,7 +2173,7 @@ void FNativeClassHeaderGenerator::ExportNativeGeneratedInitCode(FOutputDevice& O
 			}
 		}
 
-		FString MetaDataParams = OutputMetaDataCodeForObject(GeneratedClassRegisterFunctionText, StaticDefinitions, Class, *FString::Printf(TEXT("%s::Class_MetaDataParams"), *StaticsStructName), TEXT("\t\t"), TEXT("\t"));
+		FString MetaDataParams = OutputMetaDataCodeForObject(GeneratedClassRegisterFunctionText, StaticDefinitions, ClassDef, *FString::Printf(TEXT("%s::Class_MetaDataParams"), *StaticsStructName), TEXT("\t\t"), TEXT("\t"));
 
 		TTuple<FString, FString> PropertyRange = OutputProperties(GeneratedClassRegisterFunctionText, StaticDefinitions, OutReferenceGatherers, *FString::Printf(TEXT("%s::"), *StaticsStructName), ClassDef.GetProperties(), TEXT("\t\t"), TEXT("\t"));
 
@@ -2416,15 +2418,13 @@ void FNativeClassHeaderGenerator::ExportNativeGeneratedInitCode(FOutputDevice& O
 
 void FNativeClassHeaderGenerator::ExportFunction(FOutputDevice& Out, FReferenceGatherers& OutReferenceGatherers, const FUnrealSourceFile& SourceFile, FUnrealFunctionDefinitionInfo& FunctionDef, bool bIsNoExport) const
 {
-	UFunction* Function = FunctionDef.GetFunction();
 	FunctionDef.AddCrossModuleReference(OutReferenceGatherers.UniqueCrossModuleReferences, true);
 
 	FUnrealFunctionDefinitionInfo* SuperFunctionDef = FunctionDef.GetSuperFunction();
-	UFunction* SuperFunction = SuperFunctionDef ? SuperFunctionDef->GetFunction() : nullptr;
 
-	const bool bIsEditorOnlyFunction = Function->HasAnyFunctionFlags(FUNC_EditorOnly);
+	const bool bIsEditorOnlyFunction = FunctionDef.HasAnyFunctionFlags(FUNC_EditorOnly);
 
-	bool bIsDelegate = Function->HasAnyFunctionFlags(FUNC_Delegate);
+	bool bIsDelegate = FunctionDef.HasAnyFunctionFlags(FUNC_Delegate);
 
 	const FString& SingletonName = FunctionDef.GetSingletonName(true);
 	FString StaticsStructName = FunctionDef.GetSingletonNameChopped(true) + TEXT("_Statics");
@@ -2441,7 +2441,7 @@ void FNativeClassHeaderGenerator::ExportFunction(FOutputDevice& Out, FReferenceG
 	CurrentFunctionText.Logf(TEXT("\tstruct %s\r\n"), *StaticsStructName);
 	CurrentFunctionText.Log (TEXT("\t{\r\n"));
 
-	bool bParamsInStatic = bIsNoExport || !(Function->FunctionFlags & FUNC_Event); // non-events do not export a params struct, so lets do that locally for offset determination
+	bool bParamsInStatic = bIsNoExport || !FunctionDef.HasAnyFunctionFlags(FUNC_Event); // non-events do not export a params struct, so lets do that locally for offset determination
 	if (bParamsInStatic)
 	{
 		TArray<FUnrealScriptStructDefinitionInfo*> StructDefs = FindNoExportStructs(&FunctionDef);
@@ -2470,24 +2470,24 @@ void FNativeClassHeaderGenerator::ExportFunction(FOutputDevice& Out, FReferenceG
 	FString StructureSize;
 	if (FunctionDef.GetProperties().Num())
 	{
-		UFunction* TempFunction = Function;
-		while (TempFunction->GetSuperFunction())
+		FUnrealFunctionDefinitionInfo* TempFunctionDef = &FunctionDef;
+		while (TempFunctionDef->GetSuperFunction())
 		{
-			TempFunction = TempFunction->GetSuperFunction();
+			TempFunctionDef = TempFunctionDef->GetSuperFunction();
 		}
-		FString FunctionName = TempFunction->GetName();
-		if (TempFunction->HasAnyFunctionFlags(FUNC_Delegate))
+		FString FunctionName = TempFunctionDef->GetName();
+		if (TempFunctionDef->HasAnyFunctionFlags(FUNC_Delegate))
 		{
 			FunctionName.LeftChopInline(HEADER_GENERATED_DELEGATE_SIGNATURE_SUFFIX_LENGTH, false);
 		}
 
 		if (bParamsInStatic)
 		{
-			StructureSize = FString::Printf(TEXT("sizeof(%s::%s)"), *StaticsStructName, *GetEventStructParamsName(TempFunction->GetOuter(), *FunctionName));
+			StructureSize = FString::Printf(TEXT("sizeof(%s::%s)"), *StaticsStructName, *GetEventStructParamsName(*TempFunctionDef->GetOuter(), *FunctionName));
 		}
 		else
 		{
-			StructureSize = FString::Printf(TEXT("sizeof(%s)"), *GetEventStructParamsName(TempFunction->GetOuter(), *FunctionName));
+			StructureSize = FString::Printf(TEXT("sizeof(%s)"), *GetEventStructParamsName(*TempFunctionDef->GetOuter(), *FunctionName));
 		}
 	}
 	else
@@ -2495,15 +2495,15 @@ void FNativeClassHeaderGenerator::ExportFunction(FOutputDevice& Out, FReferenceG
 		StructureSize = TEXT("0");
 	}
 
-	USparseDelegateFunction* SparseDelegateFunction = Cast<USparseDelegateFunction>(Function);
+	USparseDelegateFunction* SparseDelegateFunction = Cast<USparseDelegateFunction>(FunctionDef.GetFunction());
 	const TCHAR* UFunctionObjectFlags = FunctionDef.IsOwnedByDynamicType() ? TEXT("RF_Public|RF_Transient") : TEXT("RF_Public|RF_Transient|RF_MarkAsNative");
 
 	TTuple<FString, FString> PropertyRange = OutputProperties(CurrentFunctionText, StaticDefinitions, OutReferenceGatherers, *FString::Printf(TEXT("%s::"), *StaticsStructName), FunctionDef.GetProperties(), TEXT("\t\t"), TEXT("\t"));
 
 	const FFuncInfo&     FunctionData = FunctionDef.GetFunctionData();
-	const bool           bIsNet       = !!(FunctionData.FunctionFlags & (FUNC_NetRequest | FUNC_NetResponse));
+	const bool           bIsNet       = FunctionDef.HasAnyFunctionFlags(FUNC_NetRequest | FUNC_NetResponse);
 
-	FString MetaDataParams = OutputMetaDataCodeForObject(CurrentFunctionText, StaticDefinitions, Function, *FString::Printf(TEXT("%s::Function_MetaDataParams"), *StaticsStructName), TEXT("\t\t"), TEXT("\t"));
+	FString MetaDataParams = OutputMetaDataCodeForObject(CurrentFunctionText, StaticDefinitions, FunctionDef, *FString::Printf(TEXT("%s::Function_MetaDataParams"), *StaticsStructName), TEXT("\t\t"), TEXT("\t"));
 
 	CurrentFunctionText.Log(TEXT("\t\tstatic const UECodeGen_Private::FFunctionParams FuncParams;\r\n"));
 
@@ -2512,14 +2512,14 @@ void FNativeClassHeaderGenerator::ExportFunction(FOutputDevice& Out, FReferenceG
 		*StaticsStructName,
 		*OuterFunc,
 		*GetSingletonNameFuncAddr(SuperFunctionDef, OutReferenceGatherers.UniqueCrossModuleReferences),
-		*CreateUTF8LiteralString(FNativeClassHeaderGenerator::GetOverriddenName(Function)),
+		*CreateUTF8LiteralString(FNativeClassHeaderGenerator::GetOverriddenName(&FunctionDef)),
 		(SparseDelegateFunction ? *CreateUTF8LiteralString(SparseDelegateFunction->OwningClassName.ToString()) : TEXT("nullptr")),
 		(SparseDelegateFunction ? *CreateUTF8LiteralString(SparseDelegateFunction->DelegateName.ToString()) : TEXT("nullptr")),
 		*StructureSize,
 		*PropertyRange.Get<0>(),
 		*PropertyRange.Get<1>(),
 		UFunctionObjectFlags,
-		(uint32)Function->FunctionFlags,
+		(uint32)FunctionDef.GetFunctionFlags(),
 		bIsNet ? FunctionData.RPCId : 0,
 		bIsNet ? FunctionData.RPCResponseId : 0,
 		*MetaDataParams
@@ -2537,7 +2537,7 @@ void FNativeClassHeaderGenerator::ExportFunction(FOutputDevice& Out, FReferenceG
 	}
 	else
 	{
-		FString FunctionName = FNativeClassHeaderGenerator::GetUTF8OverriddenNameForLiteral(Function);
+		FString FunctionName = FNativeClassHeaderGenerator::GetUTF8OverriddenNameForLiteral(&FunctionDef);
 		CurrentFunctionText.Logf(TEXT("\t\tUObject* Outer = %s();\r\n"), *OuterFunc);
 		CurrentFunctionText.Logf(TEXT("\t\tUFunction* ReturnFunction = static_cast<UFunction*>(StaticFindObjectFast( UFunction::StaticClass(), Outer, %s ));\r\n"), *FunctionName);
 	}
@@ -2572,23 +2572,22 @@ void FNativeClassHeaderGenerator::ExportNatives(FOutputDevice& Out, FUnrealClass
 	{
 		bool bAllEditorOnly = true;
 
-		TArray<TTuple<UFunction*, FString>> NamedFunctionsToExport;
+		TArray<TTuple<FUnrealFunctionDefinitionInfo*, FString>> NamedFunctionsToExport;
 		for (FUnrealFunctionDefinitionInfo* FunctionDef : ClassDef.GetFunctions())
 		{
-			UFunction* Function = FunctionDef->GetFunction();
-			if ((Function->FunctionFlags & (FUNC_Native | FUNC_NetRequest)) == FUNC_Native)
+			if (FunctionDef->HasSpecificFunctionFlags(FUNC_Native | FUNC_NetRequest, FUNC_Native))
 			{
-				FString OverriddenName = FNativeClassHeaderGenerator::GetUTF8OverriddenNameForLiteral(Function);
-				NamedFunctionsToExport.Emplace(Function, MoveTemp(OverriddenName));
+				FString OverriddenName = FNativeClassHeaderGenerator::GetUTF8OverriddenNameForLiteral(FunctionDef);
+				NamedFunctionsToExport.Emplace(FunctionDef, MoveTemp(OverriddenName));
 
-				if (!Function->HasAnyFunctionFlags(FUNC_EditorOnly))
+				if (!FunctionDef->HasAnyFunctionFlags(FUNC_EditorOnly))
 				{
 					bAllEditorOnly = false;
 				}
 			}
 		}
 
-		Algo::SortBy(NamedFunctionsToExport, [](const TTuple<UFunction*, FString>& Pair){ return Pair.Get<0>()->GetFName(); }, FNameLexicalLess());
+		Algo::SortBy(NamedFunctionsToExport, [](const TTuple<FUnrealFunctionDefinitionInfo*, FString>& Pair){ return Pair.Get<0>()->GetFName(); }, FNameLexicalLess());
 
 		if (NamedFunctionsToExport.Num() > 0)
 		{
@@ -2598,17 +2597,17 @@ void FNativeClassHeaderGenerator::ExportNatives(FOutputDevice& Out, FUnrealClass
 			Out.Logf(TEXT("\t\tUClass* Class = %s::StaticClass();\r\n"), *ClassCPPName);
 			Out.Log(TEXT("\t\tstatic const FNameNativePtrPair Funcs[] = {\r\n"));
 
-			for (const TTuple<UFunction*, FString>& Func : NamedFunctionsToExport)
+			for (const TTuple<FUnrealFunctionDefinitionInfo*, FString>& Func : NamedFunctionsToExport)
 			{
-				UFunction* Function = Func.Get<0>();
+				FUnrealFunctionDefinitionInfo* FunctionDef = Func.Get<0>();
 
-				EditorOnly(Function->HasAnyFunctionFlags(FUNC_EditorOnly));
+				EditorOnly(FunctionDef->HasAnyFunctionFlags(FUNC_EditorOnly));
 
 				Out.Logf(
 					TEXT("\t\t\t{ %s, &%s::exec%s },\r\n"),
 					*Func.Get<1>(),
 					*TypeName,
-					*Function->GetName()
+					*FunctionDef->GetName()
 				);
 			}
 
@@ -2628,18 +2627,17 @@ void FNativeClassHeaderGenerator::ExportInterfaceCallFunctions(FOutputDevice& Ou
 
 	for (FUnrealFunctionDefinitionInfo* FunctionDef : CallbackFunctions)
 	{
-		UFunction* Function = FunctionDef->GetFunction();
-		FString FunctionName = Function->GetName();
+		FString FunctionName = FunctionDef->GetName();
 
 		const FFuncInfo& FunctionData = FunctionDef->GetFunctionData();
-		const TCHAR* ConstQualifier = FunctionData.FunctionReference->HasAllFunctionFlags(FUNC_Const) ? TEXT("const ") : TEXT("");
+		const TCHAR* ConstQualifier = FunctionDef->HasAllFunctionFlags(FUNC_Const) ? TEXT("const ") : TEXT("");
 		FString ExtraParam = FString::Printf(TEXT("%sUObject* O"), ConstQualifier);
 
 		ExportNativeFunctionHeader(Out, OutReferenceGatherers.ForwardDeclarations, *FunctionDef, FunctionData, EExportFunctionType::Interface, EExportFunctionHeaderStyle::Declaration, *ExtraParam, *APIString);
 		Out.Logf( TEXT(";") LINE_TERMINATOR );
 
-		FString FunctionNameName = FString::Printf(TEXT("NAME_%s_%s"), *FNameLookupCPP::GetNameCPP(CastChecked<UStruct>(Function->GetOuter())), *FunctionName);
-		OutCpp.Logf(TEXT("\tstatic FName %s = FName(TEXT(\"%s\"));") LINE_TERMINATOR, *FunctionNameName, *GetOverriddenFName(Function).ToString());
+		FString FunctionNameName = FString::Printf(TEXT("NAME_%s_%s"), *UHTCastChecked<FUnrealStructDefinitionInfo>(FunctionDef->GetOuter()).GetAlternateNameCPP(), *FunctionName);
+		OutCpp.Logf(TEXT("\tstatic FName %s = FName(TEXT(\"%s\"));") LINE_TERMINATOR, *FunctionNameName, *GetOverriddenFName(FunctionDef).ToString());
 
 		ExportNativeFunctionHeader(OutCpp, OutReferenceGatherers.ForwardDeclarations, *FunctionDef, FunctionData, EExportFunctionType::Interface, EExportFunctionHeaderStyle::Definition, *ExtraParam, *APIString);
 		OutCpp.Logf( LINE_TERMINATOR TEXT("\t{") LINE_TERMINATOR );
@@ -2653,7 +2651,7 @@ void FNativeClassHeaderGenerator::ExportInterfaceCallFunctions(FOutputDevice& Ou
 		const bool bHasParms = Parameters.HasParms();
 		if (bHasParms)
 		{
-			FString EventParmStructName = GetEventStructParamsName(Function->GetOuter(), *FunctionName);
+			FString EventParmStructName = GetEventStructParamsName(*FunctionDef->GetOuter(), *FunctionName);
 			OutCpp.Logf(TEXT("\t\t%s Parms;") LINE_TERMINATOR, *EventParmStructName);
 		}
 
@@ -2668,7 +2666,7 @@ void FNativeClassHeaderGenerator::ExportInterfaceCallFunctions(FOutputDevice& Ou
 			OutCpp.Logf(TEXT("\t\t\tParms.%s=%s;") LINE_TERMINATOR, *ParamName, *ParamName);
 		}
 
-		const FString ObjectRef = FunctionData.FunctionReference->HasAllFunctionFlags(FUNC_Const) ? FString::Printf(TEXT("const_cast<UObject*>(O)")) : TEXT("O");
+		const FString ObjectRef = FunctionDef->HasAllFunctionFlags(FUNC_Const) ? FString::Printf(TEXT("const_cast<UObject*>(O)")) : TEXT("O");
 		OutCpp.Logf(TEXT("\t\t\t%s->ProcessEvent(Func, %s);") LINE_TERMINATOR, *ObjectRef, bHasParms ? TEXT("&Parms") : TEXT("NULL"));
 
 		for (FUnrealPropertyDefinitionInfo* ParamDef : Parameters.Parms)
@@ -2685,7 +2683,7 @@ void FNativeClassHeaderGenerator::ExportInterfaceCallFunctions(FOutputDevice& Ou
 
 
 		// else clause to call back into native if it's a BlueprintNativeEvent
-		if (Function->FunctionFlags & FUNC_Native)
+		if (FunctionDef->HasAnyFunctionFlags(FUNC_Native))
 		{
 			OutCpp.Logf(TEXT("\t\telse if (auto I = (%sI%s*)(O->GetNativeInterfaceAddress(U%s::StaticClass())))") LINE_TERMINATOR, ConstQualifier, ClassName, ClassName);
 			OutCpp.Log(TEXT("\t\t{") LINE_TERMINATOR);
@@ -2800,8 +2798,7 @@ void FNativeClassHeaderGenerator::ExportClassFromSourceFileInner(
 	TArray<FUnrealFunctionDefinitionInfo*> CallbackFunctions;
 	for (FUnrealFunctionDefinitionInfo* FunctionDef : ClassDef.GetFunctions())
 	{
-		UFunction* Function = FunctionDef->GetFunction();
-		if ((Function->FunctionFlags & FUNC_Event) && Function->GetSuperFunction() == nullptr)
+		if (FunctionDef->HasAnyFunctionFlags(FUNC_Event) && FunctionDef->GetSuperFunction() == nullptr)
 		{
 			CallbackFunctions.Add(FunctionDef);
 		}
@@ -2810,7 +2807,7 @@ void FNativeClassHeaderGenerator::ExportClassFromSourceFileInner(
 	FUHTStringBuilder PrologMacroCalls;
 	if (CallbackFunctions.Num() != 0)
 	{
-		Algo::SortBy(CallbackFunctions, [](FUnrealFunctionDefinitionInfo* Obj) { return Obj->GetFunction()->GetName(); });
+		Algo::SortBy(CallbackFunctions, [](FUnrealFunctionDefinitionInfo* Obj) { return Obj->GetName(); });
 
 		FUHTStringBuilder UClassMacroContent;
 
@@ -3390,33 +3387,29 @@ FString FNativeClassHeaderGenerator::GetClassFlagExportText(FUnrealClassDefiniti
 */
 void FNativeClassHeaderGenerator::ExportEnum(FOutputDevice& Out, FUnrealEnumDefinitionInfo& EnumDef) const
 {
-	UEnum* Enum = EnumDef.GetEnum();
-
 	// Export FOREACH macro
-	Out.Logf( TEXT("#define FOREACH_ENUM_%s(op) "), *Enum->GetName().ToUpper() );
-	bool bHasExistingMax = Enum->ContainsExistingMax();
-	int64 MaxEnumVal = bHasExistingMax ? Enum->GetMaxEnumValue() : 0;
-	for (int32 i = 0; i < Enum->NumEnums(); i++)
+	Out.Logf( TEXT("#define FOREACH_ENUM_%s(op) "), *EnumDef.GetName().ToUpper() );
+	bool bHasExistingMax = EnumDef.ContainsExistingMax();
+	int64 MaxEnumVal = bHasExistingMax ? EnumDef.GetMaxEnumValue() : 0;
+	for (int32 i = 0; i < EnumDef.NumEnums(); i++)
 	{
-		if (bHasExistingMax && Enum->GetValueByIndex(i) == MaxEnumVal)
+		if (bHasExistingMax && EnumDef.GetValueByIndex(i) == MaxEnumVal)
 		{
 			continue;
 		}
 
-		const FString QualifiedEnumValue = Enum->GetNameByIndex(i).ToString();
+		const FString QualifiedEnumValue = EnumDef.GetNameByIndex(i).ToString();
 		Out.Logf( TEXT("\\\r\n\top(%s) "), *QualifiedEnumValue );
 	}
 	Out.Logf( TEXT("\r\n") );
 
 	// Forward declare the StaticEnum<> specialisation for enum classes
-	if (Enum->GetCppForm() == UEnum::ECppForm::EnumClass)
+	if (EnumDef.GetCppForm() == UEnum::ECppForm::EnumClass)
 	{
 		FString UnderlyingTypeString;
 
 		if (EnumDef.GetUnderlyingType() != EUnderlyingEnumType::Unspecified)
 		{
-			check(Enum->GetCppForm() == UEnum::ECppForm::EnumClass);
-
 			UnderlyingTypeString = TEXT(" : ");
 
 			switch (EnumDef.GetUnderlyingType())
@@ -3435,8 +3428,8 @@ void FNativeClassHeaderGenerator::ExportEnum(FOutputDevice& Out, FUnrealEnumDefi
 		}
 
 		Out.Logf( TEXT("\r\n") );
-		Out.Logf( TEXT("enum class %s%s;\r\n"), *Enum->CppType, *UnderlyingTypeString );
-		Out.Logf( TEXT("template<> %sUEnum* StaticEnum<%s>();\r\n"), *GetAPIString(), *Enum->CppType );
+		Out.Logf( TEXT("enum class %s%s;\r\n"), *EnumDef.GetCppType(), *UnderlyingTypeString );
+		Out.Logf( TEXT("template<> %sUEnum* StaticEnum<%s>();\r\n"), *GetAPIString(), *EnumDef.GetCppType());
 		Out.Logf( TEXT("\r\n") );
 	}
 }
@@ -3808,7 +3801,7 @@ void FNativeClassHeaderGenerator::ExportGeneratedStructBodyMacros(FOutputDevice&
 		OuterFunc = TEXT("&OuterFuncGetter");
 	}
 
-	FString MetaDataParams = OutputMetaDataCodeForObject(GeneratedStructRegisterFunctionText, StaticDefinitions, Struct, *FString::Printf(TEXT("%s::Struct_MetaDataParams"), *StaticsStructName), TEXT("\t\t"), TEXT("\t"));
+	FString MetaDataParams = OutputMetaDataCodeForObject(GeneratedStructRegisterFunctionText, StaticDefinitions, ScriptStructDef, *FString::Printf(TEXT("%s::Struct_MetaDataParams"), *StaticsStructName), TEXT("\t\t"), TEXT("\t"));
 
 	FString NewStructOps;
 	if (Struct->StructFlags & STRUCT_Native)
@@ -3972,11 +3965,10 @@ void FNativeClassHeaderGenerator::ExportGeneratedStructBodyMacros(FOutputDevice&
 
 void FNativeClassHeaderGenerator::ExportGeneratedEnumInitCode(FOutputDevice& Out, FReferenceGatherers& OutReferenceGatherers, const FUnrealSourceFile& SourceFile, FUnrealEnumDefinitionInfo& EnumDef) const
 {
-	UEnum* Enum                         = EnumDef.GetEnum();
 	const bool bIsDynamic               = EnumDef.IsDynamic();
 	const FString& SingletonName        = EnumDef.GetSingletonNameChopped(true);
-	const FString EnumNameCpp           = Enum->GetName(); //UserDefinedEnum should already have a valid cpp name.
-	const FString OverriddenEnumNameCpp = FNativeClassHeaderGenerator::GetOverriddenName(Enum);
+	const FString EnumNameCpp           = EnumDef.GetName(); //UserDefinedEnum should already have a valid cpp name.
+	const FString OverriddenEnumNameCpp = FNativeClassHeaderGenerator::GetOverriddenName(&EnumDef);
 	const FString StaticsStructName     = SingletonName + TEXT("_Statics");
 	const bool bIsEditorOnlyDataType    = EnumDef.IsEditorOnly();
 
@@ -3989,22 +3981,22 @@ void FNativeClassHeaderGenerator::ExportGeneratedEnumInitCode(FOutputDevice& Out
 
 	if (!bIsDynamic)
 	{
-		Out.Logf(TEXT("\tstatic FEnumRegistrationInfo& Z_Registration_Info_UEnum_%s()\r\n"), *Enum->GetName());
+		Out.Logf(TEXT("\tstatic FEnumRegistrationInfo& Z_Registration_Info_UEnum_%s()\r\n"), *EnumNameCpp);
 		Out.Logf(TEXT("\t{\r\n"));
 		Out.Logf(TEXT("\t\tstatic FEnumRegistrationInfo info;\r\n"));
 		Out.Logf(TEXT("\t\treturn info;\r\n"));
 		Out.Logf(TEXT("\t}\r\n"));
 	}
 
-	Out.Logf(TEXT("\tstatic UEnum* %s_StaticEnum()\r\n"), *Enum->GetName());
+	Out.Logf(TEXT("\tstatic UEnum* %s_StaticEnum()\r\n"), *EnumNameCpp);
 	Out.Logf(TEXT("\t{\r\n"));
 
 	if (!bIsDynamic)
 	{
-		Out.Logf(TEXT("\t\tstatic UEnum*& OuterSingleton = Z_Registration_Info_UEnum_%s().OuterSingleton;\r\n"), *Enum->GetName());
+		Out.Logf(TEXT("\t\tstatic UEnum*& OuterSingleton = Z_Registration_Info_UEnum_%s().OuterSingleton;\r\n"), *EnumNameCpp);
 		Out.Logf(TEXT("\t\tif (!OuterSingleton)\r\n"));
 		Out.Logf(TEXT("\t\t{\r\n"));
-		Out.Logf(TEXT("\t\t\tOuterSingleton = GetStaticEnum(%s, %s, TEXT(\"%s\"));\r\n"), *SingletonName, *PackageSingletonName, *Enum->GetName());
+		Out.Logf(TEXT("\t\t\tOuterSingleton = GetStaticEnum(%s, %s, TEXT(\"%s\"));\r\n"), *SingletonName, *PackageSingletonName, *EnumNameCpp);
 		Out.Logf(TEXT("\t\t}\r\n"));
 		Out.Logf(TEXT("\t\treturn OuterSingleton;\r\n"));
 	}
@@ -4024,9 +4016,9 @@ void FNativeClassHeaderGenerator::ExportGeneratedEnumInitCode(FOutputDevice& Out
 	const FString& EnumSingletonName = EnumDef.GetSingletonName(true);
 	const FString HashFuncName = FString::Printf(TEXT("Get_%s_Hash"), *SingletonName);
 
-	Out.Logf(TEXT("\ttemplate<> %sUEnum* StaticEnum<%s>()\r\n"), *GetAPIString(), *Enum->CppType);
+	Out.Logf(TEXT("\ttemplate<> %sUEnum* StaticEnum<%s>()\r\n"), *GetAPIString(), *EnumDef.GetCppType());
 	Out.Logf(TEXT("\t{\r\n"));
-	Out.Logf(TEXT("\t\treturn %s_StaticEnum();\r\n"), *Enum->GetName());
+	Out.Logf(TEXT("\t\treturn %s_StaticEnum();\r\n"), *EnumNameCpp);
 	Out.Logf(TEXT("\t}\r\n"));
 
 	FUHTStringBuilder StaticDefinitions;
@@ -4047,40 +4039,40 @@ void FNativeClassHeaderGenerator::ExportGeneratedEnumInitCode(FOutputDevice& Out
 		}
 
 		const TCHAR* UEnumObjectFlags = bIsDynamic ? TEXT("RF_Public|RF_Transient") : TEXT("RF_Public|RF_Transient|RF_MarkAsNative");
-		const TCHAR* EnumFlags = Enum->HasAnyEnumFlags(EEnumFlags::Flags) ? TEXT("EEnumFlags::Flags") : TEXT("EEnumFlags::None");
+		const TCHAR* EnumFlags = EnumDef.HasAnyEnumFlags(EEnumFlags::Flags) ? TEXT("EEnumFlags::Flags") : TEXT("EEnumFlags::None");
 
 		const TCHAR* EnumFormStr = TEXT("");
-		switch (Enum->GetCppForm())
+		switch (EnumDef.GetCppForm())
 		{
 		case UEnum::ECppForm::Regular:    EnumFormStr = TEXT("UEnum::ECppForm::Regular");    break;
 		case UEnum::ECppForm::Namespaced: EnumFormStr = TEXT("UEnum::ECppForm::Namespaced"); break;
 		case UEnum::ECppForm::EnumClass:  EnumFormStr = TEXT("UEnum::ECppForm::EnumClass");  break;
 		}
 
-		const FString& EnumDisplayNameFn = Enum->GetMetaData(TEXT("EnumDisplayNameFn"));
+		const FString& EnumDisplayNameFn = EnumDef.GetMetaData(TEXT("EnumDisplayNameFn"));
 
 		StaticDeclarations.Logf(TEXT("\tstruct %s\r\n"), *StaticsStructName);
 		StaticDeclarations.Logf(TEXT("\t{\r\n"));
 
 		StaticDeclarations.Logf(TEXT("\t\tstatic const UECodeGen_Private::FEnumeratorParam Enumerators[];\r\n"));
 		StaticDefinitions.Logf(TEXT("\tconst UECodeGen_Private::FEnumeratorParam %s::Enumerators[] = {\r\n"), *StaticsStructName);
-		for (int32 Index = 0; Index != Enum->NumEnums(); ++Index)
+		for (int32 Index = 0; Index != EnumDef.NumEnums(); ++Index)
 		{
 			const TCHAR* OverridenNameMetaDatakey = TEXT("OverrideName");
-			const FString KeyName = Enum->HasMetaData(OverridenNameMetaDatakey, Index) ? Enum->GetMetaData(OverridenNameMetaDatakey, Index) : Enum->GetNameByIndex(Index).ToString();
-			StaticDefinitions.Logf(TEXT("\t\t{ %s, (int64)%s },\r\n"), *CreateUTF8LiteralString(KeyName), *Enum->GetNameByIndex(Index).ToString());
+			const FString KeyName = EnumDef.HasMetaData(OverridenNameMetaDatakey, Index) ? EnumDef.GetMetaData(OverridenNameMetaDatakey, Index) : EnumDef.GetNameByIndex(Index).ToString();
+			StaticDefinitions.Logf(TEXT("\t\t{ %s, (int64)%s },\r\n"), *CreateUTF8LiteralString(KeyName), *EnumDef.GetNameByIndex(Index).ToString());
 		}
 		StaticDefinitions.Logf(TEXT("\t};\r\n"));
 
 		FString MetaDataParamsName = StaticsStructName + TEXT("::Enum_MetaDataParams");
-		FString MetaDataParams = OutputMetaDataCodeForObject(StaticDeclarations, StaticDefinitions, Enum, *MetaDataParamsName, TEXT("\t\t"), TEXT("\t"));
+		FString MetaDataParams = OutputMetaDataCodeForObject(StaticDeclarations, StaticDefinitions, EnumDef, *MetaDataParamsName, TEXT("\t\t"), TEXT("\t"));
 
 		StaticDeclarations.Logf(TEXT("\t\tstatic const UECodeGen_Private::FEnumParams EnumParams;\r\n"));
 		StaticDefinitions.Logf(TEXT("\tconst UECodeGen_Private::FEnumParams %s::EnumParams = {\r\n"), *StaticsStructName);
 		StaticDefinitions.Logf(TEXT("\t\t(UObject*(*)())%s,\r\n"), *OuterString.LeftChop(2));
 		StaticDefinitions.Logf(TEXT("\t\t%s,\r\n"), EnumDisplayNameFn.IsEmpty() ? TEXT("nullptr") : *EnumDisplayNameFn);
 		StaticDefinitions.Logf(TEXT("\t\t%s,\r\n"), *CreateUTF8LiteralString(OverriddenEnumNameCpp));
-		StaticDefinitions.Logf(TEXT("\t\t%s,\r\n"), *CreateUTF8LiteralString(Enum->CppType));
+		StaticDefinitions.Logf(TEXT("\t\t%s,\r\n"), *CreateUTF8LiteralString(EnumDef.GetCppType()));
 		StaticDefinitions.Logf(TEXT("\t\t%s::Enumerators,\r\n"), *StaticsStructName);
 		StaticDefinitions.Logf(TEXT("\t\tUE_ARRAY_COUNT(%s::Enumerators),\r\n"), *StaticsStructName);
 		StaticDefinitions.Logf(TEXT("\t\t%s,\r\n"), UEnumObjectFlags);
@@ -4106,7 +4098,7 @@ void FNativeClassHeaderGenerator::ExportGeneratedEnumInitCode(FOutputDevice& Out
 	// Enums can either have a UClass or UPackage as outer (if declared in non-UClass header).
 	if (!bIsDynamic)
 	{
-		GeneratedEnumRegisterFunctionText.Logf(TEXT("\t\tstatic UEnum*& ReturnEnum = Z_Registration_Info_UEnum_%s().InnerSingleton;\r\n"), *Enum->GetName());
+		GeneratedEnumRegisterFunctionText.Logf(TEXT("\t\tstatic UEnum*& ReturnEnum = Z_Registration_Info_UEnum_%s().InnerSingleton;\r\n"), *EnumNameCpp);
 	}
 	else
 	{
@@ -4146,7 +4138,7 @@ void FNativeClassHeaderGenerator::ExportGeneratedEnumInitCode(FOutputDevice& Out
 			TEXT("\tstatic FRegisterCompiledInInfo Z_CompiledInDeferEnum_UEnum_%s(%s_StaticEnum, TEXT(\"%s\"), TEXT(\"%s\"), Z_Registration_Info_UEnum_%s(), CONSTUCT_RELOAD_VERSION_INFO(FEnumReloadVersionInfo, %s()));\r\n"),
 			*EnumNameCpp,
 			*EnumNameCpp,
-			*Enum->GetOutermost()->GetName(),
+			*EnumDef.GetPackageDef().GetName(),
 			*OverriddenEnumNameCpp,
 			*EnumNameCpp,
 			*HashFuncName
@@ -4181,7 +4173,7 @@ bool FNativeClassHeaderGenerator::WillExportEventParms(FUnrealFunctionDefinition
 	return Properties.Num() > 0 && (Properties[0]->GetProperty()->PropertyFlags & CPF_Parm) != 0;
 }
 
-void WriteEventFunctionPrologue(FOutputDevice& Output, int32 Indent, const FParmsAndReturnProperties& Parameters, UObject* FunctionOuter, const TCHAR* FunctionName)
+void WriteEventFunctionPrologue(FOutputDevice& Output, int32 Indent, const FParmsAndReturnProperties& Parameters, FUnrealObjectDefinitionInfo& FunctionOuterDef, const TCHAR* FunctionName)
 {
 	// now the body - first we need to declare a struct which will hold the parameters for the event/delegate call
 	Output.Logf(TEXT("\r\n%s{\r\n"), FCString::Tab(Indent));
@@ -4190,7 +4182,7 @@ void WriteEventFunctionPrologue(FOutputDevice& Output, int32 Indent, const FParm
 	if (!Parameters.HasParms())
 		return;
 
-	FString EventStructName = GetEventStructParamsName(FunctionOuter, FunctionName);
+	FString EventStructName = GetEventStructParamsName(FunctionOuterDef, FunctionName);
 
 	Output.Logf(TEXT("%s%s Parms;\r\n"), FCString::Tab(Indent + 1), *EventStructName );
 
@@ -4251,15 +4243,14 @@ void FNativeClassHeaderGenerator::ExportDelegateDeclaration(FOutputDevice& Out,
 {
 	static const TCHAR DelegateStr[] = TEXT("delegate");
 
-	UFunction* Function = FunctionDef.GetFunction();
 	FFuncInfo FunctionData = FunctionDef.GetFunctionData(); // THIS IS A COPY
 
-	check(Function->HasAnyFunctionFlags(FUNC_Delegate));
+	check(FunctionDef.HasAnyFunctionFlags(FUNC_Delegate));
 
-	const bool bIsMulticastDelegate = Function->HasAnyFunctionFlags( FUNC_MulticastDelegate );
+	const bool bIsMulticastDelegate = FunctionDef.HasAnyFunctionFlags( FUNC_MulticastDelegate );
 
 	// Unmangle the function name
-	const FString DelegateName = Function->GetName().LeftChop(HEADER_GENERATED_DELEGATE_SIGNATURE_SUFFIX_LENGTH);
+	const FString DelegateName = FunctionDef.GetName().LeftChop(HEADER_GENERATED_DELEGATE_SIGNATURE_SUFFIX_LENGTH);
 
 	// Add class name to beginning of function, to avoid collisions with other classes with the same delegate name in this scope
 	check(FunctionData.MarshallAndCallName.StartsWith(DelegateStr));
@@ -4289,19 +4280,18 @@ void FNativeClassHeaderGenerator::ExportDelegateDefinition(FOutputDevice& Out, F
 {
 	const TCHAR DelegateStr[] = TEXT("delegate");
 
-	UFunction* Function = FunctionDef.GetFunction();
 	FFuncInfo FunctionData = FunctionDef.GetFunctionData(); // THIS IS A COPY
 
-	check(Function->HasAnyFunctionFlags(FUNC_Delegate));
+	check(FunctionDef.HasAnyFunctionFlags(FUNC_Delegate));
 
 	// Export parameters structs for all delegates.  We'll need these to declare our delegate execution function.
 	FUHTStringBuilder DelegateOutput;
 	ExportEventParm(DelegateOutput, OutReferenceGatherers.ForwardDeclarations, FunctionDef, /*Indent=*/ 0, /*bOutputConstructor=*/ true, EExportingState::Normal);
 
-	const bool bIsMulticastDelegate = Function->HasAnyFunctionFlags( FUNC_MulticastDelegate );
+	const bool bIsMulticastDelegate = FunctionDef.HasAnyFunctionFlags( FUNC_MulticastDelegate );
 
 	// Unmangle the function name
-	const FString DelegateName = Function->GetName().LeftChop(HEADER_GENERATED_DELEGATE_SIGNATURE_SUFFIX_LENGTH);
+	const FString DelegateName = FunctionDef.GetName().LeftChop(HEADER_GENERATED_DELEGATE_SIGNATURE_SUFFIX_LENGTH);
 
 	// Always export delegate wrapper functions as inline
 	FunctionData.FunctionExportFlags |= FUNCEXPORT_Inline;
@@ -4325,7 +4315,7 @@ void FNativeClassHeaderGenerator::ExportDelegateDefinition(FOutputDevice& Out, F
 
 	FParmsAndReturnProperties Parameters = GetFunctionParmsAndReturn(FunctionDef);
 
-	WriteEventFunctionPrologue(DelegateOutput, 0, Parameters, Function->GetOuter(), *DelegateName);
+	WriteEventFunctionPrologue(DelegateOutput, 0, Parameters, *FunctionDef.GetOuter(), *DelegateName);
 	{
 		const TCHAR* DelegateType = bIsMulticastDelegate ? TEXT( "ProcessMulticastDelegate" ) : TEXT( "ProcessDelegate" );
 		const TCHAR* DelegateArg  = Parameters.HasParms() ? TEXT("&Parms") : TEXT("NULL");
@@ -4339,19 +4329,18 @@ void FNativeClassHeaderGenerator::ExportDelegateDefinition(FOutputDevice& Out, F
 
 void FNativeClassHeaderGenerator::ExportEventParm(FUHTStringBuilder& Out, TSet<FString>& PropertyFwd, FUnrealFunctionDefinitionInfo& FunctionDef, int32 Indent, bool bOutputConstructor, EExportingState ExportingState)
 {
-	UFunction* Function = FunctionDef.GetFunction();
 	if (!WillExportEventParms(FunctionDef))
 	{
 		return;
 	}
 
-	FString FunctionName = Function->GetName();
-	if (Function->HasAnyFunctionFlags(FUNC_Delegate))
+	FString FunctionName = FunctionDef.GetName();
+	if (FunctionDef.HasAnyFunctionFlags(FUNC_Delegate))
 	{
 		FunctionName.LeftChopInline(HEADER_GENERATED_DELEGATE_SIGNATURE_SUFFIX_LENGTH, false);
 	}
 
-	FString EventParmStructName = GetEventStructParamsName(Function->GetOuter(), *FunctionName);
+	FString EventParmStructName = GetEventStructParamsName(*FunctionDef.GetOuter(), *FunctionName);
 	Out.Logf(TEXT("%sstruct %s\r\n"), FCString::Tab(Indent), *EventParmStructName);
 	Out.Logf(TEXT("%s{\r\n"), FCString::Tab(Indent));
 
@@ -4599,16 +4588,15 @@ void FNativeClassHeaderGenerator::CheckRPCFunctions(FReferenceGatherers& OutRefe
 	bool bHasImplementation = ImplementationPosition != INDEX_NONE;
 	bool bHasValidate = ValidatePosition != INDEX_NONE;
 
-	UFunction* Function = FunctionDef.GetFunction();
 	const FFuncInfo& FunctionData = FunctionDef.GetFunctionData();
 	FString FunctionReturnType = GetFunctionReturnString(FunctionDef, OutReferenceGatherers);
-	const TCHAR* ConstModifier = (Function->HasAllFunctionFlags(FUNC_Const) ? TEXT("const ") : TEXT(" "));
+	const TCHAR* ConstModifier = (FunctionDef.HasAllFunctionFlags(FUNC_Const) ? TEXT("const ") : TEXT(" "));
 
-	const bool bIsNative = Function->HasAllFunctionFlags(FUNC_Native);
-	const bool bIsNet = Function->HasAllFunctionFlags(FUNC_Net);
-	const bool bIsNetValidate = Function->HasAllFunctionFlags(FUNC_NetValidate);
-	const bool bIsNetResponse = Function->HasAllFunctionFlags(FUNC_NetResponse);
-	const bool bIsBlueprintEvent = Function->HasAllFunctionFlags(FUNC_BlueprintEvent);
+	const bool bIsNative = FunctionDef.HasAllFunctionFlags(FUNC_Native);
+	const bool bIsNet = FunctionDef.HasAllFunctionFlags(FUNC_Net);
+	const bool bIsNetValidate = FunctionDef.HasAllFunctionFlags(FUNC_NetValidate);
+	const bool bIsNetResponse = FunctionDef.HasAllFunctionFlags(FUNC_NetResponse);
+	const bool bIsBlueprintEvent = FunctionDef.HasAllFunctionFlags(FUNC_BlueprintEvent);
 
 	bool bNeedsImplementation = (bIsNet && !bIsNetResponse) || bIsBlueprintEvent || bIsNative;
 	bool bNeedsValidate = (bIsNative || bIsNet) && !bIsNetResponse && bIsNetValidate;
@@ -4634,7 +4622,7 @@ void FNativeClassHeaderGenerator::CheckRPCFunctions(FReferenceGatherers& OutRefe
 	// Coin static_assert message
 	//
 	FUHTStringBuilder AssertMessage;
-	AssertMessage.Logf(TEXT("Function %s was marked as %s"), *(Function->GetName()), FunctionSpecifiers[0]);
+	AssertMessage.Logf(TEXT("Function %s was marked as %s"), *(FunctionDef.GetName()), FunctionSpecifiers[0]);
 	for (int32 i = 1; i < FunctionSpecifiers.Num(); ++i)
 	{
 		AssertMessage.Logf(TEXT(", %s"), FunctionSpecifiers[i]);
@@ -4693,11 +4681,9 @@ void FNativeClassHeaderGenerator::ExportNativeFunctionHeader(
 	const TCHAR*                     APIString
 )
 {
-	UFunction* Function = FunctionDef.GetFunction();
-
-	const bool bIsDelegate   = Function->HasAnyFunctionFlags( FUNC_Delegate );
-	const bool bIsInterface  = !bIsDelegate && Function->GetOwnerClass()->HasAnyClassFlags(CLASS_Interface);
-	const bool bIsK2Override = Function->HasAnyFunctionFlags( FUNC_BlueprintEvent );
+	const bool bIsDelegate   = FunctionDef.HasAnyFunctionFlags( FUNC_Delegate );
+	const bool bIsInterface  = !bIsDelegate && FunctionDef.GetOwnerClass()->HasAnyClassFlags(CLASS_Interface);
+	const bool bIsK2Override = FunctionDef.HasAnyFunctionFlags( FUNC_BlueprintEvent );
 
 	if (!bIsDelegate)
 	{
@@ -4711,7 +4697,7 @@ void FNativeClassHeaderGenerator::ExportNativeFunctionHeader(
 		// If the function was marked as 'RequiredAPI', then add the *_API macro prefix.  Note that if the class itself
 		// was marked 'RequiredAPI', this is not needed as C++ will exports all methods automatically.
 		if (FunctionType != EExportFunctionType::Event &&
-			!Function->GetOwnerClass()->HasAnyClassFlags(CLASS_RequiredAPI) &&
+			!FunctionDef.GetOwnerClass()->HasAnyClassFlags(CLASS_RequiredAPI) &&
 			(FunctionData.FunctionExportFlags & FUNCEXPORT_RequiredAPI))
 		{
 			Out.Log(APIString);
@@ -4731,7 +4717,7 @@ void FNativeClassHeaderGenerator::ExportNativeFunctionHeader(
 			Out.Log(TEXT("virtual "));
 		}
 		// this is not an event, the function is not a static function and the function is not marked final
-		else if ( FunctionType != EExportFunctionType::Event && !Function->HasAnyFunctionFlags(FUNC_Static) && !(FunctionData.FunctionExportFlags & FUNCEXPORT_Final) )
+		else if ( FunctionType != EExportFunctionType::Event && !FunctionDef.HasAnyFunctionFlags(FUNC_Static) && !(FunctionData.FunctionExportFlags & FUNCEXPORT_Final) )
 		{
 			Out.Log(TEXT("virtual "));
 		}
@@ -4766,12 +4752,12 @@ void FNativeClassHeaderGenerator::ExportNativeFunctionHeader(
 	FString FunctionName;
 	if (FunctionHeaderStyle == EExportFunctionHeaderStyle::Definition)
 	{
-		FunctionName = FString::Printf(TEXT("%s::"), *FNameLookupCPP::GetNameCPP(CastChecked<UClass>(Function->GetOuter()), bIsInterface || FunctionType == EExportFunctionType::Interface));
+		FunctionName = FString::Printf(TEXT("%s::"), *UHTCastChecked<FUnrealClassDefinitionInfo>(FunctionDef.GetOuter()).GetAlternateNameCPP(bIsInterface || FunctionType == EExportFunctionType::Interface));
 	}
 
 	if (FunctionType == EExportFunctionType::Interface)
 	{
-		FunctionName += FString::Printf(TEXT("Execute_%s"), *Function->GetName());
+		FunctionName += FString::Printf(TEXT("Execute_%s"), *FunctionDef.GetName());
 	}
 	else if (FunctionType == EExportFunctionType::Event)
 	{
@@ -4819,7 +4805,7 @@ void FNativeClassHeaderGenerator::ExportNativeFunctionHeader(
 	Out.Log( TEXT(")") );
 	if (FunctionType != EExportFunctionType::Interface)
 	{
-		if (!bIsDelegate && Function->HasAllFunctionFlags(FUNC_Const))
+		if (!bIsDelegate && FunctionDef.HasAllFunctionFlags(FUNC_Const))
 		{
 			Out.Log( TEXT(" const") );
 		}
@@ -4865,7 +4851,6 @@ void FNativeClassHeaderGenerator::ExportNativeFunctionHeader(
  */
 void FNativeClassHeaderGenerator::ExportFunctionThunk(FUHTStringBuilder& RPCWrappers, FReferenceGatherers& OutReferenceGatherers, FUnrealFunctionDefinitionInfo& FunctionDef, const TArray<FUnrealPropertyDefinitionInfo*>& ParameterDefs, FUnrealPropertyDefinitionInfo* ReturnDef) const
 {
-	UFunction* Function = FunctionDef.GetFunction();
 	const FFuncInfo& FunctionData = FunctionDef.GetFunctionData();
 
 	// export the GET macro for this parameter
@@ -5029,20 +5014,20 @@ void FNativeClassHeaderGenerator::ExportFunctionThunk(FUHTStringBuilder& RPCWrap
 		// Enable deprecation warnings only if GENERATED_BODY is used inside class or interface (not GENERATED_UCLASS_BODY etc.)
 		OwnerClassDef->HasGeneratedBody()
 		// and implementation function is called, but not the one declared by user
-		&& (FunctionData.CppImplName != Function->GetName() && !bHasImplementation);
+		&& (FunctionData.CppImplName != FunctionDef.GetName() && !bHasImplementation);
 
 	bool bShouldEnableValidateDeprecation =
 		// Enable deprecation warnings only if GENERATED_BODY is used inside class or interface (not GENERATED_UCLASS_BODY etc.)
 		OwnerClassDef->HasGeneratedBody()
 		// and validation function is called
-		&& (FunctionData.FunctionFlags & FUNC_NetValidate) && !bHasValidate;
+		&& FunctionDef.HasAnyFunctionFlags(FUNC_NetValidate) && !bHasValidate;
 
 	//Emit warning here if necessary
 	FUHTStringBuilder FunctionDeclaration;
 	ExportNativeFunctionHeader(FunctionDeclaration, OutReferenceGatherers.ForwardDeclarations, FunctionDef, FunctionData, EExportFunctionType::Function, EExportFunctionHeaderStyle::Declaration, nullptr, *GetAPIString());
 
 	// Call the validate function if there is one
-	if (!(FunctionData.FunctionExportFlags & FUNCEXPORT_CppStatic) && (FunctionData.FunctionFlags & FUNC_NetValidate))
+	if (!(FunctionData.FunctionExportFlags & FUNCEXPORT_CppStatic) && FunctionDef.HasAnyFunctionFlags(FUNC_NetValidate))
 	{
 		RPCWrappers.Logf(TEXT("\t\tif (!P_THIS->%s(%s))") LINE_TERMINATOR, *FunctionData.CppValidationImplName, *ParameterList);
 		RPCWrappers.Logf(TEXT("\t\t{") LINE_TERMINATOR);
@@ -5074,7 +5059,7 @@ void FNativeClassHeaderGenerator::ExportFunctionThunk(FUHTStringBuilder& RPCWrap
 	// export the call to the C++ version
 	if (FunctionData.FunctionExportFlags & FUNCEXPORT_CppStatic)
 	{
-		RPCWrappers.Logf(TEXT("%s::%s(%s);") LINE_TERMINATOR, *FNameLookupCPP::GetNameCPP(Function->GetOwnerClass()), *FunctionData.CppImplName, *ParameterList);
+		RPCWrappers.Logf(TEXT("%s::%s(%s);") LINE_TERMINATOR, *FunctionDef.GetOwnerClass()->GetAlternateNameCPP(), *FunctionData.CppImplName, *ParameterList);
 	}
 	else
 	{
@@ -5210,14 +5195,13 @@ void FNativeClassHeaderGenerator::ExportNativeFunctions(FOutputDevice& OutGenera
 	// export the C++ stubs
 	for (FUnrealFunctionDefinitionInfo* FunctionDef : Functions)
 	{
-		UFunction* Function = FunctionDef->GetFunction();
 		const FFuncInfo& FunctionData = FunctionDef->GetFunctionData();
-		if (!(Function->FunctionFlags & FUNC_Native))
+		if (!FunctionDef->HasAnyFunctionFlags(FUNC_Native))
 		{
 			continue;
 		}
 
-		const bool bEditorOnlyFunc = Function->HasAnyFunctionFlags(FUNC_EditorOnly);
+		const bool bEditorOnlyFunc = FunctionDef->HasAnyFunctionFlags(FUNC_EditorOnly);
 		FNativeFunctionStringBuilder& FuncStringBuilders = bEditorOnlyFunc ? EditorStringBuilders : RuntimeStringBuilders;
 
 		// Custom thunks don't get any C++ stub function generated
@@ -5227,13 +5211,13 @@ void FNativeClassHeaderGenerator::ExportNativeFunctions(FOutputDevice& OutGenera
 		}
 
 		// Should we emit these to RPC wrappers or just ignore them?
-		const bool bWillBeProgrammerTyped = FunctionData.CppImplName == Function->GetName();
+		const bool bWillBeProgrammerTyped = FunctionData.CppImplName == FunctionDef->GetName();
 
 		if (!bWillBeProgrammerTyped)
 		{
 			FString ClassDefinition(UE_PTRDIFF_TO_INT32(ClassEnd - ClassStart), ClassStart);
 
-			FString FunctionName = Function->GetName();
+			FString FunctionName = FunctionDef->GetName();
 			int32 ClassDefinitionStartPosition = UE_PTRDIFF_TO_INT32(ClassStart - *SourceFile.GetContent());
 
 			int32 ImplementationPosition = FindIdentifierExactMatch(ClassDefinition, FunctionData.CppImplName);
@@ -5256,11 +5240,11 @@ void FNativeClassHeaderGenerator::ExportNativeFunctions(FOutputDevice& OutGenera
 			FunctionDeclaration.Log(TEXT(";\r\n"));
 
 			// Declare validation function if needed
-			if (FunctionData.FunctionFlags & FUNC_NetValidate)
+			if (FunctionDef->HasAnyFunctionFlags(FUNC_NetValidate))
 			{
 				FString ParameterList = GetFunctionParameterString(*FunctionDef, OutReferenceGatherers);
 
-				const TCHAR* Virtual = (!FunctionData.FunctionReference->HasAnyFunctionFlags(FUNC_Static) && !(FunctionData.FunctionExportFlags & FUNCEXPORT_Final)) ? TEXT("virtual") : TEXT("");
+				const TCHAR* Virtual = (!FunctionDef->HasAnyFunctionFlags(FUNC_Static) && !(FunctionData.FunctionExportFlags & FUNCEXPORT_Final)) ? TEXT("virtual") : TEXT("");
 				FStringOutputDevice ValidDecl;
 				ValidDecl.Logf(TEXT("\t%s bool %s(%s);\r\n"), Virtual, *FunctionData.CppValidationImplName, *ParameterList);
 				FuncStringBuilders.AutogeneratedBlueprintFunctionDeclarations.Log(*ValidDecl);
@@ -5287,7 +5271,7 @@ void FNativeClassHeaderGenerator::ExportNativeFunctions(FOutputDevice& OutGenera
 
 		// if this function was originally declared in a base class, and it isn't a static function,
 		// only the C++ function header will be exported
-		if (!ShouldExportUFunction(Function))
+		if (!ShouldExportFunction(*FunctionDef))
 		{
 			continue;
 		}
@@ -5408,23 +5392,22 @@ void FNativeClassHeaderGenerator::ExportCallbackFunctions(
 	FMacroBlockEmitter OutCppEditorOnly(OutCpp, TEXT("WITH_EDITOR"));
 	for (FUnrealFunctionDefinitionInfo* FunctionDef : CallbackFunctions)
 	{
-		UFunction* Function = FunctionDef->GetFunction();
-
 		// Never expecting to export delegate functions this way
-		check(!Function->HasAnyFunctionFlags(FUNC_Delegate));
+		check(!FunctionDef->HasAnyFunctionFlags(FUNC_Delegate));
 
 		const FFuncInfo& FunctionData = FunctionDef->GetFunctionData();
-		FString          FunctionName = Function->GetName();
-		UClass*          Class = CastChecked<UClass>(Function->GetOuter());
+		FString          FunctionName = FunctionDef->GetName();
+		FUnrealClassDefinitionInfo&    ClassDef = UHTCastChecked<FUnrealClassDefinitionInfo>(FunctionDef->GetOuter());
+		UClass*          Class = ClassDef.GetClass();
 		const FString    ClassName = FNameLookupCPP::GetNameCPP(Class);
 
-		if (FunctionData.FunctionFlags & FUNC_NetResponse)
+		if (FunctionDef->HasAnyFunctionFlags(FUNC_NetResponse))
 		{
 			// Net response functions don't go into the VM
 			continue;
 		}
 
-		const bool bIsEditorOnly = Function->HasAnyFunctionFlags(FUNC_EditorOnly);
+		const bool bIsEditorOnly = FunctionDef->HasAnyFunctionFlags(FUNC_EditorOnly);
 
 		OutCppEditorOnly(bIsEditorOnly);
 
@@ -5444,7 +5427,7 @@ void FNativeClassHeaderGenerator::ExportCallbackFunctions(
 		if (ExportCallbackType != EExportCallbackType::Interface)
 		{
 			FunctionNameName = FString::Printf(TEXT("NAME_%s_%s"), *ClassName, *FunctionName);
-			OutCpp.Logf(TEXT("\tstatic FName %s = FName(TEXT(\"%s\"));") LINE_TERMINATOR, *FunctionNameName, *GetOverriddenFName(Function).ToString());
+			OutCpp.Logf(TEXT("\tstatic FName %s = FName(TEXT(\"%s\"));") LINE_TERMINATOR, *FunctionNameName, *GetOverriddenFName(FunctionDef).ToString());
 		}
 
 		// Emit the thunk implementation
@@ -5454,12 +5437,12 @@ void FNativeClassHeaderGenerator::ExportCallbackFunctions(
 
 		if (ExportCallbackType != EExportCallbackType::Interface)
 		{
-			WriteEventFunctionPrologue(OutCpp, /*Indent=*/ 1, Parameters, Class, *FunctionName);
+			WriteEventFunctionPrologue(OutCpp, /*Indent=*/ 1, Parameters, ClassDef, *FunctionName);
 			{
 				// Cast away const just in case, because ProcessEvent isn't const
 				OutCpp.Logf(
 					TEXT("\t\t%sProcessEvent(FindFunctionChecked(%s),%s);\r\n"),
-					(Function->HasAllFunctionFlags(FUNC_Const)) ? *FString::Printf(TEXT("const_cast<%s*>(this)->"), *ClassName) : TEXT(""),
+					(FunctionDef->HasAllFunctionFlags(FUNC_Const)) ? *FString::Printf(TEXT("const_cast<%s*>(this)->"), *ClassName) : TEXT(""),
 					*FunctionNameName,
 					Parameters.HasParms() ? TEXT("&Parms") : TEXT("NULL")
 				);
@@ -5477,7 +5460,7 @@ void FNativeClassHeaderGenerator::ExportCallbackFunctions(
 			// satisfy compiler if it's expecting a return value
 			if (Parameters.Return)
 			{
-				FString EventParmStructName = GetEventStructParamsName(Class, *FunctionName);
+				FString EventParmStructName = GetEventStructParamsName(ClassDef, *FunctionName);
 				OutCpp.Logf(TEXT("\t\t%s Parms;") LINE_TERMINATOR, *EventParmStructName);
 				OutCpp.Log(TEXT("\t\treturn Parms.ReturnValue;") LINE_TERMINATOR);
 			}
@@ -5538,10 +5521,10 @@ void FNativeClassHeaderGenerator::ApplyAlternatePropertyExportText(FUnrealProper
 	}
 }
 
-static bool HasDynamicOuter(UField* Field)
+static bool HasDynamicOuter(FUnrealFieldDefinitionInfo& FieldDef)
 {
-	UField* FieldOuter = Cast<UField>(Field->GetOuter());
-	return FieldOuter && FUnrealTypeDefinitionInfo::IsDynamic(FieldOuter);
+	FUnrealFieldDefinitionInfo* FieldOuter = UHTCast<FUnrealFieldDefinitionInfo>(FieldDef.GetOuter());
+	return FieldOuter && FieldOuter->IsDynamic();
 }
 
 static void RecordPackageSingletons(
@@ -5554,7 +5537,7 @@ static void RecordPackageSingletons(
 	for (FUnrealScriptStructDefinitionInfo* StructDef : StructDefs)
 	{
 		UScriptStruct* Struct = StructDef->GetScriptStruct();
-		if (Struct->StructFlags & STRUCT_NoExport && !HasDynamicOuter(Struct))
+		if (Struct->StructFlags & STRUCT_NoExport && !HasDynamicOuter(*StructDef))
 		{
 			Singletons.Add(StructDef);
 		}
@@ -5562,8 +5545,7 @@ static void RecordPackageSingletons(
 
 	for (FUnrealFunctionDefinitionInfo* DelegateDef : DelegateDefs)
 	{
-		UFunction* Delegate = DelegateDef->GetFunction();
-		if (!HasDynamicOuter(Delegate))
+		if (!HasDynamicOuter(*DelegateDef))
 		{
 			Singletons.Add(DelegateDef);
 		}
@@ -5736,7 +5718,7 @@ void FNativeClassHeaderGenerator::GenerateSourceFiles(
 					for (FUnrealEnumDefinitionInfo* EnumDef : EnumDefs)
 					{
 						// Is this ever not the case?
-						if (EnumDef->GetEnum()->GetOuter()->IsA(UPackage::StaticClass()))
+						if (EnumDef->GetOuter()->GetObject()->IsA(UPackage::StaticClass()))
 						{
 							GeneratedFunctionDeclarations.Log(EnumDef->GetExternDecl(true));
 							Generator.ExportGeneratedEnumInitCode(GeneratedCPPText, ReferenceGatherers, SourceFile, *EnumDef);
diff --git a/Engine/Source/Programs/UnrealHeaderTool/Private/HeaderParser.cpp b/Engine/Source/Programs/UnrealHeaderTool/Private/HeaderParser.cpp
--- a/Engine/Source/Programs/UnrealHeaderTool/Private/HeaderParser.cpp
+++ b/Engine/Source/Programs/UnrealHeaderTool/Private/HeaderParser.cpp
@@ -671,18 +671,17 @@ namespace
 				if (FUnrealFunctionDefinitionInfo* FuncDef = TypeDef.AsFunction())
 				{
 					UFunction* Function = FuncDef->GetFunction();
-					if (Function->HasAnyFunctionFlags(FUNC_Static))
+					if (FuncDef->HasAnyFunctionFlags(FUNC_Static))
 					{
 						// Determine if it's a function library
-						UClass* Class = Function->GetOuterUClass();
-						while (Class != nullptr && Class->GetSuperClass() != UObject::StaticClass())
+						FUnrealClassDefinitionInfo* ClassDef = FuncDef->GetOuterClass();
+						for (; ClassDef != nullptr && ClassDef->GetSuperClass() != GUObjectDef; ClassDef = ClassDef->GetSuperClass())
 						{
-							Class = Class->GetSuperClass();
 						}
 
-						if (Class != nullptr && Class->GetName() == TEXT("BlueprintFunctionLibrary"))
+						if (ClassDef != nullptr && ClassDef->GetName() == TEXT("BlueprintFunctionLibrary"))
 						{
-							FError::Throwf(TEXT("%s doesn't make sense on static method '%s' in a blueprint function library"), *InKey.ToString(), *Function->GetName());
+							FError::Throwf(TEXT("%s doesn't make sense on static method '%s' in a blueprint function library"), *InKey.ToString(), *FuncDef->GetName());
 						}
 					}
 				}
@@ -693,8 +692,7 @@ namespace
 			{
 				if (FUnrealFunctionDefinitionInfo* FuncDef = TypeDef.AsFunction())
 				{
-					UFunction* Function = FuncDef->GetFunction();
-					bool bGoodParams = (Function->NumParms == 3);
+					bool bGoodParams = (FuncDef->GetProperties().Num() == 3);
 					if (bGoodParams)
 					{
 						FProperty* FirstParam = nullptr;
@@ -739,7 +737,6 @@ namespace
 			{
 				if (FUnrealFunctionDefinitionInfo* FuncDef = TypeDef.AsFunction())
 				{
-					UFunction* Function = FuncDef->GetFunction();
 					// multiple entry parsing in the same format as eg SetParam.
 					TArray<FString> RawGroupings;
 					InValue.ParseIntoArray(RawGroupings, TEXT(","), false);
@@ -810,7 +807,14 @@ namespace
 
 				if (!FUnitConversion::UnitFromString(*InValue))
 				{
-					FError::Throwf(TEXT("Unrecognized units (%s) specified for property '%s'"), *InValue, *TypeDef.AsFieldChecked().GetField()->GetFullName());
+					if (TypeDef.AsField() != nullptr)
+					{
+						FError::Throwf(TEXT("Unrecognized units (%s) specified for property '%s'"), *InValue, *TypeDef.AsFieldChecked().GetField()->GetFullName());
+					}
+					else
+					{
+						FError::Throwf(TEXT("Unrecognized units (%s) specified for property '%s'"), *InValue, *TypeDef.AsPropertyChecked().GetProperty()->GetFullName());
+					}
 				}
 			}
 			break;
@@ -1279,7 +1283,6 @@ FUnrealEnumDefinitionInfo& FHeaderParser::CompileEnum()
 
 	// Create enum definition.
 	FUnrealEnumDefinitionInfo& EnumDef = GTypeDefinitionInfoMap.FindByNameChecked<FUnrealEnumDefinitionInfo>(EnumToken.Identifier);
-	UEnum* Enum = EnumDef.GetEnum();
 
 	Scope->AddType(EnumDef);
 
@@ -1391,7 +1394,7 @@ FUnrealEnumDefinitionInfo& FHeaderParser::CompileEnum()
 				FError::Throwf(TEXT("Missing enumeration name") );
 			}
 
-			Enum->CppType = FString::Printf(TEXT("%s::%s"), EnumToken.Identifier, InnerEnumToken.Identifier);
+			EnumDef.SetCppType(FString::Printf(TEXT("%s::%s"), EnumToken.Identifier, InnerEnumToken.Identifier));
 
 			RequireSymbol( TEXT('{'), TEXT("'Enum'") );
 		}
@@ -1400,7 +1403,7 @@ FUnrealEnumDefinitionInfo& FHeaderParser::CompileEnum()
 		case UEnum::ECppForm::Regular:
 		case UEnum::ECppForm::EnumClass:
 		{
-			Enum->CppType = EnumToken.Identifier;
+			EnumDef.SetCppType(EnumToken.Identifier);
 		}
 		break;
 	}
@@ -1543,10 +1546,10 @@ FUnrealEnumDefinitionInfo& FHeaderParser::CompileEnum()
 	}
 
 	// Register the list of enum names.
-	if (!Enum->SetEnums(EnumNames, CppForm, Flags, false))
+	if (!EnumDef.SetEnums(EnumNames, CppForm, Flags, false))
 	{
-		const FName MaxEnumItem      = *(Enum->GenerateEnumPrefix() + TEXT("_MAX"));
-		const int32 MaxEnumItemIndex = Enum->GetIndexByName(MaxEnumItem);
+		const FName MaxEnumItem      = *(EnumDef.GenerateEnumPrefix() + TEXT("_MAX"));
+		const int32 MaxEnumItemIndex = EnumDef.GetIndexByName(MaxEnumItem);
 		if (MaxEnumItemIndex != INDEX_NONE)
 		{
 			FError::Throwf(TEXT("Illegal enumeration tag specified.  Conflicts with auto-generated tag '%s'"), *MaxEnumItem.ToString());
@@ -1560,15 +1563,15 @@ FUnrealEnumDefinitionInfo& FHeaderParser::CompileEnum()
 	// Add the metadata gathered for the enum to the package
 	if (EnumValueMetaData.Num() > 0)
 	{
-		UMetaData* PackageMetaData = Enum->GetOutermost()->GetMetaData();
+		UMetaData* PackageMetaData = EnumDef.GetPackageDef().GetPackage()->GetMetaData();
 		checkSlow(PackageMetaData);
 
-		PackageMetaData->SetObjectValues(Enum, MoveTemp(EnumValueMetaData));
+		PackageMetaData->SetObjectValues(EnumDef.GetEnum(), MoveTemp(EnumValueMetaData));
 	}
 
-	if (!Enum->IsValidEnumValue(0) && EnumMetaData.Contains(FHeaderParserNames::NAME_BlueprintType))
+	if (!EnumDef.IsValidEnumValue(0) && EnumMetaData.Contains(FHeaderParserNames::NAME_BlueprintType))
 	{
-		UE_LOG_WARNING_UHT(TEXT("'%s' does not have a 0 entry! (This is a problem when the enum is initalized by default)"), *Enum->GetName());
+		UE_LOG_WARNING_UHT(TEXT("'%s' does not have a 0 entry! (This is a problem when the enum is initalized by default)"), *EnumDef.GetName());
 	}
 
 	return EnumDef;
@@ -2647,7 +2650,7 @@ void FHeaderParser::FixupDelegateProperties(FUnrealStructDefinitionInfo& StructD
 					UngetToken(PropertyDef->GetLineNumber(), PropertyDef->GetParsePosition());
 					FError::Throwf(TEXT("Failed to find delegate function '%s'"), *NameOfDelegateFunction);
 				}
-				else if ( (SourceDelegateFunctionDef->GetFunction()->FunctionFlags&FUNC_Delegate) == 0 )
+				else if (!SourceDelegateFunctionDef->HasAnyFunctionFlags(FUNC_Delegate))
 				{
 					UngetToken(PropertyDef->GetLineNumber(), PropertyDef->GetParsePosition());
 					FError::Throwf(TEXT("Only delegate functions can be used as the type for a delegate property; '%s' is not a delegate."), *NameOfDelegateFunction);
@@ -2658,26 +2661,25 @@ void FHeaderParser::FixupDelegateProperties(FUnrealStructDefinitionInfo& StructD
 
 				// save this into the delegate cache for faster lookup later
 			DelegateCache.Add(DelegatePropertyToken.DelegateName, SourceDelegateFunctionDef);
-			UFunction* SourceDelegateFunction = SourceDelegateFunctionDef->GetFunction();
 
 			// bind it to the delegate property
 			if( DelegateProperty != NULL )
 			{
-				if( !SourceDelegateFunction->HasAnyFunctionFlags( FUNC_MulticastDelegate ) )
+				if( !SourceDelegateFunctionDef->HasAnyFunctionFlags( FUNC_MulticastDelegate ) )
 				{
-					DelegateProperty->SignatureFunction = SourceDelegateFunction;
+					DelegateProperty->SignatureFunction = SourceDelegateFunctionDef->GetFunction();
 					DelegatePropertyToken.FunctionDef = SourceDelegateFunctionDef;
 				}
 				else
 				{
-					FError::Throwf(TEXT("Unable to declare a single-cast delegate property for a multi-cast delegate type '%s'.  Either add a 'multicast' qualifier to the property or change the delegate type to be single-cast as well."), *SourceDelegateFunction->GetName());
+					FError::Throwf(TEXT("Unable to declare a single-cast delegate property for a multi-cast delegate type '%s'.  Either add a 'multicast' qualifier to the property or change the delegate type to be single-cast as well."), *SourceDelegateFunctionDef->GetName());
 				}
 			}
 			else if( MulticastDelegateProperty != NULL )
 			{
-				if( SourceDelegateFunction->HasAnyFunctionFlags( FUNC_MulticastDelegate ) )
+				if(SourceDelegateFunctionDef->HasAnyFunctionFlags( FUNC_MulticastDelegate ) )
 				{
-					MulticastDelegateProperty->SignatureFunction = SourceDelegateFunction;
+					MulticastDelegateProperty->SignatureFunction = SourceDelegateFunctionDef->GetFunction();
 					DelegatePropertyToken.FunctionDef = SourceDelegateFunctionDef;
 
 					if(MulticastDelegateProperty->HasAnyPropertyFlags(CPF_BlueprintAssignable | CPF_BlueprintCallable))
@@ -2690,7 +2692,7 @@ void FHeaderParser::FixupDelegateProperties(FUnrealStructDefinitionInfo& StructD
 							{
 								FString ExtendedCPPType;
 								FString CPPType = FuncParam->GetCPPType(&ExtendedCPPType);
-								UE_LOG_ERROR_UHT(TEXT("Type '%s%s' is not supported by blueprint. %s.%s"), *CPPType, *ExtendedCPPType, *SourceDelegateFunction->GetName(), *FuncParam->GetName());
+								UE_LOG_ERROR_UHT(TEXT("Type '%s%s' is not supported by blueprint. %s.%s"), *CPPType, *ExtendedCPPType, *SourceDelegateFunctionDef->GetName(), *FuncParam->GetName());
 							}
 
 							if(FuncParam->HasAllPropertyFlags(CPF_OutParm) && !FuncParam->HasAllPropertyFlags(CPF_ConstParm)  )
@@ -2699,7 +2701,7 @@ void FHeaderParser::FixupDelegateProperties(FUnrealStructDefinitionInfo& StructD
 								const bool bAllowedArrayRefFromBP = bClassGeneratedFromBP && FuncParam->IsA<FArrayProperty>();
 								if (!bAllowedArrayRefFromBP)
 								{
-									UE_LOG_ERROR_UHT(TEXT("BlueprintAssignable delegates do not support non-const references at the moment. Function: %s Parameter: '%s'"), *SourceDelegateFunction->GetName(), *FuncParam->GetName());
+									UE_LOG_ERROR_UHT(TEXT("BlueprintAssignable delegates do not support non-const references at the moment. Function: %s Parameter: '%s'"), *SourceDelegateFunctionDef->GetName(), *FuncParam->GetName());
 								}
 							}
 						}
@@ -2708,7 +2710,7 @@ void FHeaderParser::FixupDelegateProperties(FUnrealStructDefinitionInfo& StructD
 				}
 				else
 				{
-					FError::Throwf(TEXT("Unable to declare a multi-cast delegate property for a single-cast delegate type '%s'.  Either remove the 'multicast' qualifier from the property or change the delegate type to be 'multicast' as well."), *SourceDelegateFunction->GetName());
+					FError::Throwf(TEXT("Unable to declare a multi-cast delegate property for a single-cast delegate type '%s'.  Either remove the 'multicast' qualifier from the property or change the delegate type to be 'multicast' as well."), *SourceDelegateFunctionDef->GetName());
 				}
 			}
 		}
@@ -2826,27 +2828,26 @@ void FHeaderParser::VerifyBlueprintPropertyGetter(FUnrealPropertyDefinitionInfo&
 {
 	check(TargetFuncDef);
 	FProperty* Prop = PropertyDef.GetProperty();
-	UFunction* TargetFunc = TargetFuncDef->GetFunction();
 
 	const TArray<FUnrealPropertyDefinitionInfo*>& Properties = TargetFuncDef->GetProperties();
 	FUnrealPropertyDefinitionInfo* ReturnPropDef = TargetFuncDef->GetReturn();
 	if (Properties.Num() > 1 || (Properties.Num() == 1 && ReturnPropDef == nullptr))
 	{
-		UE_LOG_ERROR_UHT(TEXT("Blueprint Property getter function %s must not have parameters."), *TargetFunc->GetName());
+		UE_LOG_ERROR_UHT(TEXT("Blueprint Property getter function %s must not have parameters."), *TargetFuncDef->GetName());
 	}
 
 	if (ReturnPropDef == nullptr || !Prop->SameType(ReturnPropDef->GetProperty()))
 	{
 		FString ExtendedCPPType;
 		FString CPPType = Prop->GetCPPType(&ExtendedCPPType);
-		UE_LOG_ERROR_UHT(TEXT("Blueprint Property getter function %s must have return value of type %s%s."), *TargetFunc->GetName(), *CPPType, *ExtendedCPPType);
+		UE_LOG_ERROR_UHT(TEXT("Blueprint Property getter function %s must have return value of type %s%s."), *TargetFuncDef->GetName(), *CPPType, *ExtendedCPPType);
 	}
 
-	if (TargetFunc->HasAnyFunctionFlags(FUNC_Event))
+	if (TargetFuncDef->HasAnyFunctionFlags(FUNC_Event))
 	{
 		UE_LOG_ERROR_UHT(TEXT("Blueprint Property setter function cannot be a blueprint event."));
 	}
-	else if (!TargetFunc->HasAnyFunctionFlags(FUNC_BlueprintPure))
+	else if (!TargetFuncDef->HasAnyFunctionFlags(FUNC_BlueprintPure))
 	{
 		UE_LOG_ERROR_UHT(TEXT("Blueprint Property getter function must be pure."));
 	}
@@ -2856,33 +2857,32 @@ void FHeaderParser::VerifyBlueprintPropertySetter(FUnrealPropertyDefinitionInfo&
 {
 	check(TargetFuncDef);
 	FProperty* Prop = PropertyDef.GetProperty();
-	UFunction* TargetFunc = TargetFuncDef->GetFunction();
 	FUnrealPropertyDefinitionInfo* ReturnPropDef = TargetFuncDef->GetReturn();
 
 	if (ReturnPropDef)
 	{
-		UE_LOG_ERROR_UHT(TEXT("Blueprint Property setter function %s must not have a return value."), *TargetFunc->GetName());
+		UE_LOG_ERROR_UHT(TEXT("Blueprint Property setter function %s must not have a return value."), *TargetFuncDef->GetName());
 	}
 	else
 	{
 		const TArray<FUnrealPropertyDefinitionInfo*>& Properties = TargetFuncDef->GetProperties();
-		if (TargetFunc->NumParms != 1 || !Prop->SameType(Properties[0]->GetProperty()))
+		if (TargetFuncDef->GetProperties().Num() != 1 || !Prop->SameType(Properties[0]->GetProperty()))
 		{
 			FString ExtendedCPPType;
 			FString CPPType = Prop->GetCPPType(&ExtendedCPPType);
-			UE_LOG_ERROR_UHT(TEXT("Blueprint Property setter function %s must have exactly one parameter of type %s%s."), *TargetFunc->GetName(), *CPPType, *ExtendedCPPType);
+			UE_LOG_ERROR_UHT(TEXT("Blueprint Property setter function %s must have exactly one parameter of type %s%s."), *TargetFuncDef->GetName(), *CPPType, *ExtendedCPPType);
 		}
 	}
 
-	if (TargetFunc->HasAnyFunctionFlags(FUNC_Event))
+	if (TargetFuncDef->HasAnyFunctionFlags(FUNC_Event))
 	{
 		UE_LOG_ERROR_UHT(TEXT("Blueprint Property setter function cannot be a blueprint event."));
 	}
-	else if (!TargetFunc->HasAnyFunctionFlags(FUNC_BlueprintCallable))
+	else if (!TargetFuncDef->HasAnyFunctionFlags(FUNC_BlueprintCallable))
 	{
 		UE_LOG_ERROR_UHT(TEXT("Blueprint Property setter function must be blueprint callable."));
 	}
-	else if (TargetFunc->HasAnyFunctionFlags(FUNC_BlueprintPure))
+	else if (TargetFuncDef->HasAnyFunctionFlags(FUNC_BlueprintPure))
 	{
 		UE_LOG_ERROR_UHT(TEXT("Blueprint Property setter function must not be pure."));
 	}
@@ -2893,10 +2893,9 @@ void FHeaderParser::VerifyRepNotifyCallback(FUnrealPropertyDefinitionInfo& Prope
 	FProperty* Prop = PropertyDef.GetProperty();
 	if (TargetFuncDef)
 	{
-		UFunction* TargetFunc = TargetFuncDef->GetFunction();
-		if (TargetFunc->GetReturnProperty())
+		if (TargetFuncDef->GetReturn())
 		{
-			UE_LOG_ERROR_UHT(TEXT("Replication notification function %s must not have return value."), *TargetFunc->GetName());
+			UE_LOG_ERROR_UHT(TEXT("Replication notification function %s must not have return value."), *TargetFuncDef->GetName());
 		}
 
 		const bool bIsArrayProperty = ( Prop->ArrayDim > 1 || CastField<FArrayProperty>(Prop) );
@@ -2905,7 +2904,7 @@ void FHeaderParser::VerifyRepNotifyCallback(FUnrealPropertyDefinitionInfo& Prope
 		const TArray<FUnrealPropertyDefinitionInfo*>& Properties = TargetFuncDef->GetProperties();
 		if (Properties.Num() > MaxParms)
 		{
-			UE_LOG_ERROR_UHT(TEXT("Replication notification function %s has too many parameters."), *TargetFunc->GetName());
+			UE_LOG_ERROR_UHT(TEXT("Replication notification function %s has too many parameters."), *TargetFuncDef->GetName());
 		}
 
 
@@ -2917,18 +2916,18 @@ void FHeaderParser::VerifyRepNotifyCallback(FUnrealPropertyDefinitionInfo& Prope
 			{
 				FString ExtendedCPPType;
 				FString CPPType = Prop->GetCPPType(&ExtendedCPPType);
-				UE_LOG_ERROR_UHT(TEXT("Replication notification function %s has invalid parameter for property %s. First (optional) parameter must be of type %s%s."), *TargetFunc->GetName(), *Prop->GetName(), *CPPType, *ExtendedCPPType);
+				UE_LOG_ERROR_UHT(TEXT("Replication notification function %s has invalid parameter for property %s. First (optional) parameter must be of type %s%s."), *TargetFuncDef->GetName(), *Prop->GetName(), *CPPType, *ExtendedCPPType);
 			}
 		}
 
-		if ( TargetFunc->NumParms >= 2)
+		if (TargetFuncDef->GetProperties().Num() >= 2)
 		{
 			FProperty* Parm = Properties[1]->GetProperty();
 			// A 2nd parameter for arrays can be specified as a const TArray<uint8>&. This is a list of element indices that have changed
 			FArrayProperty *ArrayProp = CastField<FArrayProperty>(Parm);
 			if (!(ArrayProp && CastField<FByteProperty>(ArrayProp->Inner)) || !(Parm->GetPropertyFlags() & CPF_ConstParm) || !(Parm->GetPropertyFlags() & CPF_ReferenceParm))
 			{
-				UE_LOG_ERROR_UHT(TEXT("Replication notification function %s (optional) second parameter must be of type 'const TArray<uint8>&'"), *TargetFunc->GetName());
+				UE_LOG_ERROR_UHT(TEXT("Replication notification function %s (optional) second parameter must be of type 'const TArray<uint8>&'"), *TargetFuncDef->GetName());
 			}
 		}
 	}
@@ -2955,8 +2954,7 @@ void FHeaderParser::VerifyPropertyMarkups(FUnrealClassDefinitionInfo& TargetClas
 				// Since the function map is not valid yet, we have to iterate over the fields to look for the function
 				for (FUnrealFunctionDefinitionInfo* FunctionDef : SearchClassDef->GetFunctions())
 				{
-					UFunction* TestFunc = FunctionDef->GetFunction();
-					if (TestFunc && FNativeClassHeaderGenerator::GetOverriddenFName(TestFunc) == FuncName)
+					if (FNativeClassHeaderGenerator::GetOverriddenFName(FunctionDef) == FuncName)
 					{
 						return FunctionDef;
 					}
@@ -4218,9 +4216,9 @@ void FHeaderParser::GetVarType(
 
 		if (VariableCategory == EVariableCategory::Member)
 		{
-			if (EnumDef->GetEnum()->GetCppForm() != UEnum::ECppForm::EnumClass)
+			if (EnumDef->GetCppForm() != UEnum::ECppForm::EnumClass)
 			{
-				FError::Throwf(TEXT("You cannot use the raw enum name as a type for member variables, instead use TEnumAsByte or a C++11 enum class with an explicit underlying type."), *EnumDef->GetEnum()->CppType);
+				FError::Throwf(TEXT("You cannot use the raw enum name as a type for member variables, instead use TEnumAsByte or a C++11 enum class with an explicit underlying type."), *EnumDef->GetCppType());
 			}
 		}
 
@@ -4252,13 +4250,13 @@ void FHeaderParser::GetVarType(
 			bStripped = true;
 		}
 
-		auto SetDelegateType = [&](FUnrealFunctionDefinitionInfo& InFunction, const FString& InIdentifierStripped)
+		auto SetDelegateType = [&](FUnrealFunctionDefinitionInfo& InFunctionDef, const FString& InIdentifierStripped)
 		{
 			bHandledType = true;
 
-			VarProperty = FPropertyBase(InFunction.GetFunction()->HasAnyFunctionFlags(FUNC_MulticastDelegate) ? CPT_MulticastDelegate : CPT_Delegate);
+			VarProperty = FPropertyBase(InFunctionDef.HasAnyFunctionFlags(FUNC_MulticastDelegate) ? CPT_MulticastDelegate : CPT_Delegate);
 			VarProperty.DelegateName = *InIdentifierStripped;
-			VarProperty.FunctionDef = &InFunction;
+			VarProperty.FunctionDef = &InFunctionDef;
 
 			if (!(Disallow & CPF_InstancedReference))
 			{
@@ -6252,8 +6250,6 @@ bool FHeaderParser::IsValidDelegateDeclaration(const FToken& Token) const
 // Parse the parameter list of a function or delegate declaration
 void FHeaderParser::ParseParameterList(FUnrealFunctionDefinitionInfo& FunctionDef, bool bExpectCommaBeforeName, TMap<FName, FString>* MetaData)
 {
-	UFunction* Function = FunctionDef.GetFunction();
-
 	// Get parameter list.
 	if (MatchSymbol(TEXT(')')))
 	{
@@ -6265,7 +6261,7 @@ void FHeaderParser::ParseParameterList(FUnrealFunctionDefinitionInfo& FunctionDe
 	{
 		// Get parameter type.
 		FPropertyBase Property(CPT_None);
-		EVariableCategory::Type VariableCategory = (Function->FunctionFlags & FUNC_Net) ? EVariableCategory::ReplicatedParameter : EVariableCategory::RegularParameter;
+		EVariableCategory::Type VariableCategory = FunctionDef.HasAnyFunctionFlags(FUNC_Net) ? EVariableCategory::ReplicatedParameter : EVariableCategory::RegularParameter;
 		GetVarType(GetCurrentScope(), Property, ~(CPF_ParmFlags | CPF_AutoWeak | CPF_RepSkip | CPF_UObjectWrapper | CPF_NativeAccessSpecifiers), nullptr, nullptr, EPropertyDeclarationStyle::None, VariableCategory);
 		Property.PropertyFlags |= CPF_Parm;
 
@@ -6277,7 +6273,7 @@ void FHeaderParser::ParseParameterList(FUnrealFunctionDefinitionInfo& FunctionDe
 		FUnrealPropertyDefinitionInfo& PropDef = GetVarNameAndDim(FunctionDef, Property, VariableCategory);
 		FProperty* Prop = PropDef.GetProperty();
 
-		Function->NumParms++;
+		FunctionDef.GetFunction()->NumParms++;
 
 		if( AdvancedDisplay.CanMarkMore() && AdvancedDisplay.ShouldMarkParameter(Prop->GetName()) )
 		{
@@ -6285,18 +6281,18 @@ void FHeaderParser::ParseParameterList(FUnrealFunctionDefinitionInfo& FunctionDe
 		}
 
 		// Check parameters.
-		if ((Function->FunctionFlags & FUNC_Net))
+		if (FunctionDef.HasAnyFunctionFlags(FUNC_Net))
 		{
 			if (Property.MapKeyProp.IsValid())
 			{
-				if (!(Function->FunctionFlags & FUNC_NetRequest || Function->FunctionFlags & FUNC_NetResponse))
+				if (!FunctionDef.HasAnyFunctionFlags(FUNC_NetRequest | FUNC_NetResponse))
 				{
 					UE_LOG_ERROR_UHT(TEXT("Maps are not supported in an RPC."));
 				}
 			}
 			else if (Property.ArrayType == EArrayType::Set)
 			{
-				if (!(Function->FunctionFlags & FUNC_NetRequest || Function->FunctionFlags & FUNC_NetResponse))
+				if (!FunctionDef.HasAnyFunctionFlags(FUNC_NetRequest | FUNC_NetResponse))
 				{
 					UE_LOG_ERROR_UHT(TEXT("Sets are not supported in an RPC."));
 				}
@@ -6304,13 +6300,13 @@ void FHeaderParser::ParseParameterList(FUnrealFunctionDefinitionInfo& FunctionDe
 
 			if (Property.Type == CPT_Struct)
 			{
-				if (!(Function->FunctionFlags & FUNC_NetRequest || Function->FunctionFlags & FUNC_NetResponse))
+				if (!FunctionDef.HasAnyFunctionFlags(FUNC_NetRequest | FUNC_NetResponse))
 				{
 					ValidateScriptStructOkForNet(Property.ScriptStructDef->GetStruct()->GetName(), *Property.ScriptStructDef);
 				}
 			}
 
-			if (!(Function->FunctionFlags & FUNC_NetRequest))
+			if (!FunctionDef.HasAnyFunctionFlags(FUNC_NetRequest))
 			{
 				if (Property.PropertyFlags & CPF_OutParm)
 				{
@@ -6350,7 +6346,7 @@ void FHeaderParser::ParseParameterList(FUnrealFunctionDefinitionInfo& FunctionDe
 				}
 			}
 		}
-		if ((Function->FunctionFlags & (FUNC_BlueprintEvent|FUNC_BlueprintCallable)) != 0)
+		if (FunctionDef.HasAnyFunctionFlags(FUNC_BlueprintEvent|FUNC_BlueprintCallable))
 		{
 			if (Property.Type == CPT_Byte)
 			{
@@ -6414,7 +6410,7 @@ void FHeaderParser::ParseParameterList(FUnrealFunctionDefinitionInfo& FunctionDe
 			}
 
 			// allow exec functions to be added to the metaData, this is so we can have default params for them.
-			const bool bStoreCppDefaultValueInMetaData = Function->HasAnyFunctionFlags(FUNC_BlueprintCallable | FUNC_Exec);
+			const bool bStoreCppDefaultValueInMetaData = FunctionDef.HasAnyFunctionFlags(FUNC_BlueprintCallable | FUNC_Exec);
 				
 			if((EndPos > -1) && bStoreCppDefaultValueInMetaData) 
 			{
@@ -6559,9 +6555,9 @@ FUnrealFunctionDefinitionInfo& FHeaderParser::CompileDelegateDeclaration(const T
 		FCString::Strcpy( FuncInfo.Function.Identifier, *Name );
 	}
 
-	FUnrealFunctionDefinitionInfo& FuncDef = bIsSparse ? CreateDelegateFunction<USparseDelegateFunction>(MoveTemp(FuncInfo)) : CreateDelegateFunction<UDelegateFunction>(MoveTemp(FuncInfo));
-	UDelegateFunction* DelegateSignatureFunction = CastChecked<UDelegateFunction>(FuncDef.GetFunction());
-	FFuncInfo& FuncDefFuncInfo = FuncDef.GetFunctionData();
+	FUnrealFunctionDefinitionInfo& DelegateSignatureFunctionDef = bIsSparse ? CreateDelegateFunction<USparseDelegateFunction>(MoveTemp(FuncInfo)) : CreateDelegateFunction<UDelegateFunction>(MoveTemp(FuncInfo));
+	UDelegateFunction* DelegateSignatureFunction = CastChecked<UDelegateFunction>(DelegateSignatureFunctionDef.GetFunction());
+	FFuncInfo& FuncDefFuncInfo = DelegateSignatureFunctionDef.GetFunctionData();
 
 	// determine whether this function should be 'const'
 	if (bDeclaredConst)
@@ -6592,14 +6588,14 @@ FUnrealFunctionDefinitionInfo& FHeaderParser::CompileDelegateDeclaration(const T
 		SDF->DelegateName = DelegateName.Identifier;
 	}
 
-	FuncDef.SetLineNumber(InputLine);
+	DelegateSignatureFunctionDef.SetLineNumber(InputLine);
 
 	// Get parameter list.
 	if (FoundParamCount)
 	{
 		RequireSymbol(TEXT(','), CurrentScopeName);
 
-		ParseParameterList(FuncDef, /*bExpectCommaBeforeName=*/ true);
+		ParseParameterList(DelegateSignatureFunctionDef, /*bExpectCommaBeforeName=*/ true);
 
 		// Check the expected versus actual number of parameters
 		int32 ParamCount = UE_PTRDIFF_TO_INT32(FoundParamCount - UHTConfig.DelegateParameterCountStrings.GetData()) + 1;
@@ -6615,13 +6611,13 @@ FUnrealFunctionDefinitionInfo& FHeaderParser::CompileDelegateDeclaration(const T
 	}
 
 	// The macro line must be set here
-	FuncDef.GetFunctionData().MacroLine = InputLine;
+	DelegateSignatureFunctionDef.GetFunctionData().MacroLine = InputLine;
 
 	// Create the return value property
 	if (bHasReturnValue)
 	{
 		ReturnType.PropertyFlags |= CPF_Parm | CPF_OutParm | CPF_ReturnParm;
-		FUnrealPropertyDefinitionInfo& PropDef = GetVarNameAndDim(FuncDef, ReturnType, EVariableCategory::Return);
+		FUnrealPropertyDefinitionInfo& PropDef = GetVarNameAndDim(DelegateSignatureFunctionDef, ReturnType, EVariableCategory::Return);
 
 		DelegateSignatureFunction->NumParms++;
 	}
@@ -6631,7 +6627,7 @@ FUnrealFunctionDefinitionInfo& FHeaderParser::CompileDelegateDeclaration(const T
 
 	AddFormattedPrevCommentAsTooltipMetaData(MetaData);
 
-	AddMetaDataToClassData(FuncDef, MoveTemp(MetaData));
+	AddMetaDataToClassData(DelegateSignatureFunctionDef, MoveTemp(MetaData));
 
 	// Optionally consume a semicolon, it's not required for the delegate macro since it contains one internally
 	MatchSemi();
@@ -6640,21 +6636,20 @@ FUnrealFunctionDefinitionInfo& FHeaderParser::CompileDelegateDeclaration(const T
 	DelegateSignatureFunction->Bind();
 
 	// End the nesting
-	PostPopFunctionDeclaration(FuncDef);
+	PostPopFunctionDeclaration(DelegateSignatureFunctionDef);
 
 	// Don't allow delegate signatures to be redefined.
 	auto FunctionIterator = GetCurrentScope()->GetTypeIterator<FUnrealFunctionDefinitionInfo>();
 	while (FunctionIterator.MoveNext())
 	{
 		FUnrealFunctionDefinitionInfo* TestFuncDef = *FunctionIterator;
-		UFunction* TestFunc = TestFuncDef->GetFunction();
-		if ((TestFunc->GetFName() == DelegateSignatureFunction->GetFName()) && (TestFunc != DelegateSignatureFunction))
+		if ((TestFuncDef->GetFName() == DelegateSignatureFunctionDef.GetFName()) && (TestFuncDef->GetFunction() != DelegateSignatureFunction))
 		{
 			FError::Throwf(TEXT("Can't override delegate signature function '%s'"), FuncDefFuncInfo.Function.Identifier);
 		}
 	}
 
-	return FuncDef;
+	return DelegateSignatureFunctionDef;
 }
 
 /**
@@ -7671,7 +7666,7 @@ void FHeaderParser::ComputeFunctionParametersSize(FUnrealClassDefinitionInfo& Cl
 		UFunction* ThisFunction = FunctionDef->GetFunction();
 
 		// Fix up any structs that were used as a parameter in a delegate before being defined
-		if (ThisFunction->HasAnyFunctionFlags(FUNC_Delegate))
+		if (FunctionDef->HasAnyFunctionFlags(FUNC_Delegate))
 		{
 			for (FUnrealPropertyDefinitionInfo* PropertyDef : FunctionDef->GetProperties())
 			{
@@ -7697,6 +7692,7 @@ void FHeaderParser::ComputeFunctionParametersSize(FUnrealClassDefinitionInfo& Cl
 			if (!(Param->PropertyFlags & CPF_ReturnParm) && (Param->PropertyFlags & CPF_OutParm))
 			{
 				ThisFunction->FunctionFlags |= FUNC_HasOutParms;
+				FunctionDef->GetFunctionData().FunctionFlags |= FUNC_HasOutParms;
 			}
 				
 			if (FStructProperty* StructProp = CastField<FStructProperty>(Param))
@@ -7704,6 +7700,7 @@ void FHeaderParser::ComputeFunctionParametersSize(FUnrealClassDefinitionInfo& Cl
 				if (StructProp->Struct->HasDefaults())
 				{
 					ThisFunction->FunctionFlags |= FUNC_HasDefaults;
+					FunctionDef->GetFunctionData().FunctionFlags |= FUNC_HasDefaults;
 				}
 			}
 		}
@@ -9041,29 +9038,29 @@ void FHeaderParser::ResetClassData()
 		}
 
 		// Copy special categories from parent
-		if (SuperClass->HasMetaData(FHeaderParserNames::NAME_HideCategories))
+		if (SuperClassDef->HasMetaData(FHeaderParserNames::NAME_HideCategories))
 		{
-			CurrentClass->SetMetaData(FHeaderParserNames::NAME_HideCategories, *SuperClass->GetMetaData(FHeaderParserNames::NAME_HideCategories));
+			CurrentClassDef.SetMetaData(FHeaderParserNames::NAME_HideCategories, *SuperClassDef->GetMetaData(FHeaderParserNames::NAME_HideCategories));
 		}
-		if (SuperClass->HasMetaData(FHeaderParserNames::NAME_ShowCategories))
+		if (SuperClassDef->HasMetaData(FHeaderParserNames::NAME_ShowCategories))
 		{
-			CurrentClass->SetMetaData(FHeaderParserNames::NAME_ShowCategories, *SuperClass->GetMetaData(FHeaderParserNames::NAME_ShowCategories));
+			CurrentClassDef.SetMetaData(FHeaderParserNames::NAME_ShowCategories, *SuperClassDef->GetMetaData(FHeaderParserNames::NAME_ShowCategories));
 		}
-		if (SuperClass->HasMetaData(FHeaderParserNames::NAME_SparseClassDataTypes))
+		if (SuperClassDef->HasMetaData(FHeaderParserNames::NAME_SparseClassDataTypes))
 		{
-			CurrentClass->SetMetaData(FHeaderParserNames::NAME_SparseClassDataTypes, *SuperClass->GetMetaData(FHeaderParserNames::NAME_SparseClassDataTypes));
+			CurrentClassDef.SetMetaData(FHeaderParserNames::NAME_SparseClassDataTypes, *SuperClassDef->GetMetaData(FHeaderParserNames::NAME_SparseClassDataTypes));
 		}
-		if (SuperClass->HasMetaData(FHeaderParserNames::NAME_HideFunctions))
+		if (SuperClassDef->HasMetaData(FHeaderParserNames::NAME_HideFunctions))
 		{
-			CurrentClass->SetMetaData(FHeaderParserNames::NAME_HideFunctions, *SuperClass->GetMetaData(FHeaderParserNames::NAME_HideFunctions));
+			CurrentClassDef.SetMetaData(FHeaderParserNames::NAME_HideFunctions, *SuperClassDef->GetMetaData(FHeaderParserNames::NAME_HideFunctions));
 		}
-		if (SuperClass->HasMetaData(FHeaderParserNames::NAME_AutoExpandCategories))
+		if (SuperClassDef->HasMetaData(FHeaderParserNames::NAME_AutoExpandCategories))
 		{
-			CurrentClass->SetMetaData(FHeaderParserNames::NAME_AutoExpandCategories, *SuperClass->GetMetaData(FHeaderParserNames::NAME_AutoExpandCategories));
+			CurrentClassDef.SetMetaData(FHeaderParserNames::NAME_AutoExpandCategories, *SuperClassDef->GetMetaData(FHeaderParserNames::NAME_AutoExpandCategories));
 		}
-		if (SuperClass->HasMetaData(FHeaderParserNames::NAME_AutoCollapseCategories))
+		if (SuperClassDef->HasMetaData(FHeaderParserNames::NAME_AutoCollapseCategories))
 		{
-			CurrentClass->SetMetaData(FHeaderParserNames::NAME_AutoCollapseCategories, *SuperClass->GetMetaData(FHeaderParserNames::NAME_AutoCollapseCategories));
+			CurrentClassDef.SetMetaData(FHeaderParserNames::NAME_AutoCollapseCategories, *SuperClassDef->GetMetaData(FHeaderParserNames::NAME_AutoCollapseCategories));
 		}
 	}
 
@@ -9100,33 +9097,27 @@ void FHeaderParser::PostPopNestClass(FUnrealClassDefinitionInfo& CurrentClassDef
 				continue;
 			}
 
-			UClass* Interface = InterfaceDef->GetClass();
-
 			// So iterate over all functions this interface declares
 			for (FUnrealFunctionDefinitionInfo* InterfaceFunctionDef : TUHTFieldRange<FUnrealFunctionDefinitionInfo>(*InterfaceDef, EFieldIteratorFlags::ExcludeSuper))
 			{
-				UFunction* InterfaceFunction = InterfaceFunctionDef->GetFunction();
-
 				bool bImplemented = false;
 
 				// And try to find one that matches
 				for (FUnrealFunctionDefinitionInfo* ClassFunctionDef : TUHTFieldRange<FUnrealFunctionDefinitionInfo>(CurrentClassDef))
 				{
-					UFunction* ClassFunction = ClassFunctionDef->GetFunction();
-
-					if (ClassFunction->GetFName() != InterfaceFunction->GetFName())
+					if (ClassFunctionDef->GetFName() != InterfaceFunctionDef->GetFName())
 					{
 						continue;
 					}
 
-					if ((InterfaceFunction->FunctionFlags & FUNC_Event) && !(ClassFunction->FunctionFlags & FUNC_Event))
+					if (InterfaceFunctionDef->HasAnyFunctionFlags(FUNC_Event) && !ClassFunctionDef->HasAnyFunctionFlags(FUNC_Event))
 					{
-						FError::Throwf(TEXT("Implementation of function '%s::%s' must be declared as 'event' to match declaration in interface '%s'"), *ClassFunction->GetOuter()->GetName(), *ClassFunction->GetName(), *Interface->GetName());
+						FError::Throwf(TEXT("Implementation of function '%s::%s' must be declared as 'event' to match declaration in interface '%s'"), *ClassFunctionDef->GetOuter()->GetName(), *ClassFunctionDef->GetName(), *InterfaceDef->GetName());
 					}
 
-					if ((InterfaceFunction->FunctionFlags & FUNC_Delegate) && !(ClassFunction->FunctionFlags & FUNC_Delegate))
+					if (InterfaceFunctionDef->HasAnyFunctionFlags(FUNC_Delegate) && !ClassFunctionDef->HasAnyFunctionFlags(FUNC_Delegate))
 					{
-						FError::Throwf(TEXT("Implementation of function '%s::%s' must be declared as 'delegate' to match declaration in interface '%s'"), *ClassFunction->GetOuter()->GetName(), *ClassFunction->GetName(), *Interface->GetName());
+						FError::Throwf(TEXT("Implementation of function '%s::%s' must be declared as 'delegate' to match declaration in interface '%s'"), *ClassFunctionDef->GetOuter()->GetName(), *ClassFunctionDef->GetName(), *InterfaceDef->GetName());
 					}
 
 					// Making sure all the parameters match up correctly
@@ -9134,7 +9125,7 @@ void FHeaderParser::PostPopNestClass(FUnrealClassDefinitionInfo& CurrentClassDef
 
 					if (ClassFunctionDef->GetProperties().Num() != InterfaceFunctionDef->GetProperties().Num())
 					{
-						FError::Throwf(TEXT("Implementation of function '%s' conflicts with interface '%s' - different number of parameters (%i/%i)"), *InterfaceFunction->GetName(), *Interface->GetName(), ClassFunctionDef->GetProperties().Num(), InterfaceFunctionDef->GetProperties().Num());
+						FError::Throwf(TEXT("Implementation of function '%s' conflicts with interface '%s' - different number of parameters (%i/%i)"), *InterfaceFunctionDef->GetName(), *InterfaceDef->GetName(), ClassFunctionDef->GetProperties().Num(), InterfaceFunctionDef->GetProperties().Num());
 					}
 
 					for (int32 Index = 0, End = ClassFunctionDef->GetProperties().Num(); Index != End; ++Index)
@@ -9146,29 +9137,29 @@ void FHeaderParser::PostPopNestClass(FUnrealClassDefinitionInfo& CurrentClassDef
 							FProperty* Property = InterfaceParamDef->GetProperty();
 							if (Property->HasAnyPropertyFlags(CPF_ReturnParm))
 							{
-								FError::Throwf(TEXT("Implementation of function '%s' conflicts only by return type with interface '%s'"), *InterfaceFunction->GetName(), *Interface->GetName());
+								FError::Throwf(TEXT("Implementation of function '%s' conflicts only by return type with interface '%s'"), *InterfaceFunctionDef->GetName(), *InterfaceDef->GetName());
 							}
 							else
 							{
-								FError::Throwf(TEXT("Implementation of function '%s' conflicts with interface '%s' - parameter %i '%s'"), *InterfaceFunction->GetName(), *Interface->GetName(), End, *Property->GetName());
+								FError::Throwf(TEXT("Implementation of function '%s' conflicts with interface '%s' - parameter %i '%s'"), *InterfaceFunctionDef->GetName(), *InterfaceDef->GetName(), End, *Property->GetName());
 							}
 						}
 					}
 				}
 
 				// Delegate signature functions are simple stubs and aren't required to be implemented (they are not callable)
-				if (InterfaceFunction->FunctionFlags & FUNC_Delegate)
+				if (InterfaceFunctionDef->HasAnyFunctionFlags(FUNC_Delegate))
 				{
 					bImplemented = true;
 				}
 
 				// Verify that if this has blueprint-callable functions that are not implementable events, we've implemented them as a UFunction in the target class
 				if (!bImplemented
-					&& InterfaceFunction->HasAnyFunctionFlags(FUNC_BlueprintCallable)
-					&& !InterfaceFunction->HasAnyFunctionFlags(FUNC_BlueprintEvent)
-					&& !Interface->HasMetaData(NAME_CannotImplementInterfaceInBlueprint))  // FBlueprintMetadata::MD_CannotImplementInterfaceInBlueprint
+					&& InterfaceFunctionDef->HasAnyFunctionFlags(FUNC_BlueprintCallable)
+					&& !InterfaceFunctionDef->HasAnyFunctionFlags(FUNC_BlueprintEvent)
+					&& !InterfaceDef->HasMetaData(NAME_CannotImplementInterfaceInBlueprint))  // FBlueprintMetadata::MD_CannotImplementInterfaceInBlueprint
 				{
-					FError::Throwf(TEXT("Missing UFunction implementation of function '%s' from interface '%s'.  This function needs a UFUNCTION() declaration."), *InterfaceFunction->GetName(), *Interface->GetName());
+					FError::Throwf(TEXT("Missing UFunction implementation of function '%s' from interface '%s'.  This function needs a UFUNCTION() declaration."), *InterfaceFunctionDef->GetName(), *InterfaceDef->GetName());
 				}
 			}
 		}
@@ -9232,8 +9223,6 @@ void FHeaderParser::CheckDocumentationPolicyForEnum(FUnrealEnumDefinitionInfo& E
 {
 	SCOPE_SECONDS_COUNTER_UHT(DocumentationPolicy);
 
-	UEnum* Enum = EnumDef.GetEnum();
-
 	const FString* DocumentationPolicyName = MetaData.Find(NAME_DocumentationPolicy);
 	if (DocumentationPolicyName == nullptr)
 	{
@@ -9248,7 +9237,7 @@ void FHeaderParser::CheckDocumentationPolicyForEnum(FUnrealEnumDefinitionInfo& E
 		const FString* EnumToolTip = MetaData.Find(NAME_ToolTip);
 		if (EnumToolTip == nullptr)
 		{
-			UE_LOG_ERROR_UHT(TEXT("Enum '%s' does not provide a tooltip / comment (DocumentationPolicy)."), *Enum->GetName());
+			UE_LOG_ERROR_UHT(TEXT("Enum '%s' does not provide a tooltip / comment (DocumentationPolicy)."), *EnumDef.GetName());
 		}
 	}
 
@@ -9264,14 +9253,14 @@ void FHeaderParser::CheckDocumentationPolicyForEnum(FUnrealEnumDefinitionInfo& E
 		const FString* ToolTip = Entry.Find(NAME_ToolTip);
 		if (ToolTip == nullptr)
 		{
-			UE_LOG_ERROR_UHT(TEXT("Enum entry '%s::%s' does not provide a tooltip / comment (DocumentationPolicy)."), *Enum->GetName(), *(*EntryName));
+			UE_LOG_ERROR_UHT(TEXT("Enum entry '%s::%s' does not provide a tooltip / comment (DocumentationPolicy)."), *EnumDef.GetName(), *(*EntryName));
 			continue;
 		}
 
 		const FString* ExistingEntry = ToolTipToEntry.Find(*ToolTip);
 		if (ExistingEntry != nullptr)
 		{
-			UE_LOG_ERROR_UHT(TEXT("Enum entries '%s::%s' and '%s::%s' have identical tooltips / comments (DocumentationPolicy)."), *Enum->GetName(), *(*ExistingEntry), *Enum->GetName(), *(*EntryName));
+			UE_LOG_ERROR_UHT(TEXT("Enum entries '%s::%s' and '%s::%s' have identical tooltips / comments (DocumentationPolicy)."), *EnumDef.GetName(), *(*ExistingEntry), *EnumDef.GetName(), *(*EntryName));
 		}
 		ToolTipToEntry.Add(*ToolTip, *EntryName);
 	}
@@ -9351,15 +9340,15 @@ void FHeaderParser::CheckDocumentationPolicyForStruct(FUnrealStructDefinitionInf
 				FString ToolTip = Func->GetToolTipText().ToString();
 				if (ToolTip.IsEmpty())
 				{
-					UE_LOG_ERROR_UHT(TEXT("Function '%s::%s' does not provide a tooltip / comment (DocumentationPolicy)."), *Class->GetName(), *Func->GetName());
+					UE_LOG_ERROR_UHT(TEXT("Function '%s::%s' does not provide a tooltip / comment (DocumentationPolicy)."), *Class->GetName(), *FunctionDef->GetName());
 					continue;
 				}
 				const FName* ExistingFuncName = ToolTipToFunc.Find(ToolTip);
 				if (ExistingFuncName != nullptr)
 				{
-					UE_LOG_ERROR_UHT(TEXT("Functions '%s::%s' and '%s::%s' uses identical tooltips / comments (DocumentationPolicy)."), *Class->GetName(), *(*ExistingFuncName).ToString(), *Class->GetName(), *Func->GetName());
+					UE_LOG_ERROR_UHT(TEXT("Functions '%s::%s' and '%s::%s' uses identical tooltips / comments (DocumentationPolicy)."), *Class->GetName(), *(*ExistingFuncName).ToString(), *Class->GetName(), *FunctionDef->GetName());
 				}
-				ToolTipToFunc.Add(MoveTemp(ToolTip), Func->GetFName());
+				ToolTipToFunc.Add(MoveTemp(ToolTip), FunctionDef->GetFName());
 			}
 		}
 	}
@@ -9376,24 +9365,23 @@ void FHeaderParser::CheckDocumentationPolicyForFunc(FUnrealClassDefinitionInfo&
 	SCOPE_SECONDS_COUNTER_UHT(DocumentationPolicy);
 
 	UClass* Class = ClassDef.GetClass();
-	UFunction* Func = FunctionDef.GetFunction();
 
 	FDocumentationPolicy DocumentationPolicy = GetDocumentationPolicyForStruct(ClassDef);
 	if (DocumentationPolicy.bFunctionToolTipsRequired)
 	{
-		const FString* FunctionTooltip = Func->FindMetaData(NAME_ToolTip);
+		const FString* FunctionTooltip = FunctionDef.FindMetaData(NAME_ToolTip);
 		if (FunctionTooltip == nullptr)
 		{
-			UE_LOG_ERROR_UHT(TEXT("Function '%s::%s' does not provide a tooltip / comment (DocumentationPolicy)."), *Class->GetName(), *Func->GetName());
+			UE_LOG_ERROR_UHT(TEXT("Function '%s::%s' does not provide a tooltip / comment (DocumentationPolicy)."), *Class->GetName(), *FunctionDef.GetName());
 		}
 	}
 
 	if (DocumentationPolicy.bParameterToolTipsRequired)
 	{
-		const FString* FunctionComment = Func->FindMetaData(NAME_Comment);
+		const FString* FunctionComment = FunctionDef.FindMetaData(NAME_Comment);
 		if (FunctionComment == nullptr)
 		{
-			UE_LOG_ERROR_UHT(TEXT("Function '%s::%s' does not provide a comment (DocumentationPolicy)."), *Class->GetName(), *Func->GetName());
+			UE_LOG_ERROR_UHT(TEXT("Function '%s::%s' does not provide a comment (DocumentationPolicy)."), *Class->GetName(), *FunctionDef.GetName());
 			return;
 		}
 		
@@ -9424,7 +9412,7 @@ void FHeaderParser::CheckDocumentationPolicyForFunc(FUnrealClassDefinitionInfo&
 				const FString* ParamToolTip = ParamToolTips.Find(ParamName);
 				if (ParamToolTip == nullptr)
 				{
-					UE_LOG_ERROR_UHT(TEXT("Function '%s::%s' doesn't provide a tooltip for parameter '%s' (DocumentationPolicy)."), *Class->GetName(), *Func->GetName(), *ParamName.ToString());
+					UE_LOG_ERROR_UHT(TEXT("Function '%s::%s' doesn't provide a tooltip for parameter '%s' (DocumentationPolicy)."), *Class->GetName(), *FunctionDef.GetName(), *ParamName.ToString());
 				}
 				ExistingFields.Add(ParamName);
 			}
@@ -9439,7 +9427,7 @@ void FHeaderParser::CheckDocumentationPolicyForFunc(FUnrealClassDefinitionInfo&
 				}
 				if (!ExistingFields.Contains(ParamName))
 				{
-					UE_LOG_ERROR_UHT(TEXT("Function '%s::%s' provides a tooltip for an unknown parameter '%s' (DocumentationPolicy)."), *Class->GetName(), *Func->GetName(), *Pair.Key.ToString());
+					UE_LOG_ERROR_UHT(TEXT("Function '%s::%s' provides a tooltip for an unknown parameter '%s' (DocumentationPolicy)."), *Class->GetName(), *FunctionDef.GetName(), *Pair.Key.ToString());
 				}
 			}
 
@@ -9455,7 +9443,7 @@ void FHeaderParser::CheckDocumentationPolicyForFunc(FUnrealClassDefinitionInfo&
 				const FName* ExistingParam = ToolTipToParam.Find(Pair.Value);
 				if (ExistingParam != nullptr)
 				{
-					UE_LOG_ERROR_UHT(TEXT("Function '%s::%s' uses identical tooltips for parameters '%s' and '%s' (DocumentationPolicy)."), *Class->GetName(), *Func->GetName(), *ExistingParam->ToString(), *Pair.Key.ToString());
+					UE_LOG_ERROR_UHT(TEXT("Function '%s::%s' uses identical tooltips for parameters '%s' and '%s' (DocumentationPolicy)."), *Class->GetName(), *FunctionDef.GetName(), *ExistingParam->ToString(), *Pair.Key.ToString());
 				}
 				ToolTipToParam.Add(MoveTemp(Pair.Value), Pair.Key);
 			}
@@ -9552,13 +9540,13 @@ FUnrealFunctionDefinitionInfo& CreateFunctionImpl(FUnrealSourceFile& SourceFile,
 	Function->ReturnValueOffset = MAX_uint16;
 	Function->FirstPropertyToInit = nullptr;
 	Function->FunctionFlags |= FuncInfo.FunctionFlags;
-	FuncInfo.FunctionReference = Function;
-	FuncInfo.SetFunctionNames();
 
 	TSharedRef<FUnrealFunctionDefinitionInfo> FuncDefRef = MakeShared<FUnrealFunctionDefinitionInfo>(SourceFile, InputLine, FString(FuncInfo.Function.Identifier), OuterDef, MoveTemp(FuncInfo));
 	FUnrealFunctionDefinitionInfo& FuncDef = *FuncDefRef;
 	FuncDef.SetObject(Function);
 	GTypeDefinitionInfoMap.Add(Function, MoveTemp(FuncDefRef));
+	FuncDef.GetFunctionData().SetFunctionNames(FuncDef);
+
 	CurrentScope->AddType(FuncDef);
 
 	if (!CurrentScope->IsFileScope())
diff --git a/Engine/Source/Programs/UnrealHeaderTool/Private/ParserHelper.cpp b/Engine/Source/Programs/UnrealHeaderTool/Private/ParserHelper.cpp
--- a/Engine/Source/Programs/UnrealHeaderTool/Private/ParserHelper.cpp
+++ b/Engine/Source/Programs/UnrealHeaderTool/Private/ParserHelper.cpp
@@ -316,6 +316,75 @@ FToken& FToken::operator=(FToken&& Other)
 	return *this;
 }
 
+/////////////////////////////////////////////////////
+// FFuncData
+
+void FFuncInfo::SetFunctionNames(FUnrealFunctionDefinitionInfo& FunctionDef)
+{
+	UFunction* FunctionReference = FunctionDef.GetFunction();
+	FString FunctionName = FunctionReference->GetName();
+	if (FunctionDef.HasAnyFunctionFlags(FUNC_Delegate))
+	{
+		FunctionName.LeftChopInline(FString(HEADER_GENERATED_DELEGATE_SIGNATURE_SUFFIX).Len(), false);
+	}
+	UnMarshallAndCallName = FString(TEXT("exec")) + FunctionName;
+
+	if (FunctionDef.HasAnyFunctionFlags(FUNC_BlueprintEvent))
+	{
+		MarshallAndCallName = FunctionName;
+	}
+	else
+	{
+		MarshallAndCallName = FString(TEXT("event")) + FunctionName;
+	}
+
+	if (FunctionDef.HasAllFunctionFlags(FUNC_Native | FUNC_Net))
+	{
+		MarshallAndCallName = FunctionName;
+		if (FunctionDef.HasAllFunctionFlags(FUNC_NetResponse))
+		{
+			// Response function implemented by programmer and called directly from thunk
+			CppImplName = FunctionReference->GetName();
+		}
+		else
+		{
+			if (CppImplName.IsEmpty())
+			{
+				CppImplName = FunctionReference->GetName() + TEXT("_Implementation");
+			}
+			else if (CppImplName == FunctionName)
+			{
+				FError::Throwf(TEXT("Native implementation function must be different than original function name."));
+			}
+
+			if (CppValidationImplName.IsEmpty() && FunctionDef.HasAllFunctionFlags(FUNC_NetValidate))
+			{
+				CppValidationImplName = FunctionReference->GetName() + TEXT("_Validate");
+			}
+			else if (CppValidationImplName == FunctionName)
+			{
+				FError::Throwf(TEXT("Validation function must be different than original function name."));
+			}
+		}
+	}
+
+	if (FunctionDef.HasAllFunctionFlags(FUNC_Delegate))
+	{
+		MarshallAndCallName = FString(TEXT("delegate")) + FunctionName;
+	}
+
+	if (FunctionDef.HasAllFunctionFlags(FUNC_BlueprintEvent | FUNC_Native))
+	{
+		MarshallAndCallName = FunctionName;
+		CppImplName = FunctionReference->GetName() + TEXT("_Implementation");
+	}
+
+	if (CppImplName.IsEmpty())
+	{
+		CppImplName = FunctionName;
+	}
+}
+
 /////////////////////////////////////////////////////
 // FAdvancedDisplayParameterHandler
 static const FName NAME_AdvancedDisplay(TEXT("AdvancedDisplay"));
diff --git a/Engine/Source/Programs/UnrealHeaderTool/Private/ParserHelper.h b/Engine/Source/Programs/UnrealHeaderTool/Private/ParserHelper.h
--- a/Engine/Source/Programs/UnrealHeaderTool/Private/ParserHelper.h
+++ b/Engine/Source/Programs/UnrealHeaderTool/Private/ParserHelper.h
@@ -564,8 +564,6 @@ struct FFuncInfo
 	uint32		FunctionExportFlags;
 	/** Number of parameters expected for operator. */
 	int32		ExpectParms;
-	/** Pointer to the UFunction corresponding to this FFuncInfo */
-	UFunction*	FunctionReference;
 	/** Name of the wrapper function that marshalls the arguments and does the indirect call **/
 	FString		MarshallAndCallName;
 	/** Name of the actual implementation **/
@@ -597,7 +595,6 @@ struct FFuncInfo
 		, FunctionFlags(FUNC_None)
 		, FunctionExportFlags(0)
 		, ExpectParms(0)
-		, FunctionReference(nullptr)
 		, RPCId(0)
 		, RPCResponseId(0)
 		, MacroLine(-1)
@@ -611,70 +608,7 @@ struct FFuncInfo
 	FFuncInfo& operator=(FFuncInfo&& Other) = default;
 
 	/** Set the internal function names based on flags **/
-	void SetFunctionNames()
-	{
-		FString FunctionName = FunctionReference->GetName();
-		if( FunctionReference->HasAnyFunctionFlags( FUNC_Delegate ) )
-		{
-			FunctionName.LeftChopInline( FString(HEADER_GENERATED_DELEGATE_SIGNATURE_SUFFIX).Len(), false );
-		}
-		UnMarshallAndCallName = FString(TEXT("exec")) + FunctionName;
-		
-		if (FunctionReference->HasAnyFunctionFlags(FUNC_BlueprintEvent))
-		{
-			MarshallAndCallName = FunctionName;
-		}
-		else
-		{
-			MarshallAndCallName = FString(TEXT("event")) + FunctionName;
-		}
-
-		if (FunctionReference->HasAllFunctionFlags(FUNC_Native | FUNC_Net))
-		{
-			MarshallAndCallName = FunctionName;
-			if (FunctionReference->HasAllFunctionFlags(FUNC_NetResponse))
-			{
-				// Response function implemented by programmer and called directly from thunk
-				CppImplName = FunctionReference->GetName();
-			}
-			else
-			{
-				if (CppImplName.IsEmpty())
-				{
-					CppImplName = FunctionReference->GetName() + TEXT("_Implementation");
-				}
-				else if (CppImplName == FunctionName)
-				{
-					FError::Throwf(TEXT("Native implementation function must be different than original function name."));
-				}
-
-				if (CppValidationImplName.IsEmpty() && FunctionReference->HasAllFunctionFlags(FUNC_NetValidate))
-				{
-					CppValidationImplName = FunctionReference->GetName() + TEXT("_Validate");
-				}
-				else if(CppValidationImplName == FunctionName)
-				{
-					FError::Throwf(TEXT("Validation function must be different than original function name."));
-				}
-			}
-		}
-
-		if (FunctionReference->HasAllFunctionFlags(FUNC_Delegate))
-		{
-			MarshallAndCallName = FString(TEXT("delegate")) + FunctionName;
-		}
-
-		if (FunctionReference->HasAllFunctionFlags(FUNC_BlueprintEvent | FUNC_Native))
-		{
-			MarshallAndCallName = FunctionName;
-			CppImplName = FunctionReference->GetName() + TEXT("_Implementation");
-		}
-
-		if (CppImplName.IsEmpty())
-		{
-			CppImplName = FunctionName;
-		}
-	}
+	void SetFunctionNames(FUnrealFunctionDefinitionInfo& FunctionDef);
 };
 
 enum class EParsedInterface
diff --git a/Engine/Source/Programs/UnrealHeaderTool/Private/UnrealTypeDefinitionInfo.cpp b/Engine/Source/Programs/UnrealHeaderTool/Private/UnrealTypeDefinitionInfo.cpp
--- a/Engine/Source/Programs/UnrealHeaderTool/Private/UnrealTypeDefinitionInfo.cpp
+++ b/Engine/Source/Programs/UnrealHeaderTool/Private/UnrealTypeDefinitionInfo.cpp
@@ -281,7 +281,6 @@ FString FUnrealTypeDefinitionInfo::GetNameWithPrefix(const UClass* InClass, EEnf
 	return FString::Printf(TEXT("%s%s"), Prefix, *InClass->GetName());
 }
 
-
 void FUnrealPropertyDefinitionInfo::PostParseFinalize()
 {
 	TypePackageName = GetTypePackageName_Inner(GetProperty());
@@ -326,6 +325,16 @@ bool FUnrealPropertyDefinitionInfo::IsOwnedByDynamicType() const
 	return false;
 }
 
+FUnrealPackageDefinitionInfo& FUnrealObjectDefinitionInfo::GetPackageDef() const
+{
+	if (HasSource())
+	{
+		return GetUnrealSourceFile().GetPackageDef();
+	}
+	return GTypeDefinitionInfoMap.FindChecked<FUnrealPackageDefinitionInfo>(GetObject());
+}
+
+
 FUnrealPackageDefinitionInfo::FUnrealPackageDefinitionInfo(const FManifestModule& InModule, UPackage* InPackage)
 	: FUnrealObjectDefinitionInfo(FString())
 	, Module(InModule)
@@ -382,7 +391,7 @@ void FUnrealFieldDefinitionInfo::AddCrossModuleReference(TSet<FString>* UniqueCr
 	if (UniqueCrossModuleReferences)
 	{
 		UField* Field = GetField();
-		if (!Field->IsA<UFunction>() || Field->IsA<UDelegateFunction>())
+		if (!Field->IsA<UFunction>() || IsADelegateFunction())
 		{
 			UniqueCrossModuleReferences->Add(GetExternDecl(bRequiresValidObject));
 		}
@@ -483,7 +492,7 @@ void FUnrealStructDefinitionInfo::AddFunction(FUnrealFunctionDefinitionInfo& Fun
 	// update the optimization flags
 	if (!bContainsDelegates)
 	{
-		if (FunctionDef.GetFunction()->HasAnyFunctionFlags(FUNC_Delegate))
+		if (FunctionDef.HasAnyFunctionFlags(FUNC_Delegate))
 		{
 			bContainsDelegates = true;
 		}
@@ -525,7 +534,6 @@ void FUnrealStructDefinitionInfo::SetObject(UObject* InObject)
 
 	FUnrealFieldDefinitionInfo::SetObject(InObject);
 
-
 	// Don't create a scope for things without a source.  Those are builtin types
 	if (HasSource())
 	{
diff --git a/Engine/Source/Programs/UnrealHeaderTool/Private/UnrealTypeDefinitionInfo.h b/Engine/Source/Programs/UnrealHeaderTool/Private/UnrealTypeDefinitionInfo.h
--- a/Engine/Source/Programs/UnrealHeaderTool/Private/UnrealTypeDefinitionInfo.h
+++ b/Engine/Source/Programs/UnrealHeaderTool/Private/UnrealTypeDefinitionInfo.h
@@ -609,6 +609,25 @@ public:
 		Object = InObject;
 	}
 
+	/**
+	 * Returns the name of this object (with no path information)
+	 */
+	FString GetName() const
+	{
+		return GetObject()->GetName();
+	}
+
+	/** Returns the logical name of this object */
+	FName GetFName() const
+	{
+		return GetObject()->GetFName();
+	}
+
+	/**
+	 * Return the package associated with this object
+	 */
+	FUnrealPackageDefinitionInfo& GetPackageDef() const;
+
 	/**
 	 * Return the "outer" object that contains the definition for this object
 	 */
@@ -617,6 +636,14 @@ public:
 		return static_cast<FUnrealObjectDefinitionInfo*>(FUnrealTypeDefinitionInfo::GetOuter());
 	}
 
+	/**
+	 * Helper methods to remove dependencies on engine types
+	 */
+	bool IsADelegateFunction() const
+	{
+		return GetObject()->IsA<UDelegateFunction>();
+	}
+
 protected:
 	// Used only by packages
 	explicit FUnrealObjectDefinitionInfo(FString&& InNameCPP)
@@ -788,6 +815,60 @@ public:
 		return static_cast<UField*>(GetObject());
 	}
 
+	/**
+	 * Determines if the property has any metadata associated with the key
+	 * 
+	 * @param Key The key to lookup in the metadata
+	 * @return true if there is a (possibly blank) value associated with this key
+	 */
+	bool HasMetaData(const TCHAR* Key) const { return FindMetaData(Key) != nullptr; }
+	bool HasMetaData(const FName& Key) const { return FindMetaData(Key) != nullptr; }
+
+	/**
+	 * Find the metadata value associated with the key
+	 *
+	 * @param Key The key to lookup in the metadata
+	 * @return The value associated with the key if exists, null otherwise
+	 */
+	const FString* FindMetaData(const TCHAR* Key) const
+	{
+		return GetField()->FindMetaData(Key);
+	}
+	const FString* FindMetaData(const FName& Key) const
+	{
+		return GetField()->FindMetaData(Key);
+	}
+
+	/**
+	 * Find the metadata value associated with the key
+	 * 
+	 * @param Key The key to lookup in the metadata
+	 * @return The value associated with the key
+	 */
+	const FString& GetMetaData(const TCHAR* Key) const
+	{
+		return GetField()->GetMetaData(Key);
+	}
+	const FString& GetMetaData(const FName& Key) const
+	{
+		return GetField()->GetMetaData(Key);
+	}
+
+	/**
+	 * Sets the metadata value associated with the key
+	 * 
+	 * @param Key The key to lookup in the metadata
+	 * @return The value associated with the key
+	 */
+	void SetMetaData(const TCHAR* Key, const TCHAR* InValue)
+	{
+		GetField()->SetMetaData(Key, InValue);
+	}
+	void SetMetaData(const FName& Key, const TCHAR* InValue)
+	{
+		GetField()->SetMetaData(Key, InValue);
+	}
+
 	/**
 	 * Perform any post parsing finalization and validation
 	 */
@@ -886,6 +967,149 @@ public:
 		return static_cast<UEnum*>(GetObject());
 	}
 
+	bool HasAnyEnumFlags(EEnumFlags InFlags) const
+	{
+		return GetEnum()->HasAnyEnumFlags(InFlags);
+	}
+
+	/**
+	 * Tests if the enum contains a MAX value
+	 *
+	 * @return	true if the enum contains a MAX enum, false otherwise.
+	 */
+	bool ContainsExistingMax() const
+	{
+		return GetEnum()->ContainsExistingMax();
+	}
+
+	/**
+	 * @return	 The number of enum names.
+	 */
+	int32 NumEnums() const
+	{
+		return GetEnum()->NumEnums();
+	}
+
+	/** Gets max value of Enum. Defaults to zero if there are no entries. */
+	int64 GetMaxEnumValue() const
+	{
+		return GetEnum()->GetMaxEnumValue();
+	}
+
+	/** Gets enum value by index in Names array. Asserts on invalid index */
+	FORCEINLINE int64 GetValueByIndex(int32 Index) const
+	{
+		return GetEnum()->GetValueByIndex(Index);
+	}
+
+	/** Gets enum name by index in Names array. Returns NAME_None if Index is not valid. */
+	FName GetNameByIndex(int32 Index) const
+	{
+		return GetEnum()->GetNameByIndex(Index);
+	}
+
+	/** Gets index of name in enum, returns INDEX_NONE and optionally errors when name is not found. This is faster than ByNameString if the FName is exact, but will fall back if needed */
+	int32 GetIndexByName(FName InName, EGetByNameFlags Flags = EGetByNameFlags::None) const
+	{
+		return GetEnum()->GetIndexByName(InName, Flags);
+	}
+
+	/** Checks if enum has entry with given value. Includes autogenerated _MAX entry. */
+	bool IsValidEnumValue(int64 InValue) const
+	{
+		return GetEnum()->IsValidEnumValue(InValue);
+	}
+
+	/**
+	 * Returns the type of enum: whether it's a regular enum, namespaced enum or C++11 enum class.
+	 *
+	 * @return The enum type.
+	 */
+	UEnum::ECppForm GetCppForm() const
+	{
+		return GetEnum()->GetCppForm();
+	}
+
+	/**
+	 * Returns the type of the enum
+	 */
+	const FString& GetCppType() const
+	{
+		return GetEnum()->CppType;
+	}
+
+	/**
+	 * Sets the enum type
+	 */
+	void SetCppType(FString&& CppType)
+	{
+		GetEnum()->CppType = MoveTemp(CppType);
+	}
+
+	/**
+	 * Find the longest common prefix of all items in the enumeration.
+	 *
+	 * @return	the longest common prefix between all items in the enum.  If a common prefix
+	 *			cannot be found, returns the full name of the enum.
+	 */
+	FString GenerateEnumPrefix() const
+	{
+		return GetEnum()->GenerateEnumPrefix();
+	}
+
+	/**
+	 * Sets the array of enums.
+	 *
+	 * @param InNames List of enum names.
+	 * @param InCppForm The form of enum.
+	 * @param bAddMaxKeyIfMissing Should a default Max item be added.
+	 * @return	true unless the MAX enum already exists and isn't the last enum.
+	 */
+	bool SetEnums(TArray<TPair<FName, int64>>& InNames, UEnum::ECppForm InCppForm, EEnumFlags InFlags = EEnumFlags::None, bool bAddMaxKeyIfMissing = true)
+	{
+		return GetEnum()->SetEnums(InNames, InCppForm, InFlags, bAddMaxKeyIfMissing);
+	}
+
+	/**
+	 * Wrapper method for easily determining whether this enum has metadata associated with it.
+	 *
+	 * @param	Key			the metadata tag to check for
+	 * @param	NameIndex	if specified, will search for metadata linked to a specified value in this enum; otherwise, searches for metadata for the enum itself
+	 *
+	 * @return true if the specified key exists in the list of metadata for this enum, even if the value of that key is empty
+	 */
+	bool HasMetaData(const TCHAR* Key, int32 NameIndex = INDEX_NONE) const
+	{
+		return GetEnum()->HasMetaData(Key, NameIndex);
+	}
+
+	/**
+	 * Return the metadata value associated with the specified key.
+	 *
+	 * @param	Key			the metadata tag to find the value for
+	 * @param	NameIndex	if specified, will search the metadata linked for that enum value; otherwise, searches the metadata for the enum itself
+	 * @param	bAllowRemap	if true, the returned value may be remapped from a .ini if the value starts with ini: Pass false when you need the exact string, including any ini:
+	 *
+	 * @return	the value for the key specified, or an empty string if the key wasn't found or had no value.
+	 */
+	FString GetMetaData(const TCHAR* Key, int32 NameIndex = INDEX_NONE, bool bAllowRemap = true) const
+	{
+		return GetEnum()->GetMetaData(Key, NameIndex, bAllowRemap);
+	}
+
+	/**
+	 * Set the metadata value associated with the specified key.
+	 *
+	 * @param	Key			the metadata tag to find the value for
+	 * @param	NameIndex	if specified, will search the metadata linked for that enum value; otherwise, searches the metadata for the enum itself
+	 * @param	InValue		Value of the metadata for the key
+	 *
+	 */
+	void SetMetaData(const TCHAR* Key, const TCHAR* InValue, int32 NameIndex = INDEX_NONE) const
+	{
+		GetEnum()->SetMetaData(Key, InValue, NameIndex);
+	}
+
 	/**
 	 * Returns the underlying enumeration type
 	 */
@@ -970,6 +1194,20 @@ public:
 		return static_cast<UStruct*>(GetObject());
 	}
 
+	/**
+	 * Returns the name used for declaring the passed in struct in C++
+	 * 
+	 * NOTE: This does not match the CPP name parsed from the header file.
+	 *
+	 * @param	bForceInterface If true, force an 'I' prefix
+	 * @return	Name used for C++ declaration
+	 */
+	FString GetAlternateNameCPP(bool bForceInterface = false)
+	{
+		UStruct* Struct = GetStruct();
+		return FString::Printf(TEXT("%s%s"), (bForceInterface ? TEXT("I") : Struct->GetPrefixCPP()), *Struct->GetName());
+	}
+
 	/**
 	 * Add a new property to the structure
 	 */
@@ -1192,7 +1430,8 @@ public:
 	FUnrealFunctionDefinitionInfo(FUnrealSourceFile& InSourceFile, int32 InLineNumber, FString&& InNameCPP, FUnrealObjectDefinitionInfo& InOuter, FFuncInfo&& InFuncInfo)
 		: FUnrealStructDefinitionInfo(InSourceFile, InLineNumber, MoveTemp(InNameCPP), InOuter)
 		, FunctionData(MoveTemp(InFuncInfo))
-	{ }
+	{ 
+	}
 
 	virtual FUnrealFunctionDefinitionInfo* AsFunction() override
 	{
@@ -1215,6 +1454,56 @@ public:
 		return static_cast<UFunction*>(GetField());
 	}
 
+	/**
+	 * Return the function flags
+	 */
+	FORCEINLINE EFunctionFlags GetFunctionFlags() const
+	{
+		return GetFunction()->FunctionFlags;
+	}
+
+	/**
+	 * Used to safely check whether the passed in flag is set.
+	 *
+	 * @param	FlagToCheck		Class flag to check for
+	 *
+	 * @return	true if the passed in flag is set, false otherwise
+	 *			(including no flag passed in, unless the FlagsToCheck is CLASS_AllFlags)
+	 */
+	FORCEINLINE bool HasAnyFunctionFlags(EFunctionFlags FlagsToCheck) const
+	{
+		return GetFunction()->HasAnyFunctionFlags(FlagsToCheck);
+	}
+
+	/**
+	 * Used to safely check whether all of the passed in flags are set.
+	 *
+	 * @param FlagsToCheck	Function flags to check for
+	 * @return true if all of the passed in flags are set (including no flags passed in), false otherwise
+	 */
+	FORCEINLINE bool HasAllFunctionFlags(EFunctionFlags FlagsToCheck) const
+	{
+		return GetFunction()->HasAllFunctionFlags(FlagsToCheck);
+	}
+
+	/**
+	 * Used to safely check whether specific of the passed in flags are set.
+	 *
+	 * @param FlagsToCheck	Function flags to check for
+	 * @param ExpectedFlags The flags from the flags to check that should be set
+	 * @return true if specific of the passed in flags are set (including no flags passed in), false otherwise
+	 */
+	FORCEINLINE bool HasSpecificFunctionFlags(EFunctionFlags FlagsToCheck, EFunctionFlags ExpectedFlags) const
+	{
+		EFunctionFlags Flags = GetFunction()->FunctionFlags;
+		return (Flags & FlagsToCheck) == ExpectedFlags;
+	}
+
+	/**
+	 * Return the outer in Class form.  
+	 */
+	FUnrealClassDefinitionInfo* GetOuterClass() const;
+
 	/** @name getters */
 	//@{
 	FFuncInfo& GetFunctionData() { return FunctionData; }
@@ -1799,3 +2088,7 @@ struct TUHTFieldRange
 	TUHTFieldIterator<T> Begin;
 };
 
+inline FUnrealClassDefinitionInfo* FUnrealFunctionDefinitionInfo::GetOuterClass() const
+{
+	return UHTCast<FUnrealClassDefinitionInfo>(GetOuter());
+}
