Removing GetScriptStruct and GetStruct as much as possible.

#rnx

diff --git a/Engine/Source/Programs/UnrealHeaderTool/Private/CodeGenerator.cpp b/Engine/Source/Programs/UnrealHeaderTool/Private/CodeGenerator.cpp
--- a/Engine/Source/Programs/UnrealHeaderTool/Private/CodeGenerator.cpp
+++ b/Engine/Source/Programs/UnrealHeaderTool/Private/CodeGenerator.cpp
@@ -1864,14 +1864,14 @@ void FNativeClassHeaderGenerator::OutputProperty(FOutputDevice& DeclOut, FOutput
 	}
 }
 
-static bool IsAlwaysAccessible(UScriptStruct* Script)
+static bool IsAlwaysAccessible(FUnrealScriptStructDefinitionInfo& ScriptDef)
 {
-	FName ToTest = Script->GetFName();
+	FName ToTest = ScriptDef.GetFName();
 	if (ToTest == NAME_Matrix || ToTest == NAME_Matrix44f || ToTest == NAME_Matrix44d)
 	{
 		return false; // special case, the C++ FMatrix does not have the same members.
 	}
-	bool Result = Script->HasDefaults(); // if we have cpp struct ops in it for UHT, then we can assume it is always accessible
+	bool Result = ScriptDef.HasDefaults(); // if we have cpp struct ops in it for UHT, then we can assume it is always accessible
 	if( ToTest == NAME_Plane || ToTest == NAME_Plane4f || ToTest == NAME_Plane4d
 		||	ToTest == NAME_Vector || ToTest == NAME_Vector3f || ToTest == NAME_Vector3d
 		||	ToTest == NAME_Vector4
@@ -1888,15 +1888,14 @@ static void FindNoExportStructsRecursive(TArray<FUnrealScriptStructDefinitionInf
 {
 	for (; StartDef; StartDef = StartDef->GetSuperStructInfo().Struct)
 	{
-		UStruct* Start = StartDef->GetStruct();
-		if (UScriptStruct* StartScript = Cast<UScriptStruct>(Start))
+		if (FUnrealScriptStructDefinitionInfo* StartScriptDef = UHTCast<FUnrealScriptStructDefinitionInfo>(StartDef))
 		{
-			if (StartScript->StructFlags & STRUCT_Native)
+			if (StartScriptDef->HasAnyStructFlags(STRUCT_Native))
 			{
 				break;
 			}
 
-			if (!IsAlwaysAccessible(StartScript)) // these are a special cases that already exists and if wrong if exported naively
+			if (!IsAlwaysAccessible(*StartScriptDef)) // these are a special cases that already exists and if wrong if exported naively
 			{
 				// this will topologically sort them in reverse order
 				StructDefs.Remove(&StartDef->AsScriptStructChecked());
@@ -2332,8 +2331,7 @@ void FNativeClassHeaderGenerator::ExportNativeGeneratedInitCode(FOutputDevice& O
 	{
 		if (FUnrealScriptStructDefinitionInfo* SparseScriptStructDef = GTypeDefinitionInfoMap.FindByName<FUnrealScriptStructDefinitionInfo>(*SparseClassDataString))
 		{
-			UScriptStruct* SparseClassDataStruct = SparseScriptStructDef->GetScriptStruct();
-			GeneratedClassRegisterFunctionText.Logf(TEXT("%s\r\n"), *SparseClassDataStruct->GetName());
+			GeneratedClassRegisterFunctionText.Logf(TEXT("%s\r\n"), *SparseScriptStructDef->GetName());
 			for (FUnrealPropertyDefinitionInfo* ChildDef : TUHTFieldRange<FUnrealPropertyDefinitionInfo>(*SparseScriptStructDef))
 			{
 				FProperty* Child = ChildDef->GetProperty();
@@ -3437,17 +3435,15 @@ void FNativeClassHeaderGenerator::ExportEnum(FOutputDevice& Out, FUnrealEnumDefi
 // Exports the header text for the list of structs specified (GENERATED_BODY impls)
 void FNativeClassHeaderGenerator::ExportGeneratedStructBodyMacros(FOutputDevice& OutGeneratedHeaderText, FOutputDevice& Out, FReferenceGatherers& OutReferenceGatherers, const FUnrealSourceFile& SourceFile, FUnrealScriptStructDefinitionInfo& ScriptStructDef) const
 {
-	UScriptStruct* Struct = ScriptStructDef.GetScriptStruct();
 	ScriptStructDef.AddCrossModuleReference(OutReferenceGatherers.UniqueCrossModuleReferences, true);
 
 	const bool bIsDynamic = ScriptStructDef.IsDynamic();
-	const FString ActualStructName = FNativeClassHeaderGenerator::GetOverriddenName(Struct);
+	const FString ActualStructName = FNativeClassHeaderGenerator::GetOverriddenName(&ScriptStructDef);
 	const FString& FriendApiString  = GetAPIString();
 
-	FUnrealStructDefinitionInfo* BaseStructDef = ScriptStructDef.GetSuperStructInfo().Struct;
-	UStruct* BaseStruct = Struct->GetSuperStruct();
+	FUnrealStructDefinitionInfo* BaseStructDef = ScriptStructDef.GetSuperStruct();
 
-	const FString StructNameCPP = FNameLookupCPP::GetNameCPP(Struct);
+	const FString StructNameCPP = ScriptStructDef.GetAlternateNameCPP();
 
 	const FString& SingletonName = ScriptStructDef.GetSingletonName(true);
 	const FString& ChoppedSingletonName = ScriptStructDef.GetSingletonNameChopped(true);
@@ -3584,11 +3580,11 @@ void FNativeClassHeaderGenerator::ExportGeneratedStructBodyMacros(FOutputDevice&
 	}
 
 	// Export struct.
-	if (Struct->StructFlags & STRUCT_Native)
+	if (ScriptStructDef.HasAnyStructFlags(STRUCT_Native))
 	{
-		check(Struct->StructMacroDeclaredLineNumber != INDEX_NONE);
+		check(ScriptStructDef.GetMacroDeclaredLineNumber() != INDEX_NONE);
 
-		const bool bRequiredAPI = !(Struct->StructFlags & STRUCT_RequiredAPI);
+		const bool bRequiredAPI = !ScriptStructDef.HasAnyStructFlags(STRUCT_RequiredAPI);
 
 		const FString FriendLine = FString::Printf(TEXT("\tfriend struct %s_Statics;\r\n"), *ChoppedSingletonName);
 		const FString StaticClassLine = FString::Printf(TEXT("\t%sstatic class UScriptStruct* StaticStruct();\r\n"), (bRequiredAPI ? *FriendApiString : TEXT("")));
@@ -3646,18 +3642,18 @@ void FNativeClassHeaderGenerator::ExportGeneratedStructBodyMacros(FOutputDevice&
 			}
 		}
 
-		const FString SuperTypedef = BaseStruct ? FString::Printf(TEXT("\ttypedef %s Super;\r\n"), *FNameLookupCPP::GetNameCPP(BaseStruct)) : FString();
+		const FString SuperTypedef = BaseStructDef ? FString::Printf(TEXT("\ttypedef %s Super;\r\n"), *BaseStructDef->GetAlternateNameCPP()) : FString();
 
 		FString CombinedLine = FString::Printf(TEXT("%s%s%s%s"), *FriendLine, *StaticClassLine, *RigVMMethodsDeclarations, *SuperTypedef);
-		const FString MacroName = SourceFile.GetGeneratedBodyMacroName(Struct->StructMacroDeclaredLineNumber);
+		const FString MacroName = SourceFile.GetGeneratedBodyMacroName(ScriptStructDef.GetMacroDeclaredLineNumber());
 
 		const FString Macroized = Macroize(*MacroName, MoveTemp(CombinedLine));
 		OutGeneratedHeaderText.Log(Macroized);
 
 		// Inject static assert to verify that we do not add vtable
-		if (BaseStruct)
+		if (BaseStructDef)
 		{
-			FString BaseStructNameCPP = *FNameLookupCPP::GetNameCPP(BaseStruct);
+			FString BaseStructNameCPP = BaseStructDef->GetAlternateNameCPP();
 
 			FString VerifyPolymorphicStructString = FString::Printf(TEXT("\r\nstatic_assert(std::is_polymorphic<%s>() == std::is_polymorphic<%s>(), \"USTRUCT %s cannot be polymorphic unless super %s is polymorphic\");\r\n\r\n"), *StructNameCPP, *BaseStructNameCPP, *StructNameCPP, *BaseStructNameCPP);			
 			Out.Log(VerifyPolymorphicStructString);
@@ -3667,7 +3663,7 @@ void FNativeClassHeaderGenerator::ExportGeneratedStructBodyMacros(FOutputDevice&
 
 		if (!bIsDynamic)
 		{
-			Out.Logf(TEXT("\tstatic FStructRegistrationInfo& Z_Registration_Info_UScriptStruct_%s()\r\n"), *Struct->GetName());
+			Out.Logf(TEXT("\tstatic FStructRegistrationInfo& Z_Registration_Info_UScriptStruct_%s()\r\n"), *ScriptStructDef.GetName());
 			Out.Logf(TEXT("\t{\r\n"));
 			Out.Logf(TEXT("\t\tstatic FStructRegistrationInfo info;\r\n"));
 			Out.Logf(TEXT("\t\treturn info;\r\n"));
@@ -3683,7 +3679,7 @@ void FNativeClassHeaderGenerator::ExportGeneratedStructBodyMacros(FOutputDevice&
 		const FString& OuterName (bIsDynamic ? STRING_StructPackage : GetPackageSingletonName(UHTCastChecked<FUnrealPackageDefinitionInfo>(ScriptStructDef.GetOuter()), OutReferenceGatherers.UniqueCrossModuleReferences));
 		if (!bIsDynamic)
 		{
-			Out.Logf(TEXT("\tstatic class UScriptStruct*& Singleton = Z_Registration_Info_UScriptStruct_%s().OuterSingleton;\r\n"), *Struct->GetName());
+			Out.Logf(TEXT("\tstatic class UScriptStruct*& Singleton = Z_Registration_Info_UScriptStruct_%s().OuterSingleton;\r\n"), *ScriptStructDef.GetName());
 		}
 		else
 		{
@@ -3738,19 +3734,19 @@ void FNativeClassHeaderGenerator::ExportGeneratedStructBodyMacros(FOutputDevice&
 			Out.Logf(TEXT("static FRegisterCompiledInInfo Z_CompiledInDeferStruct_UScriptStruct_%s(%s::StaticStruct, TEXT(\"%s\"), TEXT(\"%s\"), Z_Registration_Info_UScriptStruct_%s(), CONSTUCT_RELOAD_VERSION_INFO(FStructReloadVersionInfo, sizeof(%s), %s()));\r\n"),
 				*StructNameCPP,
 				*StructNameCPP,
-				*Struct->GetOutermost()->GetName(),
+				*ScriptStructDef.GetPackageDef().GetName(),
 				*ActualStructName,
-				*Struct->GetName(),
+				*ScriptStructDef.GetName(),
 				*StructNameCPP, 
 				*GetHashName
 			);
 		}
 
 		// Generate StaticRegisterNatives equivalent for structs without classes.
-		if (!Struct->GetOuter()->IsA(UStruct::StaticClass()))
+		if (!ScriptStructDef.GetOuter()->GetObject()->IsA(UStruct::StaticClass()))
 		{
 			// Do not change the Z_CompiledInDeferCppStructOps_UScriptStruct_ without changing LC_SymbolPatterns
-			const FString ShortPackageName = FPackageName::GetShortName(Struct->GetOuter()->GetName());
+			const FString ShortPackageName = FPackageName::GetShortName(ScriptStructDef.GetOuter()->GetName());
 			Out.Logf(TEXT("static struct FScriptStruct_%s_StaticRegisterNatives%s\r\n"), *ShortPackageName, *StructNameCPP);
 			Out.Logf(TEXT("{\r\n"));
 			Out.Logf(TEXT("\tFScriptStruct_%s_StaticRegisterNatives%s()\r\n"), *ShortPackageName, *StructNameCPP);
@@ -3776,13 +3772,13 @@ void FNativeClassHeaderGenerator::ExportGeneratedStructBodyMacros(FOutputDevice&
 		ExportMirrorsForNoexportStruct(GeneratedStructRegisterFunctionText, *NoExportStructDef, /*Indent=*/ 2);
 	}
 
-	if (BaseStruct)
+	if (BaseStructDef)
 	{
-		CastChecked<UScriptStruct>(BaseStruct); // this better actually be a script struct
+		UHTCastChecked<FUnrealScriptStructDefinitionInfo>(BaseStructDef); // this better actually be a script struct
 		BaseStructDef->AddCrossModuleReference(OutReferenceGatherers.UniqueCrossModuleReferences, true);
 	}
 
-	EStructFlags UncomputedFlags = (EStructFlags)(Struct->StructFlags & ~STRUCT_ComputedFlags);
+	EStructFlags UncomputedFlags = (EStructFlags)(ScriptStructDef.GetStructFlags() & ~STRUCT_ComputedFlags);
 
 	FString OuterFunc;
 	if (!bIsDynamic)
@@ -3804,7 +3800,7 @@ void FNativeClassHeaderGenerator::ExportGeneratedStructBodyMacros(FOutputDevice&
 	FString MetaDataParams = OutputMetaDataCodeForObject(GeneratedStructRegisterFunctionText, StaticDefinitions, ScriptStructDef, *FString::Printf(TEXT("%s::Struct_MetaDataParams"), *StaticsStructName), TEXT("\t\t"), TEXT("\t"));
 
 	FString NewStructOps;
-	if (Struct->StructFlags & STRUCT_Native)
+	if (ScriptStructDef.HasAnyStructFlags(STRUCT_Native))
 	{
 		GeneratedStructRegisterFunctionText.Log(TEXT("\t\tstatic void* NewStructOps();\r\n"));
 
@@ -3859,9 +3855,9 @@ void FNativeClassHeaderGenerator::ExportGeneratedStructBodyMacros(FOutputDevice&
 	// Structs can either have a UClass or UPackage as outer (if declared in non-UClass header).
 	if (!bIsDynamic)
 	{
-		if (Struct->StructFlags & STRUCT_Native)
+		if (ScriptStructDef.HasAnyStructFlags(STRUCT_Native))
 		{
-			GeneratedStructRegisterFunctionText.Logf(TEXT("\t\tstatic UScriptStruct*& ReturnStruct = Z_Registration_Info_UScriptStruct_%s().InnerSingleton;\r\n"), *Struct->GetName());
+			GeneratedStructRegisterFunctionText.Logf(TEXT("\t\tstatic UScriptStruct*& ReturnStruct = Z_Registration_Info_UScriptStruct_%s().InnerSingleton;\r\n"), *ScriptStructDef.GetName());
 		}
 		else
 		{
@@ -4150,14 +4146,12 @@ void FNativeClassHeaderGenerator::ExportGeneratedEnumInitCode(FOutputDevice& Out
 
 void FNativeClassHeaderGenerator::ExportMirrorsForNoexportStruct(FOutputDevice& Out, FUnrealScriptStructDefinitionInfo& ScriptStructDef, int32 TextIndent)
 {
-	UScriptStruct* ScriptStruct = ScriptStructDef.GetScriptStruct();
-
 	// Export struct.
-	const FString StructName = FNameLookupCPP::GetNameCPP(ScriptStruct);
+	const FString StructName = ScriptStructDef.GetAlternateNameCPP();
 	Out.Logf(TEXT("%sstruct %s"), FCString::Tab(TextIndent), *StructName);
-	if (ScriptStruct->GetSuperStruct() != NULL)
+	if (FUnrealStructDefinitionInfo* SuperStructDef = ScriptStructDef.GetSuperStruct())
 	{
-		Out.Logf(TEXT(" : public %s"), *FNameLookupCPP::GetNameCPP(ScriptStruct->GetSuperStruct()));
+		Out.Logf(TEXT(" : public %s"), *SuperStructDef->GetAlternateNameCPP());
 	}
 	Out.Logf(TEXT("\r\n%s{\r\n"), FCString::Tab(TextIndent));
 
@@ -5536,8 +5530,7 @@ static void RecordPackageSingletons(
 	Singletons.Reserve(StructDefs.Num() + DelegateDefs.Num());
 	for (FUnrealScriptStructDefinitionInfo* StructDef : StructDefs)
 	{
-		UScriptStruct* Struct = StructDef->GetScriptStruct();
-		if (Struct->StructFlags & STRUCT_NoExport && !HasDynamicOuter(*StructDef))
+		if (StructDef->HasAnyStructFlags(STRUCT_NoExport) && !HasDynamicOuter(*StructDef))
 		{
 			Singletons.Add(StructDef);
 		}
diff --git a/Engine/Source/Programs/UnrealHeaderTool/Private/HeaderParser.cpp b/Engine/Source/Programs/UnrealHeaderTool/Private/HeaderParser.cpp
--- a/Engine/Source/Programs/UnrealHeaderTool/Private/HeaderParser.cpp
+++ b/Engine/Source/Programs/UnrealHeaderTool/Private/HeaderParser.cpp
@@ -2225,12 +2225,12 @@ FUnrealScriptStructDefinitionInfo& FHeaderParser::CompileStructDeclaration()
 				FError::Throwf(TEXT("%s must be in the public scope of '%s', not private or protected."), Token.Identifier, *StructNameInScript);
 			}
 
-			if (Struct->StructMacroDeclaredLineNumber != INDEX_NONE)
+			if (StructDef.GetMacroDeclaredLineNumber() != INDEX_NONE)
 			{
 				FError::Throwf(TEXT("Multiple %s declarations found in '%s'"), Token.Identifier, *StructNameInScript);
 			}
 
-			Struct->StructMacroDeclaredLineNumber = InputLine;
+			StructDef.SetMacroDeclaredLineNumber(InputLine);
 			RequireSymbol(TEXT('('), TEXT("'struct'"));
 
 			CompileVersionDeclaration(StructDef);
@@ -2352,7 +2352,7 @@ FUnrealScriptStructDefinitionInfo& FHeaderParser::CompileStructDeclaration()
 	}
 
 	// Validation
-	bool bStructBodyFound = Struct->StructMacroDeclaredLineNumber != INDEX_NONE;
+	bool bStructBodyFound = StructDef.GetMacroDeclaredLineNumber() != INDEX_NONE;
 	bool bExported        = !!(StructFlags & STRUCT_Native);
 	if (!bStructBodyFound && bExported)
 	{
@@ -2559,8 +2559,6 @@ void FHeaderParser::PopNest(ENestType NestType, const TCHAR* Descr)
 
 void FHeaderParser::FixupDelegateProperties(FUnrealStructDefinitionInfo& StructDef, FScope& Scope, TMap<FName, FUnrealFunctionDefinitionInfo*>& DelegateCache )
 {
-	UStruct* Struct = StructDef.GetStruct();
-
 	for (FUnrealPropertyDefinitionInfo* PropertyDef : StructDef.GetProperties())
 	{
 		FProperty* Property = PropertyDef->GetProperty();
@@ -2730,15 +2728,14 @@ void FHeaderParser::FixupDelegateProperties(FUnrealStructDefinitionInfo& StructD
 void FHeaderParser::CheckSparseClassData(const FUnrealStructDefinitionInfo& StructDef)
 {
 	// we're looking for classes that have sparse class data structures
-	const UStruct* StructToCheck = StructDef.GetStruct();
 	const FUnrealClassDefinitionInfo* ClassDef = UHTCast<FUnrealClassDefinitionInfo>(StructDef);
 
 	// make sure we don't try to have sparse class data inside of a struct instead of a class
-	if (StructToCheck->HasMetaData(FHeaderParserNames::NAME_SparseClassDataTypes))
+	if (StructDef.HasMetaData(FHeaderParserNames::NAME_SparseClassDataTypes))
 	{
 		if (ClassDef == nullptr)
 		{
-			FError::Throwf(TEXT("%s contains sparse class data but is not a class."), *StructToCheck->GetName());
+			FError::Throwf(TEXT("%s contains sparse class data but is not a class."), *StructDef.GetName());
 		}
 	}
 	else
@@ -2773,27 +2770,25 @@ void FHeaderParser::CheckSparseClassData(const FUnrealStructDefinitionInfo& Stru
 			return;
 		}
 
-		UScriptStruct* SparseClassDataStruct = SparseClassDataStructDef->GetScriptStruct();
-
 		// check the data struct for invalid properties
 		for (FUnrealPropertyDefinitionInfo* PropertyDef : TUHTFieldRange<FUnrealPropertyDefinitionInfo>(*SparseClassDataStructDef))
 		{
 			FProperty* Property = PropertyDef->GetProperty();
 			if (Property->HasAnyPropertyFlags(CPF_BlueprintAssignable))
 			{
-				FError::Throwf(TEXT("Sparse class data types can not contain blueprint assignable delegates. Type '%s' Delegate '%s'"), *SparseClassDataStruct->GetName(), *Property->GetName());
+				FError::Throwf(TEXT("Sparse class data types can not contain blueprint assignable delegates. Type '%s' Delegate '%s'"), *SparseClassDataStructDef->GetName(), *Property->GetName());
 			}
 
 			// all sparse properties should have EditDefaultsOnly
 			if (!Property->HasAllPropertyFlags(CPF_Edit | CPF_DisableEditOnInstance))
 			{
-				FError::Throwf(TEXT("Sparse class data types must be EditDefaultsOnly. Type '%s' Property '%s'"), *SparseClassDataStruct->GetName(), *Property->GetName());
+				FError::Throwf(TEXT("Sparse class data types must be EditDefaultsOnly. Type '%s' Property '%s'"), *SparseClassDataStructDef->GetName(), *Property->GetName());
 			}
 
 			// no sparse properties should have BlueprintReadWrite
 			if (Property->HasAllPropertyFlags(CPF_BlueprintVisible) && !Property->HasAllPropertyFlags(CPF_BlueprintReadOnly))
 			{
-				FError::Throwf(TEXT("Sparse class data types must not be BlueprintReadWrite. Type '%s' Property '%s'"), *SparseClassDataStruct->GetName(), *Property->GetName());
+				FError::Throwf(TEXT("Sparse class data types must not be BlueprintReadWrite. Type '%s' Property '%s'"), *SparseClassDataStructDef->GetName(), *Property->GetName());
 			}
 		}
 
@@ -2806,9 +2801,9 @@ void FHeaderParser::CheckSparseClassData(const FUnrealStructDefinitionInfo& Stru
 			if (FUnrealScriptStructDefinitionInfo* ParentSparseClassDataStructDef = GTypeDefinitionInfoMap.FindByName<FUnrealScriptStructDefinitionInfo>(*ParentSparseClassDataTypeName))
 			{
 				UScriptStruct* ParentSparseClassDataStruct = ParentSparseClassDataStructDef->GetScriptStruct();
-				if (!SparseClassDataStruct->IsChildOf(ParentSparseClassDataStruct))
+				if (!SparseClassDataStructDef->GetScriptStruct()->IsChildOf(ParentSparseClassDataStruct))
 				{
-					FError::Throwf(TEXT("Class %s is a child of %s but its sparse class data struct, %s, does not inherit from %s."), *ClassToCheck->GetName(), *ParentClass->GetName(), *SparseClassDataStruct->GetName(), *ParentSparseClassDataStruct->GetName());
+					FError::Throwf(TEXT("Class %s is a child of %s but its sparse class data struct, %s, does not inherit from %s."), *ClassToCheck->GetName(), *ParentClass->GetName(), *SparseClassDataStructDef->GetName(), *ParentSparseClassDataStruct->GetName());
 				}
 			}
 		}
@@ -4297,7 +4292,7 @@ void FHeaderParser::GetVarType(
 			if (bStripped)
 			{
 				const TCHAR* PrefixCPP = UHTConfig.StructsWithTPrefix.Contains(IdentifierStripped) ? TEXT("T") : Struct->GetPrefixCPP();
-				FString ExpectedStructName = FString::Printf(TEXT("%s%s"), PrefixCPP, *Struct->GetName());
+				FString ExpectedStructName = FString::Printf(TEXT("%s%s"), PrefixCPP, *StructDef->GetName());
 				if (FString(VarType.Identifier) != ExpectedStructName)
 				{
 					FError::Throwf(TEXT("Struct '%s' is missing or has an incorrect prefix, expecting '%s'"), VarType.Identifier, *ExpectedStructName);
@@ -4306,13 +4301,13 @@ void FHeaderParser::GetVarType(
 			else if (!UHTConfig.StructsWithNoPrefix.Contains(VarType.Identifier))
 			{
 				const TCHAR* PrefixCPP = UHTConfig.StructsWithTPrefix.Contains(VarType.Identifier) ? TEXT("T") : Struct->GetPrefixCPP();
-				FError::Throwf(TEXT("Struct '%s' is missing a prefix, expecting '%s'"), VarType.Identifier, *FString::Printf(TEXT("%s%s"), PrefixCPP, *Struct->GetName()));
+				FError::Throwf(TEXT("Struct '%s' is missing a prefix, expecting '%s'"), VarType.Identifier, *FString::Printf(TEXT("%s%s"), PrefixCPP, *StructDef->GetName()));
 			}
 
 			bHandledType = true;
 
 			VarProperty = FPropertyBase(*StructDef);
-			if ((Struct->StructFlags & STRUCT_HasInstancedReference) && !(Disallow & CPF_ContainsInstancedReference))
+			if (StructDef->HasAnyStructFlags(STRUCT_HasInstancedReference) && !(Disallow & CPF_ContainsInstancedReference))
 			{
 				Flags |= CPF_ContainsInstancedReference;
 			}
@@ -4745,10 +4740,6 @@ FUnrealPropertyDefinitionInfo& FHeaderParser::GetVarNameAndDim
 	ELayoutMacroType        LayoutMacroType
 )
 {
-	UStruct* Scope = ParentStruct.GetStruct();
-
-	check(Scope);
-
 	EObjectFlags ObjectFlags = RF_Public;
 	if (VariableCategory == EVariableCategory::Member && CurrentAccessSpecifier == ACCESS_Private)
 	{
@@ -4869,7 +4860,7 @@ FUnrealPropertyDefinitionInfo& FHeaderParser::GetVarNameAndDim
 			FError::Throwf(TEXT("%s: '%s' cannot be defined in '%s' as it is already defined in scope '%s' (shadowing is not allowed)"),
 				HintText,
 				Identifier,
-				*Scope->GetName(),
+				*ParentStruct.GetName(),
 				ExistingField ? *ExistingField->GetOuter()->GetName() : *ExistingProperty->GetOwnerVariant().GetFullName());
 		}
 	}
@@ -5027,7 +5018,7 @@ FUnrealPropertyDefinitionInfo& FHeaderParser::GetVarNameAndDim
 	ParentStruct.AddProperty(PropDef);
 
 	// Add to the end of the properties slist
-	FField** Prev = &Scope->ChildProperties;
+	FField** Prev = &ParentStruct.GetStruct()->ChildProperties;
 	for (; *Prev != nullptr; Prev = &(*Prev)->Next)
 	{
 		// No body
@@ -6128,7 +6119,7 @@ void FHeaderParser::CompileRigVMMethodDeclaration(FUnrealStructDefinitionInfo& S
 
 	FRigVMStructInfo& StructRigVMInfo = StructDef.GetRigVMInfo();
 	StructRigVMInfo.bHasRigVM = true;
-	StructRigVMInfo.Name = StructDef.GetStruct()->GetName();
+	StructRigVMInfo.Name = StructDef.GetName();
 	StructRigVMInfo.Methods.Add(MethodInfo);
 }
 
@@ -6149,7 +6140,6 @@ const TCHAR* FHeaderParser::GetDynamicArrayText = TEXT("GetDynamicArray");
 
 void FHeaderParser::ParseRigVMMethodParameters(FUnrealStructDefinitionInfo& StructDef)
 {
-	UStruct* Struct = StructDef.GetStruct();
 	FRigVMStructInfo& StructRigVMInfo = StructDef.GetRigVMInfo();
 	if (!StructRigVMInfo.bHasRigVM)
 	{
@@ -6198,7 +6188,7 @@ void FHeaderParser::ParseRigVMMethodParameters(FUnrealStructDefinitionInfo& Stru
 
 		if (Parameter.bEditorOnly)
 		{
-			UE_LOG_ERROR_UHT(TEXT("RigVM Struct '%s' - Member '%s' is editor only - WITH_EDITORONLY_DATA not allowed on structs with RIGVM_METHOD."), *Struct->GetName(), *Parameter.Name, *MemberCPPType);
+			UE_LOG_ERROR_UHT(TEXT("RigVM Struct '%s' - Member '%s' is editor only - WITH_EDITORONLY_DATA not allowed on structs with RIGVM_METHOD."), *StructDef.GetName(), *Parameter.Name, *MemberCPPType);
 		}
 
 		if (!ExtendedCPPType.IsEmpty())
@@ -6206,7 +6196,7 @@ void FHeaderParser::ParseRigVMMethodParameters(FUnrealStructDefinitionInfo& Stru
 			// we only support arrays - no maps or similar data structures
 			if (MemberCPPType != TArrayText && MemberCPPType != TEnumAsByteText)
 			{
-				UE_LOG_ERROR_UHT(TEXT("RigVM Struct '%s' - Member '%s' type '%s' not supported by RigVM."), *Struct->GetName(), *Parameter.Name, *MemberCPPType);
+				UE_LOG_ERROR_UHT(TEXT("RigVM Struct '%s' - Member '%s' type '%s' not supported by RigVM."), *StructDef.GetName(), *Parameter.Name, *MemberCPPType);
 				continue;
 			}
 		}
@@ -6232,12 +6222,12 @@ void FHeaderParser::ParseRigVMMethodParameters(FUnrealStructDefinitionInfo& Stru
 
 	if (StructRigVMInfo.Members.Num() == 0)
 	{
-		UE_LOG_ERROR_UHT(TEXT("RigVM Struct '%s' - has zero members - invalid RIGVM_METHOD."), *Struct->GetName());
+		UE_LOG_ERROR_UHT(TEXT("RigVM Struct '%s' - has zero members - invalid RIGVM_METHOD."), *StructDef.GetName());
 	}
 
 	if (StructRigVMInfo.Members.Num() > 64)
 	{
-		UE_LOG_ERROR_UHT(TEXT("RigVM Struct '%s' - has %d members (64 is the limit)."), *Struct->GetName(), StructRigVMInfo.Members.Num());
+		UE_LOG_ERROR_UHT(TEXT("RigVM Struct '%s' - has %d members (64 is the limit)."), *StructDef.GetName(), StructRigVMInfo.Members.Num());
 	}
 }
 
@@ -6302,7 +6292,7 @@ void FHeaderParser::ParseParameterList(FUnrealFunctionDefinitionInfo& FunctionDe
 			{
 				if (!FunctionDef.HasAnyFunctionFlags(FUNC_NetRequest | FUNC_NetResponse))
 				{
-					ValidateScriptStructOkForNet(Property.ScriptStructDef->GetStruct()->GetName(), *Property.ScriptStructDef);
+					ValidateScriptStructOkForNet(Property.ScriptStructDef->GetName(), *Property.ScriptStructDef);
 				}
 			}
 
@@ -6465,7 +6455,7 @@ FUnrealFunctionDefinitionInfo& FHeaderParser::CompileDelegateDeclaration(const T
 
 		//Workaround for UE-28897
 		const FStructScope* CurrentStructScope = TopNest->GetScope() ? TopNest->GetScope()->AsStructScope() : nullptr;
-		const bool bDynamicClassScope = CurrentStructScope && CurrentStructScope->GetStruct() && CurrentStructScope->GetStructDef().IsDynamic();
+		const bool bDynamicClassScope = CurrentStructScope && CurrentStructScope->GetStructDef().IsDynamic();
 		CheckAllow(CurrentScopeName, bDynamicClassScope ? ENestAllowFlags::ImplicitDelegateDecl : ENestAllowFlags::TypeDecl);
 	}
 	else
@@ -7405,9 +7395,9 @@ void FHeaderParser::CompileVariableDeclaration(FUnrealStructDefinitionInfo& Stru
 	// First check if the category was specified at all and if the property was exposed to the editor.
 	if (!Category && (OriginalProperty.PropertyFlags & (CPF_Edit|CPF_BlueprintVisible)))
 	{
-		if ((Struct->GetOutermost() != nullptr) && !bIsCurrentModulePartOfEngine)
+		if ((&StructDef.GetPackageDef() != nullptr) && !bIsCurrentModulePartOfEngine)
 		{
-			Category = &OriginalProperty.MetaData.Add(NAME_Category, Struct->GetName());
+			Category = &OriginalProperty.MetaData.Add(NAME_Category, StructDef.GetName());
 		}
 		else
 		{
@@ -7465,7 +7455,7 @@ void FHeaderParser::CompileVariableDeclaration(FUnrealStructDefinitionInfo& Stru
 	// If Property is a Replicated Struct check to make sure there are no Properties that are not allowed to be Replicated in the Struct 
 	if (OriginalProperty.Type == CPT_Struct && OriginalProperty.PropertyFlags & CPF_Net && OriginalProperty.ScriptStructDef)
 	{
-		ValidateScriptStructOkForNet(OriginalProperty.ScriptStructDef->GetScriptStruct()->GetName(), *OriginalProperty.ScriptStructDef);
+		ValidateScriptStructOkForNet(OriginalProperty.ScriptStructDef->GetName(), *OriginalProperty.ScriptStructDef);
 	}
 
 	// Process all variables of this type.
@@ -7496,7 +7486,7 @@ void FHeaderParser::CompileVariableDeclaration(FUnrealStructDefinitionInfo& Stru
 		{
 			if (NewProperties.Num())
 			{
-				FError::Throwf(TEXT("Comma delimited properties cannot be converted %s.%s\n"), *Struct->GetName(), *NewProperty->GetName());
+				FError::Throwf(TEXT("Comma delimited properties cannot be converted %s.%s\n"), *StructDef.GetName(), *NewProperty->GetName());
 			}
 		}
 
@@ -7519,21 +7509,21 @@ void FHeaderParser::CompileVariableDeclaration(FUnrealStructDefinitionInfo& Stru
 
 		if (NewProperty->HasAnyPropertyFlags(CPF_BlueprintVisible))
 		{
-			if (Struct->IsA<UScriptStruct>() && !Struct->GetBoolMetaDataHierarchical(FHeaderParserNames::NAME_BlueprintType))
+			if (Struct->IsA<UScriptStruct>() && !StructDef.GetBoolMetaDataHierarchical(FHeaderParserNames::NAME_BlueprintType))
 			{
-				UE_LOG_ERROR_UHT(TEXT("Cannot expose property to blueprints in a struct that is not a BlueprintType. %s.%s"), *Struct->GetName(), *NewProperty->GetName());
+				UE_LOG_ERROR_UHT(TEXT("Cannot expose property to blueprints in a struct that is not a BlueprintType. %s.%s"), *StructDef.GetName(), *NewProperty->GetName());
 			}
 
 			if (NewProperty->ArrayDim > 1)
 			{
-				UE_LOG_ERROR_UHT(TEXT("Static array cannot be exposed to blueprint %s.%s"), *Struct->GetName(), *NewProperty->GetName());
+				UE_LOG_ERROR_UHT(TEXT("Static array cannot be exposed to blueprint %s.%s"), *StructDef.GetName(), *NewProperty->GetName());
 			}
 
 			if (!FPropertyTraits::IsSupportedByBlueprint(NewPropDef, true))
 			{
 				FString ExtendedCPPType;
 				FString CPPType = NewProperty->GetCPPType(&ExtendedCPPType);
-				UE_LOG_ERROR_UHT(TEXT("Type '%s%s' is not supported by blueprint. %s.%s"), *CPPType, *ExtendedCPPType, *Struct->GetName(), *NewProperty->GetName());
+				UE_LOG_ERROR_UHT(TEXT("Type '%s%s' is not supported by blueprint. %s.%s"), *CPPType, *ExtendedCPPType, *StructDef.GetName(), *NewProperty->GetName());
 			}
 		}
 
@@ -9212,7 +9202,7 @@ FDocumentationPolicy FHeaderParser::GetDocumentationPolicyForStruct(FUnrealStruc
 
 	FDocumentationPolicy DocumentationPolicy;
 	FString DocumentationPolicyName;
-	if (StructDef.GetStruct()->GetStringMetaDataHierarchical(NAME_DocumentationPolicy, &DocumentationPolicyName))
+	if (StructDef.GetStringMetaDataHierarchical(NAME_DocumentationPolicy, &DocumentationPolicyName))
 	{
 		DocumentationPolicy = GetDocumentationPolicyFromName(DocumentationPolicyName);
 	}
@@ -9270,21 +9260,18 @@ void FHeaderParser::CheckDocumentationPolicyForStruct(FUnrealStructDefinitionInf
 {
 	SCOPE_SECONDS_COUNTER_UHT(DocumentationPolicy);
 
-	UStruct* Struct = StructDef.GetStruct();
-	check(Struct != nullptr);
-
 	FDocumentationPolicy DocumentationPolicy = GetDocumentationPolicyForStruct(StructDef);
 	if (DocumentationPolicy.bClassOrStructCommentRequired)
 	{
 		FString ClassTooltip;
-		if (const FString* ClassTooltipPtr = Struct->FindMetaData(NAME_ToolTip))
+		if (const FString* ClassTooltipPtr = StructDef.FindMetaData(NAME_ToolTip))
 		{
 			ClassTooltip = *ClassTooltipPtr;
 		}
 
-		if (ClassTooltip.IsEmpty() || ClassTooltip.Equals(Struct->GetName()))
+		if (ClassTooltip.IsEmpty() || ClassTooltip.Equals(StructDef.GetName()))
 		{
-			UE_LOG_ERROR_UHT(TEXT("Struct '%s' does not provide a tooltip / comment (DocumentationPolicy)."), *Struct->GetName());
+			UE_LOG_ERROR_UHT(TEXT("Struct '%s' does not provide a tooltip / comment (DocumentationPolicy)."), *StructDef.GetName());
 		}
 	}
 
@@ -9297,13 +9284,13 @@ void FHeaderParser::CheckDocumentationPolicyForStruct(FUnrealStructDefinitionInf
 			FString ToolTip = Property->GetToolTipText().ToString();
 			if (ToolTip.IsEmpty() || ToolTip.Equals(Property->GetDisplayNameText().ToString()))
 			{
-				UE_LOG_ERROR_UHT(TEXT("Property '%s::%s' does not provide a tooltip / comment (DocumentationPolicy)."), *Struct->GetName(), *Property->GetName());
+				UE_LOG_ERROR_UHT(TEXT("Property '%s::%s' does not provide a tooltip / comment (DocumentationPolicy)."), *StructDef.GetName(), *Property->GetName());
 				continue;
 			}
 			const FName* ExistingPropertyName = ToolTipToPropertyName.Find(ToolTip);
 			if (ExistingPropertyName != nullptr)
 			{
-				UE_LOG_ERROR_UHT(TEXT("Property '%s::%s' and '%s::%s' are using identical tooltips (DocumentationPolicy)."), *Struct->GetName(), *ExistingPropertyName->ToString(), *Struct->GetName(), *Property->GetName());
+				UE_LOG_ERROR_UHT(TEXT("Property '%s::%s' and '%s::%s' are using identical tooltips (DocumentationPolicy)."), *StructDef.GetName(), *ExistingPropertyName->ToString(), *StructDef.GetName(), *Property->GetName());
 			}
 			ToolTipToPropertyName.Add(MoveTemp(ToolTip), Property->GetFName());
 		}
@@ -9321,7 +9308,7 @@ void FHeaderParser::CheckDocumentationPolicyForStruct(FUnrealStructDefinitionInf
 
 				if(!CheckUIMinMaxRangeFromMetaData(UIMin, UIMax))
 				{
-					UE_LOG_ERROR_UHT(TEXT("Property '%s::%s' does not provide a valid UIMin / UIMax (DocumentationPolicy)."), *Struct->GetName(), *Property->GetName());
+					UE_LOG_ERROR_UHT(TEXT("Property '%s::%s' does not provide a valid UIMin / UIMax (DocumentationPolicy)."), *StructDef.GetName(), *Property->GetName());
 				}
 			}
 		}
diff --git a/Engine/Source/Programs/UnrealHeaderTool/Private/UnrealTypeDefinitionInfo.h b/Engine/Source/Programs/UnrealHeaderTool/Private/UnrealTypeDefinitionInfo.h
--- a/Engine/Source/Programs/UnrealHeaderTool/Private/UnrealTypeDefinitionInfo.h
+++ b/Engine/Source/Programs/UnrealHeaderTool/Private/UnrealTypeDefinitionInfo.h
@@ -1284,6 +1284,29 @@ public:
 		return BaseStructInfo;
 	}
 
+	/** Try and find boolean metadata with the given key. If not found on this class, work up hierarchy looking for it. */
+	bool GetBoolMetaDataHierarchical(const FName& Key) const
+	{
+		return GetStruct()->GetBoolMetaDataHierarchical(Key);
+	}
+
+	/** Try and find string metadata with the given key. If not found on this class, work up hierarchy looking for it. */
+	bool GetStringMetaDataHierarchical(const FName& Key, FString* OutValue = nullptr) const
+	{
+		return GetStruct()->GetStringMetaDataHierarchical(Key, OutValue);
+	}
+
+	/**
+	 * Determines if the struct or any of its super structs has any metadata associated with the provided key
+	 *
+	 * @param Key The key to lookup in the metadata
+	 * @return pointer to the UStruct that has associated metadata, nullptr if Key is not associated with any UStruct in the hierarchy
+	 */
+	const UStruct* HasMetaDataHierarchical(const FName& Key) const
+	{
+		return GetStruct()->HasMetaDataHierarchical(Key);
+	}
+
 	/**
 	 * Get the contains delegates flag
 	 */
@@ -1419,6 +1442,74 @@ public:
 	{
 		return static_cast<UScriptStruct*>(GetObject());
 	}
+
+	/**
+	 * Return the struct flags
+	 */
+	FORCEINLINE EStructFlags GetStructFlags() const
+	{
+		return GetScriptStruct()->StructFlags;
+	}
+
+	/**
+	 * Used to safely check whether the passed in flag is set.
+	 *
+	 * @param	FlagToCheck		Class flag to check for
+	 *
+	 * @return	true if the passed in flag is set, false otherwise
+	 *			(including no flag passed in, unless the FlagsToCheck is CLASS_AllFlags)
+	 */
+	FORCEINLINE bool HasAnyStructFlags(EStructFlags FlagsToCheck) const
+	{
+		return EnumHasAnyFlags(GetScriptStruct()->StructFlags, FlagsToCheck);
+	}
+
+	/**
+	 * Used to safely check whether all of the passed in flags are set.
+	 *
+	 * @param FlagsToCheck	Function flags to check for
+	 * @return true if all of the passed in flags are set (including no flags passed in), false otherwise
+	 */
+	FORCEINLINE bool HasAllFunctionFlags(EStructFlags FlagsToCheck) const
+	{
+		return EnumHasAllFlags(GetScriptStruct()->StructFlags, FlagsToCheck);
+	}
+
+	/**
+	 * Used to safely check whether specific of the passed in flags are set.
+	 *
+	 * @param FlagsToCheck	Function flags to check for
+	 * @param ExpectedFlags The flags from the flags to check that should be set
+	 * @return true if specific of the passed in flags are set (including no flags passed in), false otherwise
+	 */
+	FORCEINLINE bool HasSpecificFunctionFlags(EStructFlags FlagsToCheck, EStructFlags ExpectedFlags) const
+	{
+		EStructFlags Flags = GetScriptStruct()->StructFlags;
+		return (Flags & FlagsToCheck) == ExpectedFlags;
+	}
+
+	/**
+	 * If it is native, it is assumed to have defaults because it has a constructor
+	 * @return true if this struct has defaults
+	 */
+	FORCEINLINE bool HasDefaults() const
+	{
+		return !!GetScriptStruct()->GetCppStructOps();
+	}
+
+
+	int32 GetMacroDeclaredLineNumber() const
+	{
+		return MacroDeclaredLineNumber;
+	}
+
+	void SetMacroDeclaredLineNumber(int32 InMacroDeclaredLineNumber)
+	{
+		MacroDeclaredLineNumber = InMacroDeclaredLineNumber;
+	}
+
+private:
+	int32 MacroDeclaredLineNumber = INDEX_NONE;
 };
 
 /**
