Removed most of GetClass calls.

#rnx

diff --git a/Engine/Source/Programs/UnrealHeaderTool/Private/CodeGenerator.cpp b/Engine/Source/Programs/UnrealHeaderTool/Private/CodeGenerator.cpp
--- a/Engine/Source/Programs/UnrealHeaderTool/Private/CodeGenerator.cpp
+++ b/Engine/Source/Programs/UnrealHeaderTool/Private/CodeGenerator.cpp
@@ -173,8 +173,7 @@ namespace
 	/** Whether or not the given class has any replicated properties. */
 	static bool ClassHasReplicatedProperties(FUnrealClassDefinitionInfo& ClassDef)
 	{
-		UClass* Class = ClassDef.GetClass();
-		if (!Class->HasAnyClassFlags(CLASS_ReplicationDataIsSetUp))
+		if (!ClassDef.HasAnyClassFlags(CLASS_ReplicationDataIsSetUp))
 		{
 			for (FUnrealPropertyDefinitionInfo* PropertyDef : ClassDef.GetProperties())
 			{
@@ -184,6 +183,8 @@ namespace
 				}
 			}
 		}
+
+		UClass* Class = ClassDef.GetClass();
 		return Class->FirstOwnedClassRep < Class->ClassReps.Num();
 	}
 
@@ -299,7 +300,7 @@ namespace
 				"private:\r\n"
 				"\tREPLICATED_BASE_CLASS(%s%s)\r\n"
 				"public:\r\n"
-			), Class->GetPrefixCPP(), *Class->GetName());
+			), ClassDef.GetPrefixCPP(), *ClassDef.GetName());
 		}
 	}
 }
@@ -1739,7 +1740,7 @@ inline FString GetEventStructParamsName(FUnrealObjectDefinitionInfo& OuterDef, c
 	FString OuterName;
 	if (FUnrealClassDefinitionInfo* ClassDef = UHTCast<FUnrealClassDefinitionInfo>(OuterDef))
 	{
-		OuterName = ClassDef->GetClass()->GetName();
+		OuterName = ClassDef->GetName();
 	}
 	else if (FUnrealPackageDefinitionInfo* PackageDef = UHTCast<FUnrealPackageDefinitionInfo>(OuterDef))
 	{
@@ -2010,9 +2011,9 @@ void FNativeClassHeaderGenerator::ExportNativeGeneratedInitCode(FOutputDevice& O
 
 	UE_CLOG(Class->ClassGeneratedBy, LogCompile, Fatal, TEXT("For intrinsic and compiled-in classes, ClassGeneratedBy should always be null"));
 
-	const bool   bIsNoExport  = Class->HasAnyClassFlags(CLASS_NoExport);
+	const bool   bIsNoExport  = ClassDef.HasAnyClassFlags(CLASS_NoExport);
 	const bool   bIsDynamic   = ClassDef.IsDynamic();
-	const FString ClassNameCPP = FNameLookupCPP::GetNameCPP(Class);
+	const FString ClassNameCPP = ClassDef.GetAlternateNameCPP();
 
 	const FString& ApiString = GetAPIString();
 
@@ -2029,7 +2030,7 @@ void FNativeClassHeaderGenerator::ExportNativeGeneratedInitCode(FOutputDevice& O
 			// In a dynamic class the same function signature may be used for a Multi- and a Single-cast delegate.
 			if (!LocalFuncDef->IsADelegateFunction() || !bIsDynamic)
 			{
-				FError::Throwf(TEXT("The same function linked twice. Function: %s Class: %s"), *LocalFuncDef->GetName(), *Class->GetName());
+				FError::Throwf(TEXT("The same function linked twice. Function: %s Class: %s"), *LocalFuncDef->GetName(), *ClassDef.GetName());
 			}
 			continue;
 		}
@@ -2085,7 +2086,7 @@ void FNativeClassHeaderGenerator::ExportNativeGeneratedInitCode(FOutputDevice& O
 		if (!bIsDynamic)
 		{
 			check(ClassDef.HasSource());
-			FUnrealPackageDefinitionInfo& ClassPackageDef = GTypeDefinitionInfoMap.FindChecked<FUnrealPackageDefinitionInfo>(Class->GetOutermost());
+			FUnrealPackageDefinitionInfo& ClassPackageDef = ClassDef.GetPackageDef();
 			PackageDef.AddCrossModuleReference(OutReferenceGatherers.UniqueCrossModuleReferences);
 			OutDeclarations.Logf(TEXT("\t%s_API UPackage* %s;\r\n"), *ApiString, *ClassPackageDef.GetSingletonName());
 			Singletons.Logf(TEXT("\t\t(UObject* (*)())%s,\r\n"), *ClassPackageDef.GetSingletonNameChopped());
@@ -2222,12 +2223,12 @@ void FNativeClassHeaderGenerator::ExportNativeGeneratedInitCode(FOutputDevice& O
 		GeneratedClassRegisterFunctionText.Logf(TEXT("\t\tstatic const FCppClassTypeInfoStatic StaticCppClassTypeInfo;\r\n"));
 
 		StaticDefinitions.Logf(TEXT("\tconst FCppClassTypeInfoStatic %s::StaticCppClassTypeInfo = {\r\n"), *StaticsStructName);
-		StaticDefinitions.Logf(TEXT("\t\tTCppClassTypeTraits<%s>::IsAbstract,\r\n"), *FNameLookupCPP::GetNameCPP(Class, Class->HasAllClassFlags(CLASS_Interface)));
+		StaticDefinitions.Logf(TEXT("\t\tTCppClassTypeTraits<%s>::IsAbstract,\r\n"), *ClassDef.GetAlternateNameCPP(ClassDef.HasAllClassFlags(CLASS_Interface)));
 		StaticDefinitions.Logf(TEXT("\t};\r\n"));
 
 		GeneratedClassRegisterFunctionText.Log (TEXT("\t\tstatic const UECodeGen_Private::FClassParams ClassParams;\r\n"));
 
-		uint32 ClassFlags = (uint32)Class->ClassFlags;
+		uint32 ClassFlags = (uint32)ClassDef.GetClassFlags();
 		if (!bIsNoExport)
 		{
 			ClassFlags = ClassFlags | CLASS_MatchedSerializers;
@@ -2236,7 +2237,7 @@ void FNativeClassHeaderGenerator::ExportNativeGeneratedInitCode(FOutputDevice& O
 
 		StaticDefinitions.Logf(TEXT("\tconst UECodeGen_Private::FClassParams %s::ClassParams = {\r\n"), *StaticsStructName);
 		StaticDefinitions.Logf(TEXT("\t\t&%s::StaticClass,\r\n"), *ClassNameCPP);
-		StaticDefinitions.Logf(TEXT("\t\t%s,\r\n"), (Class->ClassConfigName != NAME_None) ? *CreateUTF8LiteralString(Class->ClassConfigName.ToString()) : TEXT("nullptr"));
+		StaticDefinitions.Logf(TEXT("\t\t%s,\r\n"), (ClassDef.GetClassConfigName() != NAME_None) ? *CreateUTF8LiteralString(ClassDef.GetClassConfigName().ToString()) : TEXT("nullptr"));
 		StaticDefinitions.Log (TEXT("\t\t&StaticCppClassTypeInfo,\r\n"));
 		StaticDefinitions.Logf(TEXT("\t\t%s,\r\n"), SingletonsArray);
 		StaticDefinitions.Logf(TEXT("\t\t%s,\r\n"), FunctionsArray);
@@ -2264,7 +2265,7 @@ void FNativeClassHeaderGenerator::ExportNativeGeneratedInitCode(FOutputDevice& O
 		{
 			const FString& DynamicClassPackageName = ClassDef.GetTypePackageName();
 			GeneratedClassRegisterFunctionText.Logf(TEXT("\t\tUPackage* OuterPackage = FindOrConstructDynamicTypePackage(TEXT(\"%s\"));\r\n"), *DynamicClassPackageName);
-			GeneratedClassRegisterFunctionText.Logf(TEXT("\t\tUClass* OuterClass = Cast<UClass>(StaticFindObjectFast(UClass::StaticClass(), OuterPackage, TEXT(\"%s\")));\r\n"), *FNativeClassHeaderGenerator::GetOverriddenName(Class));
+			GeneratedClassRegisterFunctionText.Logf(TEXT("\t\tUClass* OuterClass = Cast<UClass>(StaticFindObjectFast(UClass::StaticClass(), OuterPackage, TEXT(\"%s\")));\r\n"), *FNativeClassHeaderGenerator::GetOverriddenName(&ClassDef));
 			GeneratedClassRegisterFunctionText.Logf(TEXT("\t\tif (!OuterClass || !(OuterClass->ClassFlags & CLASS_Constructed))\r\n"));
 		}
 
@@ -2306,10 +2307,10 @@ void FNativeClassHeaderGenerator::ExportNativeGeneratedInitCode(FOutputDevice& O
 	FString SingletonName = ClassDef.GetSingletonName(true);
 	SingletonName.ReplaceInline(TEXT("()"), TEXT(""), ESearchCase::CaseSensitive); // function address
 
-	FString OverriddenClassName = *FNativeClassHeaderGenerator::GetOverriddenName(Class);
+	FString OverriddenClassName = *FNativeClassHeaderGenerator::GetOverriddenName(&ClassDef);
 
 	const FString EmptyString;
-	const FString& InitSearchableValuesFunctionName = bIsDynamic ? Class->GetMetaData(NAME_InitializeStaticSearchableValues) : EmptyString;
+	const FString& InitSearchableValuesFunctionName = bIsDynamic ? ClassDef.GetMetaData(NAME_InitializeStaticSearchableValues) : EmptyString;
 	const FString InitSearchableValuesFunctionParam = InitSearchableValuesFunctionName.IsEmpty() ? FString(TEXT("nullptr")) :
 		FString::Printf(TEXT("&%s::%s"), *ClassNameCPP, *InitSearchableValuesFunctionName);
 
@@ -2317,7 +2318,7 @@ void FNativeClassHeaderGenerator::ExportNativeGeneratedInitCode(FOutputDevice& O
 	// when base class changes during hot-reload.
 	uint32 BaseClassHash = 0;
 	FUnrealClassDefinitionInfo* SuperClassDef = ClassDef.GetSuperClass();
-	if (SuperClassDef && !SuperClassDef->GetClass()->HasAnyClassFlags(CLASS_Intrinsic))
+	if (SuperClassDef && !SuperClassDef->HasAnyClassFlags(CLASS_Intrinsic))
 	{
 		BaseClassHash = SuperClassDef->GetHash();
 	}
@@ -2560,9 +2561,8 @@ void FNativeClassHeaderGenerator::ExportFunction(FOutputDevice& Out, FReferenceG
 
 void FNativeClassHeaderGenerator::ExportNatives(FOutputDevice& Out, FUnrealClassDefinitionInfo& ClassDef)
 {
-	UClass* Class = ClassDef.GetClass();
-	const FString ClassCPPName = FNameLookupCPP::GetNameCPP(Class);
-	FString TypeName = Class->HasAnyClassFlags(CLASS_Interface) ? FString::Printf(TEXT("I%s"), *Class->GetName()) : ClassCPPName;
+	const FString ClassCPPName = ClassDef.GetAlternateNameCPP();
+	const FString TypeName = ClassDef.GetAlternateNameCPP(ClassDef.HasAnyClassFlags(CLASS_Interface));
 
 	Out.Logf(TEXT("\tvoid %s::StaticRegisterNatives%s()\r\n"), *ClassCPPName, *ClassCPPName);
 	Out.Log(TEXT("\t{\r\n"));
@@ -2756,7 +2756,7 @@ FString GetPreservedAccessSpecifierString(FUnrealClassDefinitionInfo& ClassDef)
 		PreservedAccessSpecifier = "public:";
 		break;
 	case EAccessSpecifier::ACCESS_NotAnAccessSpecifier :
-		PreservedAccessSpecifier = FString::Printf(TEXT("static_assert(false, \"Unknown access specifier for GENERATED_BODY() macro in class %s.\");"), *GetNameSafe(ClassDef.GetClass()));
+		PreservedAccessSpecifier = FString::Printf(TEXT("static_assert(false, \"Unknown access specifier for GENERATED_BODY() macro in class %s.\");"), *ClassDef.GetName());
 		break;
 	}
 
@@ -2821,7 +2821,7 @@ void FNativeClassHeaderGenerator::ExportClassFromSourceFileInner(
 
 		// VM -> C++ proxies (events and delegates).
 		FOutputDeviceNull NullOutput;
-		FOutputDevice& CallbackOut = Class->HasAnyClassFlags(CLASS_NoExport) ? NullOutput : OutCpp;
+		FOutputDevice& CallbackOut = ClassDef.HasAnyClassFlags(CLASS_NoExport) ? NullOutput : OutCpp;
 		FString CallbackWrappersMacroName = SourceFile.GetGeneratedMacroName(StructData, TEXT("_CALLBACK_WRAPPERS"));
 		ExportCallbackFunctions(
 			OutGeneratedHeaderText,
@@ -2829,7 +2829,7 @@ void FNativeClassHeaderGenerator::ExportClassFromSourceFileInner(
 			OutReferenceGatherers.ForwardDeclarations,
 			CallbackFunctions,
 			*CallbackWrappersMacroName,
-			(Class->ClassFlags & CLASS_Interface) ? EExportCallbackType::Interface : EExportCallbackType::Class,
+			ClassDef.HasAnyClassFlags(CLASS_Interface) ? EExportCallbackType::Interface : EExportCallbackType::Class,
 			*GetAPIString()
 		);
 
@@ -2838,7 +2838,7 @@ void FNativeClassHeaderGenerator::ExportClassFromSourceFileInner(
 	}
 
 	// Class definition.
-	if (!Class->HasAnyClassFlags(CLASS_NoExport))
+	if (!ClassDef.HasAnyClassFlags(CLASS_NoExport))
 	{
 		ExportNatives(OutCpp, ClassDef);
 	}
@@ -2847,11 +2847,11 @@ void FNativeClassHeaderGenerator::ExportClassFromSourceFileInner(
 	ExportNativeGeneratedInitCode(OutCpp, OutDeclarations, OutReferenceGatherers, SourceFile, ClassDef, FriendText);
 
 	// the name for the C++ version of the UClass
-	const FString ClassCPPName = FNameLookupCPP::GetNameCPP(Class);
-	const FString SuperClassCPPName = (SuperClass ? FNameLookupCPP::GetNameCPP(SuperClass) : TEXT("None"));
+	const FString ClassCPPName = ClassDef.GetAlternateNameCPP();
+	const FString SuperClassCPPName = (SuperClassDef ? SuperClassDef->GetAlternateNameCPP() : TEXT("None"));
 
 	FString APIArg = PackageDef.GetShortUpperName();
-	if (!Class->HasAnyClassFlags(CLASS_MinimalAPI))
+	if (!ClassDef.HasAnyClassFlags(CLASS_MinimalAPI))
 	{
 		APIArg = TEXT("NO");
 	}
@@ -2913,12 +2913,12 @@ void FNativeClassHeaderGenerator::ExportClassFromSourceFileInner(
 		Boilerplate.Log(*FriendText);
 		Boilerplate.Logf(TEXT("public:\r\n"));
 
-		const bool bCastedClass = Class->HasAnyCastFlag(CASTCLASS_AllFlags) && SuperClass && Class->ClassCastFlags != SuperClass->ClassCastFlags;
+		const bool bCastedClass = Class->HasAnyCastFlag(CASTCLASS_AllFlags) && SuperClassDef && Class->ClassCastFlags != SuperClass->ClassCastFlags;
 
 		Boilerplate.Logf(TEXT("\tDECLARE_CLASS(%s, %s, COMPILED_IN_FLAGS(%s%s), %s, TEXT(\"%s\"), %s_API)\r\n"),
 			*ClassCPPName,
 			*SuperClassCPPName,
-			Class->HasAnyClassFlags(CLASS_Abstract) ? TEXT("CLASS_Abstract") : TEXT("0"),
+			ClassDef.HasAnyClassFlags(CLASS_Abstract) ? TEXT("CLASS_Abstract") : TEXT("0"),
 			*GetClassFlagExportText(ClassDef),
 			bCastedClass ? *FString::Printf(TEXT("CASTCLASS_%s"), *ClassCPPName) : TEXT("CASTCLASS_None"),
 			*ClassDef.GetTypePackageName(),
@@ -2934,10 +2934,10 @@ void FNativeClassHeaderGenerator::ExportClassFromSourceFileInner(
 
 		if (SuperClassDef && ClassDef.GetClassWithin() != SuperClassDef->GetClassWithin())
 		{
-			Boilerplate.Logf(TEXT("\tDECLARE_WITHIN(%s)\r\n"), *FNameLookupCPP::GetNameCPP(Class->ClassWithin));
+			Boilerplate.Logf(TEXT("\tDECLARE_WITHIN(%s)\r\n"), *ClassDef.GetClassWithin()->GetAlternateNameCPP());
 		}
 
-		if (Class->HasAnyClassFlags(CLASS_Interface))
+		if (ClassDef.HasAnyClassFlags(CLASS_Interface))
 		{
 			ExportConstructorsMacros(OutGeneratedHeaderText, OutCpp, StandardUObjectConstructorsMacroCall, EnhancedUObjectConstructorsMacroCall, SourceFile.GetGeneratedMacroName(StructData), ClassDef, *APIArg);
 
@@ -2980,12 +2980,8 @@ void FNativeClassHeaderGenerator::ExportClassFromSourceFileInner(
 			// Export the pure interface version of the class
 
 			// the name of the pure interface class
-			FString InterfaceCPPName = FString::Printf(TEXT("I%s"), *Class->GetName());
-			FString SuperInterfaceCPPName;
-			if (SuperClass != NULL)
-			{
-				SuperInterfaceCPPName = FString::Printf(TEXT("I%s"), *SuperClass->GetName());
-			}
+			FString InterfaceCPPName = ClassDef.GetAlternateNameCPP(true);
+			FString SuperInterfaceCPPName = SuperClassDef ? SuperClassDef->GetAlternateNameCPP(true) : FString();
 
 			// Thunk functions
 			FUHTStringBuilder InterfaceBoilerplate;
@@ -2994,7 +2990,7 @@ void FNativeClassHeaderGenerator::ExportClassFromSourceFileInner(
 			InterfaceBoilerplate.Logf(TEXT("public:\r\n\ttypedef %s UClassType;\r\n"), *ClassCPPName);
 			InterfaceBoilerplate.Logf(TEXT("\ttypedef %s ThisClass;\r\n"), *InterfaceCPPName);
 
-			ExportInterfaceCallFunctions(OutCpp, InterfaceBoilerplate, OutReferenceGatherers, CallbackFunctions, *Class->GetName());
+			ExportInterfaceCallFunctions(OutCpp, InterfaceBoilerplate, OutReferenceGatherers, CallbackFunctions, *ClassDef.GetName());
 
 			// we'll need a way to get to the UObject portion of a native interface, so that we can safely pass native interfaces
 			// to script VM functions
@@ -3021,9 +3017,9 @@ void FNativeClassHeaderGenerator::ExportClassFromSourceFileInner(
 		else
 		{
 			// export the class's config name
-			if (SuperClass && Class->ClassConfigName != NAME_None && Class->ClassConfigName != SuperClass->ClassConfigName)
+			if (SuperClassDef && ClassDef.GetClassConfigName() != NAME_None && ClassDef.GetClassConfigName() != SuperClassDef->GetClassConfigName())
 			{
-				Boilerplate.Logf(TEXT("\tstatic const TCHAR* StaticConfigName() {return TEXT(\"%s\");}\r\n\r\n"), *Class->ClassConfigName.ToString());
+				Boilerplate.Logf(TEXT("\tstatic const TCHAR* StaticConfigName() {return TEXT(\"%s\");}\r\n\r\n"), *ClassDef.GetClassConfigName().ToString());
 			}
 
 			// export implementation of _getUObject for classes that implement interfaces
@@ -3059,7 +3055,7 @@ void FNativeClassHeaderGenerator::ExportClassFromSourceFileInner(
 	{
 		const TCHAR* Public = TEXT("public:") LINE_TERMINATOR;
 
-		const bool bIsIInterface = Class->HasAnyClassFlags(CLASS_Interface);
+		const bool bIsIInterface = ClassDef.HasAnyClassFlags(CLASS_Interface);
 
 		const TCHAR* MacroName;
 		FString DeprecationWarning;
@@ -3142,14 +3138,13 @@ void ExportVTableHelperCtorAndCaller(FOutputDevice& Out, FUnrealClassDefinitionI
  */
 void ExportStandardConstructorsMacro(FOutputDevice& Out, FUnrealClassDefinitionInfo& ClassDef, const TCHAR* API, const TCHAR* ClassCPPName)
 {
-	UClass* Class = ClassDef.GetClass();
-	if (!Class->HasAnyClassFlags(CLASS_CustomConstructor))
+	if (!ClassDef.HasAnyClassFlags(CLASS_CustomConstructor))
 	{
 		Out.Logf(TEXT("\t/** Standard constructor, called after all reflected properties have been initialized */\r\n"));
 		Out.Logf(TEXT("\t%s_API %s(const FObjectInitializer& ObjectInitializer%s);\r\n"), API, ClassCPPName,
 			ClassDef.GetStructMetaData().bDefaultConstructorDeclared ? TEXT("") : TEXT(" = FObjectInitializer::Get()"));
 	}
-	if (Class->HasAnyClassFlags(CLASS_Abstract))
+	if (ClassDef.HasAnyClassFlags(CLASS_Abstract))
 	{
 		Out.Logf(TEXT("\tDEFINE_ABSTRACT_DEFAULT_OBJECT_INITIALIZER_CONSTRUCTOR_CALL(%s)\r\n"), ClassCPPName);
 	}
@@ -3172,7 +3167,6 @@ void ExportStandardConstructorsMacro(FOutputDevice& Out, FUnrealClassDefinitionI
  */
 void ExportConstructorDefinition(FOutputDevice& Out, FUnrealClassDefinitionInfo& ClassDef, const TCHAR* API, const TCHAR* ClassCPPName)
 {
-	UClass* Class = ClassDef.GetClass();
 	FStructMetaData& StructData = ClassDef.GetStructMetaData();
 	if (!StructData.bConstructorDeclared)
 	{
@@ -3226,7 +3220,7 @@ void ExportDefaultConstructorCallDefinition(FOutputDevice& Out, FUnrealClassDefi
 {
 	if (ClassDef.GetStructMetaData().bObjectInitializerConstructorDeclared)
 	{
-		if (ClassDef.GetClass()->HasAnyClassFlags(CLASS_Abstract))
+		if (ClassDef.HasAnyClassFlags(CLASS_Abstract))
 		{
 			Out.Logf(TEXT("\tDEFINE_ABSTRACT_DEFAULT_OBJECT_INITIALIZER_CONSTRUCTOR_CALL(%s)\r\n"), ClassCPPName);
 		}
@@ -3237,7 +3231,7 @@ void ExportDefaultConstructorCallDefinition(FOutputDevice& Out, FUnrealClassDefi
 	}
 	else if (ClassDef.GetStructMetaData().bDefaultConstructorDeclared)
 	{
-		if (ClassDef.GetClass()->HasAnyClassFlags(CLASS_Abstract))
+		if (ClassDef.HasAnyClassFlags(CLASS_Abstract))
 		{
 			Out.Logf(TEXT("\tDEFINE_ABSTRACT_DEFAULT_CONSTRUCTOR_CALL(%s)\r\n"), ClassCPPName);
 		}
@@ -3284,8 +3278,7 @@ FString GetBuildPath(FUnrealSourceFile& SourceFile)
 
 void FNativeClassHeaderGenerator::ExportConstructorsMacros(FOutputDevice& OutGeneratedHeaderText, FOutputDevice& Out, FOutputDevice& StandardUObjectConstructorsMacroCall, FOutputDevice& EnhancedUObjectConstructorsMacroCall, const FString& ConstructorsMacroPrefix, FUnrealClassDefinitionInfo& ClassDef, const TCHAR* APIArg)
 {
-	UClass* Class = ClassDef.GetClass();
-	const FString ClassCPPName = FNameLookupCPP::GetNameCPP(Class);
+	const FString ClassCPPName = ClassDef.GetAlternateNameCPP();
 
 	FUHTStringBuilder StdMacro;
 	FUHTStringBuilder EnhMacro;
@@ -3345,32 +3338,31 @@ FString FNativeClassHeaderGenerator::GetClassFlagExportText(FUnrealClassDefiniti
 {
 	FString StaticClassFlagText;
 
-	UClass* Class = ClassDef.GetClass();
-	if ( Class->HasAnyClassFlags(CLASS_Transient) )
+	if (ClassDef.HasAnyClassFlags(CLASS_Transient) )
 	{
 		StaticClassFlagText += TEXT(" | CLASS_Transient");
 	}
-	if( Class->HasAnyClassFlags(CLASS_DefaultConfig) )
+	if (ClassDef.HasAnyClassFlags(CLASS_DefaultConfig) )
 	{
 		StaticClassFlagText += TEXT(" | CLASS_DefaultConfig");
 	}
-	if( Class->HasAnyClassFlags(CLASS_GlobalUserConfig) )
+	if (ClassDef.HasAnyClassFlags(CLASS_GlobalUserConfig) )
 	{
 		StaticClassFlagText += TEXT(" | CLASS_GlobalUserConfig");
 	}
-	if (Class->HasAnyClassFlags(CLASS_ProjectUserConfig))
+	if (ClassDef.HasAnyClassFlags(CLASS_ProjectUserConfig))
 	{
 		StaticClassFlagText += TEXT(" | CLASS_ProjectUserConfig");
 	}
-	if( Class->HasAnyClassFlags(CLASS_Config) )
+	if (ClassDef.HasAnyClassFlags(CLASS_Config) )
 	{
 		StaticClassFlagText += TEXT(" | CLASS_Config");
 	}
-	if ( Class->HasAnyClassFlags(CLASS_Interface) )
+	if (ClassDef.HasAnyClassFlags(CLASS_Interface) )
 	{
 		StaticClassFlagText += TEXT(" | CLASS_Interface");
 	}
-	if ( Class->HasAnyClassFlags(CLASS_Deprecated) )
+	if (ClassDef.HasAnyClassFlags(CLASS_Deprecated) )
 	{
 		StaticClassFlagText += TEXT(" | CLASS_Deprecated");
 	}
@@ -4992,12 +4984,11 @@ void FNativeClassHeaderGenerator::ExportFunctionThunk(FUHTStringBuilder& RPCWrap
 
 	FUnrealClassDefinitionInfo* OwnerClassDef = FunctionDef.GetOwnerClass();
 	check(OwnerClassDef);
-	UClass* OwnerClass = OwnerClassDef->GetClass();
 
 	OwnerClassDef->GetDefinitionRange().Validate();
 	const TCHAR* ClassStart = OwnerClassDef->GetDefinitionRange().Start;
 	const TCHAR* ClassEnd = OwnerClassDef->GetDefinitionRange().End;
-	FString      ClassName  = OwnerClass->GetName();
+	FString      ClassName  = OwnerClassDef->GetName();
 
 	FString ClassDefinition(UE_PTRDIFF_TO_INT32(ClassEnd - ClassStart), ClassStart);
 
@@ -5107,9 +5098,8 @@ void FNativeClassHeaderGenerator::ExportNativeFunctions(FOutputDevice& OutGenera
 	FNativeFunctionStringBuilder RuntimeStringBuilders;
 	FNativeFunctionStringBuilder EditorStringBuilders;
 
-	UClass* Class = ClassDef.GetClass();
 	FStructMetaData& StructData = ClassDef.GetStructMetaData();
-	const FString ClassCPPName = FNameLookupCPP::GetNameCPP(Class, Class->HasAnyClassFlags(CLASS_Interface));
+	const FString ClassCPPName = ClassDef.GetAlternateNameCPP(ClassDef.HasAnyClassFlags(CLASS_Interface));
 
 	ClassDef.GetDefinitionRange().Validate();
 	const TCHAR* ClassStart = ClassDef.GetDefinitionRange().Start;
@@ -5118,7 +5108,7 @@ void FNativeClassHeaderGenerator::ExportNativeFunctions(FOutputDevice& OutGenera
 	// gather static class data
 	TArray<FString> SparseClassDataTypes;
 	ClassDef.GetSparseClassDataTypes(SparseClassDataTypes);
-	FString FullClassName = FUnrealTypeDefinitionInfo::GetNameWithPrefix(Class);
+	FString FullClassName = ClassDef.GetNameWithPrefix();
 	for (const FString& SparseClassDataString : SparseClassDataTypes)
 	{
 		RuntimeStringBuilders.AutogeneratedStaticData.Logf(TEXT("F%s* Get%s()\r\n"), *SparseClassDataString, *SparseClassDataString);
@@ -5392,8 +5382,7 @@ void FNativeClassHeaderGenerator::ExportCallbackFunctions(
 		const FFuncInfo& FunctionData = FunctionDef->GetFunctionData();
 		FString          FunctionName = FunctionDef->GetName();
 		FUnrealClassDefinitionInfo&    ClassDef = UHTCastChecked<FUnrealClassDefinitionInfo>(FunctionDef->GetOuter());
-		UClass*          Class = ClassDef.GetClass();
-		const FString    ClassName = FNameLookupCPP::GetNameCPP(Class);
+		const FString    ClassName = ClassDef.GetAlternateNameCPP();
 
 		if (FunctionDef->HasAnyFunctionFlags(FUNC_NetResponse))
 		{
@@ -5742,8 +5731,7 @@ void FNativeClassHeaderGenerator::GenerateSourceFiles(
 					for (const TSharedRef<FUnrealTypeDefinitionInfo>& TypeDef : SourceFile.GetDefinedClasses())
 					{
 						FUnrealClassDefinitionInfo& ClassDef = TypeDef->AsClassChecked();
-						UClass* Class = ClassDef.GetClass();
-						if (!(Class->ClassFlags & CLASS_Intrinsic))
+						if (!ClassDef.HasAnyClassFlags(CLASS_Intrinsic))
 						{
 							Generator.ExportClassFromSourceFileInner(GeneratedHeaderText, GeneratedCPPText, GeneratedFunctionDeclarations, ReferenceGatherers, ClassDef, SourceFile, GeneratedCPP.ExportFlags);
 						}
@@ -6342,34 +6330,37 @@ void ResolveSuperClasses(UPackage* Package)
 			ClassDef.SetClassWithin(&GTypeDefinitionInfoMap.FindChecked<FUnrealClassDefinitionInfo>(DefinedClass->ClassWithin));
 		}
 
-		if (DefinedClass->HasAnyClassFlags(CLASS_Intrinsic | CLASS_NoExport))
-		{
-			continue;
-		}
-
 		// Resolve the base class
 		{
 			FUnrealStructDefinitionInfo::FBaseStructInfo& SuperClassInfo = ClassDef.GetSuperStructInfo();
-			const FString& BaseClassName = SuperClassInfo.Name;
-			const FString& BaseClassNameStripped = GetClassNameWithPrefixRemoved(BaseClassName);
 
-			if (!BaseClassNameStripped.IsEmpty() && !DefinedClass->GetSuperClass())
+			if (DefinedClass->GetSuperClass() == nullptr)
 			{
-				FUnrealClassDefinitionInfo* FoundBaseClassDef = GTypeDefinitionInfoMap.FindByName<FUnrealClassDefinitionInfo>(*BaseClassNameStripped);
-				if (FoundBaseClassDef == nullptr)
+				if (!SuperClassInfo.Name.IsEmpty() && !DefinedClass->GetSuperClass())
 				{
-					FoundBaseClassDef = GTypeDefinitionInfoMap.FindByName<FUnrealClassDefinitionInfo>(*BaseClassName);
-				}
+					const FString& BaseClassName = SuperClassInfo.Name;
+					const FString& BaseClassNameStripped = GetClassNameWithPrefixRemoved(BaseClassName);
 
-				if (FoundBaseClassDef == nullptr)
-				{
-					// Don't know its parent class. Raise error.
-					FError::Throwf(TEXT("Couldn't find parent type for '%s' named '%s' in current module (Package: %s) or any other module parsed so far."), *DefinedClass->GetName(), *BaseClassName, *GetNameSafe(Package));
-				}
+					FUnrealClassDefinitionInfo* FoundBaseClassDef = GTypeDefinitionInfoMap.FindByName<FUnrealClassDefinitionInfo>(*BaseClassNameStripped);
+					if (FoundBaseClassDef == nullptr)
+					{
+						FoundBaseClassDef = GTypeDefinitionInfoMap.FindByName<FUnrealClassDefinitionInfo>(*BaseClassName);
+					}
+
+					if (FoundBaseClassDef == nullptr)
+					{
+						// Don't know its parent class. Raise error.
+						FError::Throwf(TEXT("Couldn't find parent type for '%s' named '%s' in current module (Package: %s) or any other module parsed so far."), *DefinedClass->GetName(), *BaseClassName, *Package->GetName());
+					}
 
-				DefinedClass->SetSuperStruct(FoundBaseClassDef->GetClass());
-				DefinedClass->ClassCastFlags |= FoundBaseClassDef->GetClass()->ClassCastFlags;
-				SuperClassInfo.Struct = FoundBaseClassDef;
+					DefinedClass->SetSuperStruct(FoundBaseClassDef->GetClass());
+					DefinedClass->ClassCastFlags |= FoundBaseClassDef->GetClass()->ClassCastFlags;
+					SuperClassInfo.Struct = FoundBaseClassDef;
+				}
+			}
+			else
+			{
+				SuperClassInfo.Struct = &GTypeDefinitionInfoMap.FindChecked<FUnrealClassDefinitionInfo>(DefinedClass->GetSuperClass());
 			}
 		}
 
@@ -6817,19 +6808,18 @@ void PrepareTypesForExport(TArray<FUnrealPackageDefinitionInfo*>& PackageDefs)
 				// Check to see if we need write the classes 
 				const bool bPackageHasAnyExportClasses = PackageClasses.ContainsByPredicate([](FUnrealClassDefinitionInfo* ClassDef)
 					{
-						UClass* Class = ClassDef->GetClass();
-						return Class->HasAnyClassFlags(CLASS_Native) && !Class->HasAnyClassFlags(CLASS_NoExport | CLASS_Intrinsic);
+						return ClassDef->HasAnyClassFlags(CLASS_Native) && !ClassDef->HasAnyClassFlags(CLASS_NoExport | CLASS_Intrinsic);
 					});
 				if (bPackageHasAnyExportClasses)
 				{
 					for (FUnrealClassDefinitionInfo* ClassDef : PackageClasses)
 					{
 						UClass* Class = ClassDef->GetClass();
-						if (!Class->HasAnyClassFlags(CLASS_Native))
+						if (!ClassDef->HasAnyClassFlags(CLASS_Native))
 						{
 							Class->UnMark(EObjectMark(OBJECTMARK_TagImp | OBJECTMARK_TagExp));
 						}
-						else if (!Class->HasAnyClassFlags(CLASS_NoExport))
+						else if (!ClassDef->HasAnyClassFlags(CLASS_NoExport))
 						{
 							PackageDef->SetWriteClassesH(true);
 							Class->UnMark(OBJECTMARK_TagImp);
diff --git a/Engine/Source/Programs/UnrealHeaderTool/Private/HeaderParser.cpp b/Engine/Source/Programs/UnrealHeaderTool/Private/HeaderParser.cpp
--- a/Engine/Source/Programs/UnrealHeaderTool/Private/HeaderParser.cpp
+++ b/Engine/Source/Programs/UnrealHeaderTool/Private/HeaderParser.cpp
@@ -2132,12 +2132,11 @@ FUnrealScriptStructDefinitionInfo& FHeaderParser::CompileStructDeclaration()
 			// If the struct was found, confirm it adheres to the correct syntax. This should always fail if we were expecting an override that was not found.
 			if (bOverrideParentStructName)
 			{
-				UScriptStruct* BaseStruct = BaseStructDef->AsScriptStructChecked().GetScriptStruct();
-				const TCHAR* PrefixCPP = UHTConfig.StructsWithTPrefix.Contains(ParentStructNameStripped) ? TEXT("T") : BaseStruct->GetPrefixCPP();
+				const TCHAR* PrefixCPP = UHTConfig.StructsWithTPrefix.Contains(ParentStructNameStripped) ? TEXT("T") : BaseStructDef->AsScriptStructChecked().GetPrefixCPP();
 				if (ParentStructNameInScript != FString::Printf(TEXT("%s%s"), PrefixCPP, *ParentStructNameStripped))
 				{
+					FError::Throwf(TEXT("Parent Struct '%s' is missing a valid Unreal prefix, expecting '%s'"), *ParentStructNameInScript, *FString::Printf(TEXT("%s%s"), PrefixCPP, *BaseStructDef->GetName()));
 					BaseStructDef = nullptr;
-					FError::Throwf(TEXT("Parent Struct '%s' is missing a valid Unreal prefix, expecting '%s'"), *ParentStructNameInScript, *FString::Printf(TEXT("%s%s"), PrefixCPP, *BaseStruct->GetName()));
 				}
 			}
 		}
@@ -2158,10 +2157,10 @@ FUnrealScriptStructDefinitionInfo& FHeaderParser::CompileStructDeclaration()
 	// Check to make sure the syntactic native prefix was set-up correctly.
 	// If this check results in a false positive, it will be flagged as an identifier failure.
 	FString DeclaredPrefix = GetClassPrefix( StructNameInScript );
-	if( DeclaredPrefix == Struct->GetPrefixCPP() || DeclaredPrefix == TEXT("T") )
+	if( DeclaredPrefix == StructDef.GetPrefixCPP() || DeclaredPrefix == TEXT("T") )
 	{
 		// Found a prefix, do a basic check to see if it's valid
-		const TCHAR* ExpectedPrefixCPP = UHTConfig.StructsWithTPrefix.Contains(StructNameStripped) ? TEXT("T") : Struct->GetPrefixCPP();
+		const TCHAR* ExpectedPrefixCPP = UHTConfig.StructsWithTPrefix.Contains(StructNameStripped) ? TEXT("T") : StructDef.GetPrefixCPP();
 		FString ExpectedStructName = FString::Printf(TEXT("%s%s"), ExpectedPrefixCPP, *StructNameStripped);
 		if (StructNameInScript != ExpectedStructName)
 		{
@@ -2170,7 +2169,7 @@ FUnrealScriptStructDefinitionInfo& FHeaderParser::CompileStructDeclaration()
 	}
 	else
 	{
-		const TCHAR* ExpectedPrefixCPP = UHTConfig.StructsWithTPrefix.Contains(StructNameInScript) ? TEXT("T") : Struct->GetPrefixCPP();
+		const TCHAR* ExpectedPrefixCPP = UHTConfig.StructsWithTPrefix.Contains(StructNameInScript) ? TEXT("T") : StructDef.GetPrefixCPP();
 		FString ExpectedStructName = FString::Printf(TEXT("%s%s"), ExpectedPrefixCPP, *StructNameInScript);
 		FError::Throwf(TEXT("Struct '%s' is missing a valid Unreal prefix, expecting '%s'"), *StructNameInScript, *ExpectedStructName);
 	}
@@ -2634,7 +2633,7 @@ void FHeaderParser::FixupDelegateProperties(FUnrealStructDefinitionInfo& StructD
 					FUnrealClassDefinitionInfo* DelegateOwnerClassDef = FUnrealClassDefinitionInfo::FindScriptClassOrThrow(DelegateClassName);
 					if (DelegateOwnerClassDef->GetScope()->FindTypeByName(*DelegateName) != nullptr)
 					{
-						FError::Throwf(TEXT("Inaccessible type: '%s'"), *DelegateOwnerClassDef->GetClass()->GetPathName());
+						FError::Throwf(TEXT("Inaccessible type: '%s'"), *DelegateOwnerClassDef->GetPathName());
 					}
 					UFunction* SourceDelegateFunction = Cast<UFunction>(FindField(*DelegateOwnerClassDef, *DelegateName, false, UFunction::StaticClass(), NULL));
 					if (SourceDelegateFunction != nullptr)
@@ -2742,7 +2741,6 @@ void FHeaderParser::CheckSparseClassData(const FUnrealStructDefinitionInfo& Stru
 	{
 		return;
 	}
-	const UClass* ClassToCheck = ClassDef->GetClass();
 
 	TArray<FString> SparseClassDataTypes;
 	ClassDef->GetSparseClassDataTypes(SparseClassDataTypes);
@@ -2750,12 +2748,12 @@ void FHeaderParser::CheckSparseClassData(const FUnrealStructDefinitionInfo& Stru
 	// for now we only support one sparse class data structure per class
 	if (SparseClassDataTypes.Num() > 1)
 	{
-		FError::Throwf(TEXT("Class %s contains multiple sparse class data types."), *ClassToCheck->GetName());
+		FError::Throwf(TEXT("Class %s contains multiple sparse class data types."), *ClassDef->GetName());
 		return;
 	}
 	if (SparseClassDataTypes.Num() == 0)
 	{
-		FError::Throwf(TEXT("Class %s has sparse class metadata but does not specify a type."), *ClassToCheck->GetName());
+		FError::Throwf(TEXT("Class %s has sparse class metadata but does not specify a type."), *ClassDef->GetName());
 		return;
 	}
 
@@ -2766,7 +2764,7 @@ void FHeaderParser::CheckSparseClassData(const FUnrealStructDefinitionInfo& Stru
 		// make sure the sparse class data struct actually exists
 		if (!SparseClassDataStructDef)
 		{
-			FError::Throwf(TEXT("Unable to find sparse data type %s for class %s."), *SparseClassDataTypeName, *ClassToCheck->GetName());
+			FError::Throwf(TEXT("Unable to find sparse data type %s for class %s."), *SparseClassDataTypeName, *ClassDef->GetName());
 			return;
 		}
 
@@ -2793,7 +2791,8 @@ void FHeaderParser::CheckSparseClassData(const FUnrealStructDefinitionInfo& Stru
 		}
 
 		// if the class's parent has a sparse class data struct then the current class must also use the same struct or one that inherits from it
-		const UClass* ParentClass = ClassToCheck->GetSuperClass();
+
+		const FUnrealClassDefinitionInfo* ParentClassDef = ClassDef->GetSuperClass();
 		TArray<FString> ParentSparseClassDataTypeNames;
 		ClassDef->GetSuperClass()->GetSparseClassDataTypes(ParentSparseClassDataTypeNames);
 		for (FString& ParentSparseClassDataTypeName : ParentSparseClassDataTypeNames)
@@ -2803,7 +2802,7 @@ void FHeaderParser::CheckSparseClassData(const FUnrealStructDefinitionInfo& Stru
 				UScriptStruct* ParentSparseClassDataStruct = ParentSparseClassDataStructDef->GetScriptStruct();
 				if (!SparseClassDataStructDef->GetScriptStruct()->IsChildOf(ParentSparseClassDataStruct))
 				{
-					FError::Throwf(TEXT("Class %s is a child of %s but its sparse class data struct, %s, does not inherit from %s."), *ClassToCheck->GetName(), *ParentClass->GetName(), *SparseClassDataStructDef->GetName(), *ParentSparseClassDataStruct->GetName());
+					FError::Throwf(TEXT("Class %s is a child of %s but its sparse class data struct, %s, does not inherit from %s."), *ClassDef->GetName(), *ParentClassDef->GetName(), *SparseClassDataStructDef->GetName(), *ParentSparseClassDataStruct->GetName());
 				}
 			}
 		}
@@ -2934,8 +2933,6 @@ void FHeaderParser::VerifyRepNotifyCallback(FUnrealPropertyDefinitionInfo& Prope
 }
 void FHeaderParser::VerifyPropertyMarkups(FUnrealClassDefinitionInfo& TargetClassDef)
 {
-	UClass* TargetClass = TargetClassDef.GetClass();
-
 	// Iterate over all properties, looking for those flagged as CPF_RepNotify
 	for (FUnrealPropertyDefinitionInfo* PropertyDef : TargetClassDef.GetProperties())
 	{
@@ -4291,7 +4288,7 @@ void FHeaderParser::GetVarType(
 			UScriptStruct* Struct = StructDef->GetScriptStruct();
 			if (bStripped)
 			{
-				const TCHAR* PrefixCPP = UHTConfig.StructsWithTPrefix.Contains(IdentifierStripped) ? TEXT("T") : Struct->GetPrefixCPP();
+				const TCHAR* PrefixCPP = UHTConfig.StructsWithTPrefix.Contains(IdentifierStripped) ? TEXT("T") : StructDef->GetPrefixCPP();
 				FString ExpectedStructName = FString::Printf(TEXT("%s%s"), PrefixCPP, *StructDef->GetName());
 				if (FString(VarType.Identifier) != ExpectedStructName)
 				{
@@ -4300,7 +4297,7 @@ void FHeaderParser::GetVarType(
 			}
 			else if (!UHTConfig.StructsWithNoPrefix.Contains(VarType.Identifier))
 			{
-				const TCHAR* PrefixCPP = UHTConfig.StructsWithTPrefix.Contains(VarType.Identifier) ? TEXT("T") : Struct->GetPrefixCPP();
+				const TCHAR* PrefixCPP = UHTConfig.StructsWithTPrefix.Contains(VarType.Identifier) ? TEXT("T") : StructDef->GetPrefixCPP();
 				FError::Throwf(TEXT("Struct '%s' is missing a prefix, expecting '%s'"), VarType.Identifier, *FString::Printf(TEXT("%s%s"), PrefixCPP, *StructDef->GetName()));
 			}
 
@@ -4425,8 +4422,6 @@ void FHeaderParser::GetVarType(
 
 				if (TempClassDef != NULL)
 				{
-					UClass* TempClass = TempClassDef->GetClass();
-
 					bHandledType = true;
 
 					if ((PropertyType == CPT_WeakObjectReference) && (Disallow & CPF_AutoWeak)) // if it is not allowing anything, force it strong. this is probably a function arg
@@ -4435,13 +4430,13 @@ void FHeaderParser::GetVarType(
 					}
 
 					// if this is an interface class, we use the FInterfaceProperty class instead of FObjectProperty
-					if ((PropertyType == CPT_ObjectReference) && TempClass->HasAnyClassFlags(CLASS_Interface))
+					if ((PropertyType == CPT_ObjectReference) && TempClassDef->HasAnyClassFlags(CLASS_Interface))
 					{
 						PropertyType = CPT_Interface;
 					}
 
 					VarProperty = FPropertyBase(*TempClassDef, PropertyType, bWeakIsAuto);
-					if (TempClass->IsChildOf(UClass::StaticClass()))
+					if (TempClassDef->GetClass()->IsChildOf(UClass::StaticClass()))
 					{
 						if (MatchSymbol(TEXT('<')))
 						{
@@ -4513,7 +4508,7 @@ void FHeaderParser::GetVarType(
 					}
 
 					// Imply const if it's a parameter that is a pointer to a const class
-					if (VariableCategory != EVariableCategory::Member && (TempClass != NULL) && (TempClass->HasAnyClassFlags(CLASS_Const)))
+					if (VariableCategory != EVariableCategory::Member && (TempClassDef != NULL) && (TempClassDef->HasAnyClassFlags(CLASS_Const)))
 					{
 						Flags |= CPF_ConstParm;
 					}
@@ -5304,43 +5299,41 @@ bool FHeaderParser::CompileDeclaration(TArray<FUnrealFunctionDefinitionInfo*>& D
 
 	if (bEncounteredNewStyleClass_UnmatchedBrackets && IsInAClass())
 	{
-		if (UClass* Class = GetCurrentClass())
+		FUnrealStructDefinitionInfo& StructDef = GetCurrentClassDef();
+		FToken ConstructorToken = Token;
+
+		// Allow explicit constructors
+		bool bFoundExplicit = ConstructorToken.Matches(TEXT("explicit"), ESearchCase::CaseSensitive);
+		if (bFoundExplicit)
 		{
-			FToken ConstructorToken = Token;
+			GetToken(ConstructorToken);
+		}
 
-			// Allow explicit constructors
-			bool bFoundExplicit = ConstructorToken.Matches(TEXT("explicit"), ESearchCase::CaseSensitive);
-			if (bFoundExplicit)
+		bool bSkippedAPIToken = false;
+		if (FString(ConstructorToken.Identifier).EndsWith(TEXT("_API"), ESearchCase::CaseSensitive))
+		{
+			if (!bFoundExplicit)
 			{
-				GetToken(ConstructorToken);
+				// Explicit can come before or after an _API
+				MatchIdentifier(TEXT("explicit"), ESearchCase::CaseSensitive);
 			}
 
-			bool bSkippedAPIToken = false;
-			if (FString(ConstructorToken.Identifier).EndsWith(TEXT("_API"), ESearchCase::CaseSensitive))
-			{
-				if (!bFoundExplicit)
-				{
-					// Explicit can come before or after an _API
-					MatchIdentifier(TEXT("explicit"), ESearchCase::CaseSensitive);
-				}
-
-				GetToken(ConstructorToken);
-				bSkippedAPIToken = true;
-			}
+			GetToken(ConstructorToken);
+			bSkippedAPIToken = true;
+		}
 
-			if (ConstructorToken.Matches(*FNameLookupCPP::GetNameCPP(Class), ESearchCase::IgnoreCase))
-			{
-				if (TryToMatchConstructorParameterList(ConstructorToken))
-				{
-					return true;
-				}
-			}
-			else if (bSkippedAPIToken)
+		if (ConstructorToken.Matches(*StructDef.GetAlternateNameCPP(), ESearchCase::IgnoreCase))
+		{
+			if (TryToMatchConstructorParameterList(ConstructorToken))
 			{
-				// We skipped over an _API token, but this wasn't a constructor so we need to unget so that subsequent code and still process it
-				UngetToken(ConstructorToken);
+				return true;
 			}
 		}
+		else if (bSkippedAPIToken)
+		{
+			// We skipped over an _API token, but this wasn't a constructor so we need to unget so that subsequent code and still process it
+			UngetToken(ConstructorToken);
+		}
 	}
 
 	// Skip anything that looks like a macro followed by no bracket that we don't know about
@@ -5616,15 +5609,14 @@ FUnrealClassDefinitionInfo& FHeaderParser::ParseClassNameDeclaration(FString& De
 	{
 		if (FUnrealClassDefinitionInfo* BaseClassDef = UHTCast<FUnrealClassDefinitionInfo>(BaseClassInfo.Struct))
 		{
-			UClass* BaseClass = BaseClassDef->GetClass();
-			if (!BaseClass->HasAnyClassFlags(CLASS_Interface))
+			if (!BaseClassDef->HasAnyClassFlags(CLASS_Interface))
 			{
-				FError::Throwf(TEXT("Implements: Class %s is not an interface; Can only inherit from non-UObjects or UInterface derived interfaces"), *BaseClass->GetName());
+				FError::Throwf(TEXT("Implements: Class %s is not an interface; Can only inherit from non-UObjects or UInterface derived interfaces"), *BaseClassDef->GetName());
 			}
 
 			// Propagate the inheritable ClassFlags
-			FoundClass->ClassFlags |= (BaseClass->ClassFlags & CLASS_ScriptInherit);
-			FoundClass->Interfaces.Emplace(BaseClass, 0, false);
+			ClassDef->SetClassFlags(BaseClassDef->GetClassFlags() & CLASS_ScriptInherit);
+			FoundClass->Interfaces.Emplace(BaseClassDef->GetClass(), 0, false);
 		}
 	}
 	return *ClassDef;
@@ -5635,15 +5627,13 @@ FUnrealClassDefinitionInfo& FHeaderParser::ParseClassNameDeclaration(FString& De
  */
 void PostParsingClassSetup(FUnrealClassDefinitionInfo& ClassDef)
 {
-	UClass* Class = ClassDef.GetClass();
-
 	// Cleanup after first pass.
 	FHeaderParser::ComputeFunctionParametersSize(ClassDef);
 
 	// Set all optimization ClassFlags based on property types
-	auto HasAllOptimizationClassFlags = [Class]()
+	auto HasAllOptimizationClassFlags = [&ClassDef]()
 	{
-		return (Class->HasAllClassFlags(CLASS_Config | CLASS_HasInstancedReference));
+		return (ClassDef.HasAllClassFlags(CLASS_Config | CLASS_HasInstancedReference));
 	};
 
 	if (!HasAllOptimizationClassFlags())
@@ -5653,27 +5643,27 @@ void PostParsingClassSetup(FUnrealClassDefinitionInfo& ClassDef)
 			FProperty* Property = PropertyDef->GetProperty();
 			if ((Property->PropertyFlags & CPF_Config) != 0)
 			{
-				Class->ClassFlags |= CLASS_Config;
+				ClassDef.SetClassFlags(CLASS_Config);
 				if (HasAllOptimizationClassFlags()) break;
 			}
 
 			if (Property->ContainsInstancedObjectProperty())
 			{
-				Class->ClassFlags |= CLASS_HasInstancedReference;
+				ClassDef.SetClassFlags(CLASS_HasInstancedReference);
 				if (HasAllOptimizationClassFlags()) break;
 			}
 		}
 	}
 
 	// Class needs to specify which ini file is going to be used if it contains config variables.
-	if ((Class->ClassFlags & CLASS_Config) && (Class->ClassConfigName == NAME_None))
+	if (ClassDef.HasAnyClassFlags(CLASS_Config) && ClassDef.GetClassConfigName() == NAME_None)
 	{
 		// Inherit config setting from base class.
-		Class->ClassConfigName = Class->GetSuperClass() ? Class->GetSuperClass()->ClassConfigName : NAME_None;
-		if (Class->ClassConfigName == NAME_None)
+		ClassDef.SetClassConfigName(ClassDef.GetSuperClass() ? ClassDef.GetSuperClass()->GetClassConfigName() : NAME_None);
+		if (ClassDef.GetClassConfigName() == NAME_None)
 		{
 			FError::Throwf(TEXT("Classes with config / globalconfig member variables need to specify config file."));
-			Class->ClassConfigName = NAME_Engine;
+			ClassDef.SetClassConfigName(NAME_Engine);
 		}
 	}
 }
@@ -5709,33 +5699,35 @@ FUnrealClassDefinitionInfo& FHeaderParser::CompileClassDeclaration()
 	FString RequiredAPIMacroIfPresent;
 	
 	FUnrealClassDefinitionInfo& ClassDef = ParseClassNameDeclaration(/*out*/ DeclaredClassName, /*out*/ RequiredAPIMacroIfPresent);
-	UClass* Class = ClassDef.GetClass();
 
 	ClassDef.GetDefinitionRange().Start = &Input[InputPos];
 
-	check(Class->ClassFlags == 0 || (Class->ClassFlags & ClassDef.ClassFlags) != 0);
+	// If we have existing class flags (preexisting engine objects), make sure we have some shared flags.
+	check(ClassDef.GetClassFlags() == 0 || (ClassDef.GetClassFlags() & ClassDef.GetParsedClassFlags()) != 0);
 
-	Class->ClassFlags |= CLASS_Parsed;
+	const EClassFlags PrevClassFlags = ClassDef.GetInitialEngineClassFlags();
+
+	ClassDef.SetClassFlags(CLASS_Parsed);
 
 	PushNest(ENestType::Class, &ClassDef);
 	
-	const uint32 PrevClassFlags = Class->ClassFlags;
 	ResetClassData();
 
 	// Verify class variables haven't been filled in
+	UClass* Class = ClassDef.GetClass();
 	check(Class->Children == NULL);
 	check(Class->Next == NULL);
 	check(Class->NetFields.Num() == 0);
 	check(Class->FirstOwnedClassRep == 0);
 
 	// Make sure our parent classes is parsed.
-	for (UClass* Temp = Class->GetSuperClass(); Temp; Temp = Temp->GetSuperClass())
+	for (FUnrealClassDefinitionInfo* TempDef = ClassDef.GetSuperClass(); TempDef; TempDef = TempDef->GetSuperClass())
 	{
-		bool bIsParsed = !!(Temp->ClassFlags & CLASS_Parsed);
-		bool bIsIntrinsic = !!(Temp->ClassFlags & CLASS_Intrinsic);
+		bool bIsParsed = TempDef->HasAnyClassFlags(CLASS_Parsed);
+		bool bIsIntrinsic = TempDef->HasAnyClassFlags(CLASS_Intrinsic);
 		if (!(bIsParsed || bIsIntrinsic))
 		{
-			FError::Throwf(TEXT("'%s' can't be compiled: Parent class '%s' has errors"), *Class->GetName(), *Temp->GetName());
+			FError::Throwf(TEXT("'%s' can't be compiled: Parent class '%s' has errors"), *ClassDef.GetName(), *TempDef->GetName());
 		}
 	}
 
@@ -5816,7 +5808,6 @@ FUnrealClassDefinitionInfo* FHeaderParser::ParseInterfaceNameDeclaration(FString
 	{
 		return nullptr;
 	}
-	UClass* FoundClass = ClassDef->GetClass();
 
 	// Get super interface
 	bool bSpecifiesParentClass = MatchSymbol(TEXT(':'));
@@ -5831,24 +5822,21 @@ FUnrealClassDefinitionInfo* FHeaderParser::ParseInterfaceNameDeclaration(FString
 	// the super class should have been marked as CLASS_Interface at the importing stage, if it were an interface
 	FUnrealClassDefinitionInfo* TempClassDef = GetQualifiedClass(TEXT("'extends'"));
 	check(TempClassDef);
-	UClass* TempClass = TempClassDef->GetClass();
-	if( !(TempClass->ClassFlags & CLASS_Interface) )
+	if (!TempClassDef->HasAnyClassFlags(CLASS_Interface))
 	{
 		// UInterface is special and actually extends from UObject, which isn't an interface
 		if (DeclaredInterfaceName != TEXT("UInterface"))
-			FError::Throwf(TEXT("Interface class '%s' cannot inherit from non-interface class '%s'"), *DeclaredInterfaceName, *TempClass->GetName() );
+		{
+			FError::Throwf(TEXT("Interface class '%s' cannot inherit from non-interface class '%s'"), *DeclaredInterfaceName, *TempClassDef->GetName());
+		}
 	}
 
-	UClass* SuperClass = FoundClass->GetSuperClass();
-	if (SuperClass == NULL)
-	{
-		FoundClass->SetSuperStruct(TempClass);
-		ClassDef->GetSuperStructInfo().Struct = TempClassDef;
-		ClassDef->GetSuperStructInfo().Name = TempClass->GetName();
-	}
-	else if (SuperClass != TempClass)
+	// The class super should have already been set by now
+	FUnrealClassDefinitionInfo* SuperClassDef = ClassDef->GetSuperClass();
+	check(SuperClassDef);
+	if (SuperClassDef != TempClassDef)
 	{
-		FError::Throwf(TEXT("%s's superclass must be %s, not %s"), *FoundClass->GetPathName(), *SuperClass->GetPathName(), *TempClass->GetPathName());
+		FError::Throwf(TEXT("%s's superclass must be %s, not %s"), *ClassDef->GetPathName(), *SuperClassDef->GetPathName(), *TempClassDef->GetPathName());
 	}
 
 	return ClassDef;
@@ -5915,32 +5903,32 @@ void FHeaderParser::CompileInterfaceDeclaration()
 
 	// New style files have the interface name / extends afterwards
 	RequireIdentifier(TEXT("class"), ESearchCase::CaseSensitive, TEXT("Interface declaration"));
-	FUnrealClassDefinitionInfo* ClassDef = ParseInterfaceNameDeclaration(/*out*/ DeclaredInterfaceName, /*out*/ RequiredAPIMacroIfPresent);
-	check(ClassDef);
-	UClass* InterfaceClass = ClassDef->GetClass();
-	ClassDef->GetDefinitionRange().Start = &Input[InputPos];
+	FUnrealClassDefinitionInfo* InterfaceClassDef = ParseInterfaceNameDeclaration(/*out*/ DeclaredInterfaceName, /*out*/ RequiredAPIMacroIfPresent);
+	check(InterfaceClassDef);
+	UClass* InterfaceClass = InterfaceClassDef->GetClass();
+	InterfaceClassDef->GetDefinitionRange().Start = &Input[InputPos];
 
 	// Record that this interface is RequiredAPI if the CORE_API style macro was present
 	if (!RequiredAPIMacroIfPresent.IsEmpty())
 	{
-		InterfaceClass->ClassFlags |= CLASS_RequiredAPI;
+		InterfaceClassDef->SetClassFlags(CLASS_RequiredAPI);
 	}
 
 	// Set the appropriate interface class flags
-	InterfaceClass->ClassFlags |= CLASS_Interface | CLASS_Abstract;
+	InterfaceClassDef->SetClassFlags(CLASS_Interface | CLASS_Abstract);
 	if (InterfaceClass->GetSuperClass() != NULL)
 	{
 		InterfaceClass->ClassCastFlags |= InterfaceClass->GetSuperClass()->ClassCastFlags;
 	}
 
 	// All classes that are parsed are expected to be native
-	if (InterfaceClass->GetSuperClass() && !InterfaceClass->GetSuperClass()->HasAnyClassFlags(CLASS_Native))
+	if (InterfaceClassDef->GetSuperClass() && !InterfaceClassDef->GetSuperClass()->HasAnyClassFlags(CLASS_Native))
 	{
 		FError::Throwf(TEXT("Native classes cannot extend non-native classes") );
 	}
 
 	InterfaceClass->SetInternalFlags(EInternalObjectFlags::Native);
-	InterfaceClass->ClassFlags |= CLASS_Native;
+	InterfaceClassDef->SetClassFlags(CLASS_Native);
 
 	// Process all of the interface specifiers
 	for (const FPropertySpecifier& Specifier : SpecifiersFound)
@@ -5961,7 +5949,7 @@ void FHeaderParser::CompileInterfaceDeclaration()
 
 			case EInterfaceSpecifier::MinimalAPI:
 			{
-				InterfaceClass->ClassFlags |= CLASS_MinimalAPI;
+				InterfaceClassDef->SetClassFlags(CLASS_MinimalAPI);
 			}
 			break;
 
@@ -5974,19 +5962,19 @@ void FHeaderParser::CompileInterfaceDeclaration()
 	}
 
 	// All classes must start with a valid Unreal prefix
-	const FString ExpectedInterfaceName = FUnrealTypeDefinitionInfo::GetNameWithPrefix(InterfaceClass, EEnforceInterfacePrefix::U);
+	const FString ExpectedInterfaceName = InterfaceClassDef->GetNameWithPrefix(EEnforceInterfacePrefix::U);
 	if (DeclaredInterfaceName != ExpectedInterfaceName)
 	{
 		FError::Throwf(TEXT("Interface name '%s' is invalid, the first class should be identified as '%s'"), *DeclaredInterfaceName, *ExpectedInterfaceName );
 	}
 
 	// Try parsing metadata for the interface
-	FStructMetaData& StructData = ClassDef->GetStructMetaData();
+	FStructMetaData& StructData = InterfaceClassDef->GetStructMetaData();
 	StructData.SetPrologLine(PrologFinishLine);
 
 	// Register the metadata
-	AddModuleRelativePathToMetadata(*ClassDef, MetaData);
-	AddMetaDataToClassData(*ClassDef, MoveTemp(MetaData));
+	AddModuleRelativePathToMetadata(*InterfaceClassDef, MetaData);
+	AddMetaDataToClassData(*InterfaceClassDef, MoveTemp(MetaData));
 
 	// Handle the start of the rest of the interface
 	RequireSymbol( TEXT('{'), TEXT("'Class'") );
@@ -5998,7 +5986,7 @@ void FHeaderParser::CompileInterfaceDeclaration()
 
 	// Push the interface class nesting.
 	// we need a more specific set of allow flags for ENestType::Interface, only function declaration is allowed, no other stuff are allowed
-	PushNest(ENestType::Interface, ClassDef);
+	PushNest(ENestType::Interface, InterfaceClassDef);
 }
 
 void FHeaderParser::CompileRigVMMethodDeclaration(FUnrealStructDefinitionInfo& StructDef)
@@ -7272,9 +7260,9 @@ void FHeaderParser::ValidatePropertyIsDeprecatedIfNecessary(const FPropertyBase&
 	}
 
 	// check to see if we have a FClassProperty using a deprecated class
-	if (VarProperty.MetaClassDef != nullptr && VarProperty.MetaClassDef->GetClass()->HasAnyClassFlags(CLASS_Deprecated))
+	if (VarProperty.MetaClassDef != nullptr && VarProperty.MetaClassDef->HasAnyClassFlags(CLASS_Deprecated))
 	{
-		UE_LOG_ERROR_UHT(TEXT("Property is using a deprecated class: %s.  Property should be marked deprecated as well."), *VarProperty.MetaClassDef->GetClass()->GetPathName());
+		UE_LOG_ERROR_UHT(TEXT("Property is using a deprecated class: %s.  Property should be marked deprecated as well."), *VarProperty.MetaClassDef->GetPathName());
 	}
 
 	// check to see if we have a FObjectProperty using a deprecated class.
@@ -7282,10 +7270,9 @@ void FHeaderParser::ValidatePropertyIsDeprecatedIfNecessary(const FPropertyBase&
 	if ((VarProperty.Type == CPT_ObjectReference || VarProperty.Type == CPT_WeakObjectReference || VarProperty.Type == CPT_LazyObjectReference || VarProperty.Type == CPT_ObjectPtrReference || VarProperty.Type == CPT_SoftObjectReference)
 		&& VarProperty.ClassDef != nullptr)
 	{
-		UClass* PropertyClass = VarProperty.ClassDef->GetClass();
-		if (PropertyClass->HasAnyClassFlags(CLASS_Deprecated))	// and the object class being used has been deprecated
+		if (VarProperty.ClassDef->HasAnyClassFlags(CLASS_Deprecated))	// and the object class being used has been deprecated
 		{
-			UE_LOG_ERROR_UHT(TEXT("Property is using a deprecated class: %s.  Property should be marked deprecated as well."), *PropertyClass->GetPathName());
+			UE_LOG_ERROR_UHT(TEXT("Property is using a deprecated class: %s.  Property should be marked deprecated as well."), *VarProperty.ClassDef->GetPathName());
 		}
 	}
 }
@@ -8030,7 +8017,7 @@ ECompilationResult::Type FHeaderParser::Parse(
 			FUnrealClassDefinitionInfo& ClassDef = ClassDataPair->AsClassChecked();
 			for (FUnrealClassDefinitionInfo* ParentClassDef = ClassDef.GetSuperClass(); ParentClassDef; ParentClassDef = ParentClassDef->GetSuperClass())
 			{
-				if (ParentClassDef->GetClass()->HasAnyClassFlags(CLASS_Parsed | CLASS_Intrinsic))
+				if (ParentClassDef->HasAnyClassFlags(CLASS_Parsed | CLASS_Intrinsic))
 				{
 					break;
 				}
@@ -8049,8 +8036,7 @@ ECompilationResult::Type FHeaderParser::Parse(
 
 			for (const TSharedRef<FUnrealTypeDefinitionInfo>& ClassDef : SourceFile.GetDefinedClasses())
 			{
-				UClass* Class = ClassDef->AsClassChecked().GetClass();
-				Class->ClassFlags |= CLASS_Parsed;
+				ClassDef->AsClassChecked().SetClassFlags(CLASS_Parsed);
 			}
 		}
 	}
@@ -8624,7 +8610,7 @@ void FHeaderParser::SimplifiedClassParse(FUnrealSourceFile& SourceFile, const TC
 					}
 
 					TSharedRef<FUnrealTypeDefinitionInfo> ClassDecl = Parser.ParseClassDeclaration(SourceFile, StartOfLine + (UInterfaceMacroDecl - Str), CurrentLine, true, TEXT("UINTERFACE"));
-					bFoundExportedClasses |= !(ClassDecl->AsClassChecked().ClassFlags & CLASS_NoExport);
+					bFoundExportedClasses |= !EnumHasAnyFlags(ClassDecl->AsClassChecked().GetParsedClassFlags(), CLASS_NoExport);
 					SourceFile.AddDefinedClass(MoveTemp(ClassDecl));
 				}
 			}
@@ -8639,7 +8625,7 @@ void FHeaderParser::SimplifiedClassParse(FUnrealSourceFile& SourceFile, const TC
 					}
 
 					TSharedRef<FUnrealTypeDefinitionInfo> ClassDecl = Parser.ParseClassDeclaration(SourceFile, StartOfLine + (UClassMacroDecl - Str), CurrentLine, false, TEXT("UCLASS"));
-					bFoundExportedClasses |= !(ClassDecl->AsClassChecked().ClassFlags & CLASS_NoExport);
+					bFoundExportedClasses |= !EnumHasAnyFlags(ClassDecl->AsClassChecked().GetParsedClassFlags(), CLASS_NoExport);
 					SourceFile.AddDefinedClass(MoveTemp(ClassDecl));
 				}
 			}
@@ -9012,12 +8998,12 @@ void FHeaderParser::ResetClassData()
 	CurrentClass->PropertiesSize = 0;
 
 	// Set class flags and within.
-	CurrentClass->ClassFlags &= ~CLASS_RecompilerClear;
+	CurrentClassDef.ClearClassFlags(CLASS_RecompilerClear);
 
 	if (FUnrealClassDefinitionInfo* SuperClassDef = CurrentClassDef.GetSuperClass())
 	{
 		UClass* SuperClass = SuperClassDef->GetClass();
-		CurrentClass->ClassFlags |= (SuperClass->ClassFlags) & CLASS_ScriptInherit;
+		CurrentClassDef.SetClassFlags(SuperClass->ClassFlags & CLASS_ScriptInherit);
 		CurrentClass->ClassConfigName = SuperClass->ClassConfigName;
 		check(SuperClass->ClassWithin);
 		check(SuperClassDef->GetClassWithin());
@@ -9171,7 +9157,6 @@ void FHeaderParser::PostPopNestInterface(FUnrealClassDefinitionInfo& CurrentInte
 {
 	if (CurrentInterfaceDef.ContainsDelegates())
 	{
-		UClass* CurrentInterface = CurrentInterfaceDef.GetClass();
 		TMap<FName, FUnrealFunctionDefinitionInfo*> DelegateCache;
 		FixupDelegateProperties(CurrentInterfaceDef, *CurrentInterfaceDef.GetScope(), DelegateCache);
 	}
@@ -9319,7 +9304,6 @@ void FHeaderParser::CheckDocumentationPolicyForStruct(FUnrealStructDefinitionInf
 	{
 		if (FUnrealClassDefinitionInfo* ClassDef = UHTCast<FUnrealClassDefinitionInfo>(StructDef))
 		{
-			UClass* Class = ClassDef->GetClass();
 			TMap<FString, FName> ToolTipToFunc;
 			for (FUnrealFunctionDefinitionInfo* FunctionDef : ClassDef->GetFunctions())
 			{
@@ -9327,13 +9311,13 @@ void FHeaderParser::CheckDocumentationPolicyForStruct(FUnrealStructDefinitionInf
 				FString ToolTip = Func->GetToolTipText().ToString();
 				if (ToolTip.IsEmpty())
 				{
-					UE_LOG_ERROR_UHT(TEXT("Function '%s::%s' does not provide a tooltip / comment (DocumentationPolicy)."), *Class->GetName(), *FunctionDef->GetName());
+					UE_LOG_ERROR_UHT(TEXT("Function '%s::%s' does not provide a tooltip / comment (DocumentationPolicy)."), *ClassDef->GetName(), *FunctionDef->GetName());
 					continue;
 				}
 				const FName* ExistingFuncName = ToolTipToFunc.Find(ToolTip);
 				if (ExistingFuncName != nullptr)
 				{
-					UE_LOG_ERROR_UHT(TEXT("Functions '%s::%s' and '%s::%s' uses identical tooltips / comments (DocumentationPolicy)."), *Class->GetName(), *(*ExistingFuncName).ToString(), *Class->GetName(), *FunctionDef->GetName());
+					UE_LOG_ERROR_UHT(TEXT("Functions '%s::%s' and '%s::%s' uses identical tooltips / comments (DocumentationPolicy)."), *ClassDef->GetName(), *(*ExistingFuncName).ToString(), *ClassDef->GetName(), *FunctionDef->GetName());
 				}
 				ToolTipToFunc.Add(MoveTemp(ToolTip), FunctionDef->GetFName());
 			}
@@ -9351,15 +9335,13 @@ void FHeaderParser::CheckDocumentationPolicyForFunc(FUnrealClassDefinitionInfo&
 {
 	SCOPE_SECONDS_COUNTER_UHT(DocumentationPolicy);
 
-	UClass* Class = ClassDef.GetClass();
-
 	FDocumentationPolicy DocumentationPolicy = GetDocumentationPolicyForStruct(ClassDef);
 	if (DocumentationPolicy.bFunctionToolTipsRequired)
 	{
 		const FString* FunctionTooltip = FunctionDef.FindMetaData(NAME_ToolTip);
 		if (FunctionTooltip == nullptr)
 		{
-			UE_LOG_ERROR_UHT(TEXT("Function '%s::%s' does not provide a tooltip / comment (DocumentationPolicy)."), *Class->GetName(), *FunctionDef.GetName());
+			UE_LOG_ERROR_UHT(TEXT("Function '%s::%s' does not provide a tooltip / comment (DocumentationPolicy)."), *ClassDef.GetName(), *FunctionDef.GetName());
 		}
 	}
 
@@ -9368,7 +9350,7 @@ void FHeaderParser::CheckDocumentationPolicyForFunc(FUnrealClassDefinitionInfo&
 		const FString* FunctionComment = FunctionDef.FindMetaData(NAME_Comment);
 		if (FunctionComment == nullptr)
 		{
-			UE_LOG_ERROR_UHT(TEXT("Function '%s::%s' does not provide a comment (DocumentationPolicy)."), *Class->GetName(), *FunctionDef.GetName());
+			UE_LOG_ERROR_UHT(TEXT("Function '%s::%s' does not provide a comment (DocumentationPolicy)."), *ClassDef.GetName(), *FunctionDef.GetName());
 			return;
 		}
 		
@@ -9399,7 +9381,7 @@ void FHeaderParser::CheckDocumentationPolicyForFunc(FUnrealClassDefinitionInfo&
 				const FString* ParamToolTip = ParamToolTips.Find(ParamName);
 				if (ParamToolTip == nullptr)
 				{
-					UE_LOG_ERROR_UHT(TEXT("Function '%s::%s' doesn't provide a tooltip for parameter '%s' (DocumentationPolicy)."), *Class->GetName(), *FunctionDef.GetName(), *ParamName.ToString());
+					UE_LOG_ERROR_UHT(TEXT("Function '%s::%s' doesn't provide a tooltip for parameter '%s' (DocumentationPolicy)."), *ClassDef.GetName(), *FunctionDef.GetName(), *ParamName.ToString());
 				}
 				ExistingFields.Add(ParamName);
 			}
@@ -9414,7 +9396,7 @@ void FHeaderParser::CheckDocumentationPolicyForFunc(FUnrealClassDefinitionInfo&
 				}
 				if (!ExistingFields.Contains(ParamName))
 				{
-					UE_LOG_ERROR_UHT(TEXT("Function '%s::%s' provides a tooltip for an unknown parameter '%s' (DocumentationPolicy)."), *Class->GetName(), *FunctionDef.GetName(), *Pair.Key.ToString());
+					UE_LOG_ERROR_UHT(TEXT("Function '%s::%s' provides a tooltip for an unknown parameter '%s' (DocumentationPolicy)."), *ClassDef.GetName(), *FunctionDef.GetName(), *Pair.Key.ToString());
 				}
 			}
 
@@ -9430,7 +9412,7 @@ void FHeaderParser::CheckDocumentationPolicyForFunc(FUnrealClassDefinitionInfo&
 				const FName* ExistingParam = ToolTipToParam.Find(Pair.Value);
 				if (ExistingParam != nullptr)
 				{
-					UE_LOG_ERROR_UHT(TEXT("Function '%s::%s' uses identical tooltips for parameters '%s' and '%s' (DocumentationPolicy)."), *Class->GetName(), *FunctionDef.GetName(), *ExistingParam->ToString(), *Pair.Key.ToString());
+					UE_LOG_ERROR_UHT(TEXT("Function '%s::%s' uses identical tooltips for parameters '%s' and '%s' (DocumentationPolicy)."), *ClassDef.GetName(), *FunctionDef.GetName(), *ExistingParam->ToString(), *Pair.Key.ToString());
 				}
 				ToolTipToParam.Add(MoveTemp(Pair.Value), Pair.Key);
 			}
diff --git a/Engine/Source/Programs/UnrealHeaderTool/Private/UnrealTypeDefinitionInfo.cpp b/Engine/Source/Programs/UnrealHeaderTool/Private/UnrealTypeDefinitionInfo.cpp
--- a/Engine/Source/Programs/UnrealHeaderTool/Private/UnrealTypeDefinitionInfo.cpp
+++ b/Engine/Source/Programs/UnrealHeaderTool/Private/UnrealTypeDefinitionInfo.cpp
@@ -116,7 +116,7 @@ namespace
 	 * @param InNameToCheck - Name w/ potential prefix to check
 	 * @param OriginalClass - Class to check against
 	 */
-	bool ClassNameHasValidPrefix(const FString& InNameToCheck, const UClass* OriginalClass)
+	bool ClassNameHasValidPrefix(const FString& InNameToCheck, const FUnrealClassDefinitionInfo& OriginalClass)
 	{
 		bool bIsLabledDeprecated;
 		GetClassPrefix(InNameToCheck, bIsLabledDeprecated);
@@ -127,14 +127,14 @@ namespace
 			return true;
 		}
 
-		const FString OriginalClassName = FUnrealTypeDefinitionInfo::GetNameWithPrefix(OriginalClass);
+		const FString OriginalClassName = OriginalClass.GetNameWithPrefix();
 
 		bool bNamesMatch = (InNameToCheck == OriginalClassName);
 
 		if (!bNamesMatch)
 		{
 			//@TODO: UCREMOVAL: I/U interface hack - Ignoring prefixing for this call
-			if (OriginalClass->HasAnyClassFlags(CLASS_Interface))
+			if (OriginalClass.HasAnyClassFlags(CLASS_Interface))
 			{
 				bNamesMatch = InNameToCheck.Mid(1) == OriginalClassName.Mid(1);
 			}
@@ -238,49 +238,6 @@ bool FUnrealTypeDefinitionInfo::IsDynamic(const FField* Field)
 	return Field->HasMetaData(NAME_ReplaceConverted);
 }
 
-
-FString FUnrealTypeDefinitionInfo::GetNameWithPrefix(const UClass* InClass, EEnforceInterfacePrefix EnforceInterfacePrefix)
-{
-	const TCHAR* Prefix = 0;
-
-	if (InClass->HasAnyClassFlags(CLASS_Interface))
-	{
-		// Grab the expected prefix for interfaces (U on the first one, I on the second one)
-		switch (EnforceInterfacePrefix)
-		{
-		case EEnforceInterfacePrefix::None:
-			// For old-style files: "I" for interfaces, unless it's the actual "Interface" class, which gets "U"
-			if (InClass->GetFName() == NAME_Interface)
-			{
-				Prefix = TEXT("U");
-			}
-			else
-			{
-				Prefix = TEXT("I");
-			}
-			break;
-
-		case EEnforceInterfacePrefix::I:
-			Prefix = TEXT("I");
-			break;
-
-		case EEnforceInterfacePrefix::U:
-			Prefix = TEXT("U");
-			break;
-
-		default:
-			check(false);
-		}
-	}
-	else
-	{
-		// Get the expected class name with prefix
-		Prefix = InClass->GetPrefixCPP();
-	}
-
-	return FString::Printf(TEXT("%s%s"), Prefix, *InClass->GetName());
-}
-
 void FUnrealPropertyDefinitionInfo::PostParseFinalize()
 {
 	TypePackageName = GetTypePackageName_Inner(GetProperty());
@@ -594,11 +551,11 @@ FUnrealClassDefinitionInfo* FUnrealClassDefinitionInfo::FindScriptClass(const FS
 	if (FUnrealClassDefinitionInfo* FoundClassDef = FindClass(*ClassNameStripped))
 	{
 		// If the class was found with the stripped class name, verify that the correct prefix was used and throw an error otherwise
-		if (!ClassNameHasValidPrefix(InClassName, FoundClassDef->GetClass()))
+		if (!ClassNameHasValidPrefix(InClassName, *FoundClassDef))
 		{
 			if (OutErrorMsg)
 			{
-				*OutErrorMsg = FString::Printf(TEXT("Class '%s' has an incorrect prefix, expecting '%s'"), *InClassName, *FUnrealTypeDefinitionInfo::GetNameWithPrefix(FoundClassDef->GetClass()));
+				*OutErrorMsg = FString::Printf(TEXT("Class '%s' has an incorrect prefix, expecting '%s'"), *InClassName, *FoundClassDef->GetNameWithPrefix());
 			}
 			return nullptr;
 		}
@@ -613,7 +570,7 @@ FUnrealClassDefinitionInfo* FUnrealClassDefinitionInfo::FindScriptClass(const FS
 		// If the class was found with the given identifier, the user forgot to use the correct Unreal prefix	
 		if (OutErrorMsg)
 		{
-			*OutErrorMsg = FString::Printf(TEXT("Class '%s' is missing a prefix, expecting '%s'"), *InClassName, *FUnrealTypeDefinitionInfo::GetNameWithPrefix(FoundClassDef->GetClass()));
+			*OutErrorMsg = FString::Printf(TEXT("Class '%s' is missing a prefix, expecting '%s'"), *InClassName, *FoundClassDef->GetNameWithPrefix());
 		}
 	}
 	else
@@ -656,13 +613,13 @@ void FUnrealClassDefinitionInfo::PostParseFinalize()
 
 void FUnrealClassDefinitionInfo::ParseClassProperties(TArray<FPropertySpecifier>&& InClassSpecifiers, const FString& InRequiredAPIMacroIfPresent)
 {
-	ClassFlags = CLASS_None;
+	ParsedClassFlags = CLASS_None;
 	// Record that this class is RequiredAPI if the CORE_API style macro was present
 	if (!InRequiredAPIMacroIfPresent.IsEmpty())
 	{
-		ClassFlags |= CLASS_RequiredAPI;
+		ParsedClassFlags |= CLASS_RequiredAPI;
 	}
-	ClassFlags |= CLASS_Native;
+	ParsedClassFlags |= CLASS_Native;
 
 	// Process all of the class specifiers
 
@@ -673,12 +630,11 @@ void FUnrealClassDefinitionInfo::ParseClassProperties(TArray<FPropertySpecifier>
 		case EClassMetadataSpecifier::NoExport:
 
 			// Don't export to C++ header.
-			ClassFlags |= CLASS_NoExport;
+			ParsedClassFlags |= CLASS_NoExport;
 			break;
 
 		case EClassMetadataSpecifier::Intrinsic:
-
-			ClassFlags |= CLASS_Intrinsic;
+			ParsedClassFlags |= CLASS_Intrinsic;
 			break;
 
 		case EClassMetadataSpecifier::ComponentWrapperClass:
@@ -694,37 +650,37 @@ void FUnrealClassDefinitionInfo::ParseClassProperties(TArray<FPropertySpecifier>
 		case EClassMetadataSpecifier::EditInlineNew:
 
 			// Class can be constructed from the New button in editinline
-			ClassFlags |= CLASS_EditInlineNew;
+			ParsedClassFlags |= CLASS_EditInlineNew;
 			break;
 
 		case EClassMetadataSpecifier::NotEditInlineNew:
 
 			// Class cannot be constructed from the New button in editinline
-			ClassFlags &= ~CLASS_EditInlineNew;
+			ParsedClassFlags &= ~CLASS_EditInlineNew;
 			break;
 
 		case EClassMetadataSpecifier::Placeable:
 
 			bWantsToBePlaceable = true;
-			ClassFlags &= ~CLASS_NotPlaceable;
+			ParsedClassFlags &= ~CLASS_NotPlaceable;
 			break;
 
 		case EClassMetadataSpecifier::DefaultToInstanced:
 
 			// these classed default to instanced.
-			ClassFlags |= CLASS_DefaultToInstanced;
+			ParsedClassFlags |= CLASS_DefaultToInstanced;
 			break;
 
 		case EClassMetadataSpecifier::NotPlaceable:
 
 			// Don't allow the class to be placed in the editor.
-			ClassFlags |= CLASS_NotPlaceable;
+			ParsedClassFlags |= CLASS_NotPlaceable;
 			break;
 
 		case EClassMetadataSpecifier::HideDropdown:
 
 			// Prevents class from appearing in class comboboxes in the property window
-			ClassFlags |= CLASS_HideDropDown;
+			ParsedClassFlags |= CLASS_HideDropDown;
 			break;
 
 		case EClassMetadataSpecifier::DependsOn:
@@ -734,55 +690,55 @@ void FUnrealClassDefinitionInfo::ParseClassProperties(TArray<FPropertySpecifier>
 
 		case EClassMetadataSpecifier::MinimalAPI:
 
-			ClassFlags |= CLASS_MinimalAPI;
+			ParsedClassFlags |= CLASS_MinimalAPI;
 			break;
 
 		case EClassMetadataSpecifier::Const:
 
-			ClassFlags |= CLASS_Const;
+			ParsedClassFlags |= CLASS_Const;
 			break;
 
 		case EClassMetadataSpecifier::PerObjectConfig:
 
-			ClassFlags |= CLASS_PerObjectConfig;
+			ParsedClassFlags |= CLASS_PerObjectConfig;
 			break;
 
 		case EClassMetadataSpecifier::ConfigDoNotCheckDefaults:
 
-			ClassFlags |= CLASS_ConfigDoNotCheckDefaults;
+			ParsedClassFlags |= CLASS_ConfigDoNotCheckDefaults;
 			break;
 
 		case EClassMetadataSpecifier::Abstract:
 
 			// Hide all editable properties.
-			ClassFlags |= CLASS_Abstract;
+			ParsedClassFlags |= CLASS_Abstract;
 			break;
 
 		case EClassMetadataSpecifier::Deprecated:
 
-			ClassFlags |= CLASS_Deprecated;
+			ParsedClassFlags |= CLASS_Deprecated;
 
 			// Don't allow the class to be placed in the editor.
-			ClassFlags |= CLASS_NotPlaceable;
+			ParsedClassFlags |= CLASS_NotPlaceable;
 
 			break;
 
 		case EClassMetadataSpecifier::Transient:
 
 			// Transient class.
-			ClassFlags |= CLASS_Transient;
+			ParsedClassFlags |= CLASS_Transient;
 			break;
 
 		case EClassMetadataSpecifier::NonTransient:
 
 			// this child of a transient class is not transient - remove the transient flag
-			ClassFlags &= ~CLASS_Transient;
+			ParsedClassFlags &= ~CLASS_Transient;
 			break;
 
 		case EClassMetadataSpecifier::CustomConstructor:
 
 			// we will not export a constructor for this class, assuming it is in the CPP block
-			ClassFlags |= CLASS_CustomConstructor;
+			ParsedClassFlags |= CLASS_CustomConstructor;
 			break;
 
 		case EClassMetadataSpecifier::Config:
@@ -794,19 +750,19 @@ void FUnrealClassDefinitionInfo::ParseClassProperties(TArray<FPropertySpecifier>
 		case EClassMetadataSpecifier::DefaultConfig:
 
 			// Save object config only to Default INIs, never to local INIs.
-			ClassFlags |= CLASS_DefaultConfig;
+			ParsedClassFlags |= CLASS_DefaultConfig;
 			break;
 
 		case EClassMetadataSpecifier::GlobalUserConfig:
 
 			// Save object config only to global user overrides, never to local INIs
-			ClassFlags |= CLASS_GlobalUserConfig;
+			ParsedClassFlags |= CLASS_GlobalUserConfig;
 			break;
 
 		case EClassMetadataSpecifier::ProjectUserConfig:
 
 			// Save object config only to project user overrides, never to INIs that are checked in
-			ClassFlags |= CLASS_ProjectUserConfig;
+			ParsedClassFlags |= CLASS_ProjectUserConfig;
 			break;
 
 		case EClassMetadataSpecifier::ShowCategories:
@@ -900,13 +856,13 @@ void FUnrealClassDefinitionInfo::ParseClassProperties(TArray<FPropertySpecifier>
 		case EClassMetadataSpecifier::CollapseCategories:
 
 			// Class' properties should not be shown categorized in the editor.
-			ClassFlags |= CLASS_CollapseCategories;
+			ParsedClassFlags |= CLASS_CollapseCategories;
 			break;
 
 		case EClassMetadataSpecifier::DontCollapseCategories:
 
 			// Class' properties should be shown categorized in the editor.
-			ClassFlags &= ~CLASS_CollapseCategories;
+			ParsedClassFlags &= ~CLASS_CollapseCategories;
 			break;
 
 		case EClassMetadataSpecifier::AdvancedClassDisplay:
@@ -922,7 +878,7 @@ void FUnrealClassDefinitionInfo::ParseClassProperties(TArray<FPropertySpecifier>
 
 		case EClassMetadataSpecifier::NeedsDeferredDependencyLoading:
 
-			ClassFlags |= CLASS_NeedsDeferredDependencyLoading;
+			ParsedClassFlags |= CLASS_NeedsDeferredDependencyLoading;
 			break;
 
 		default:
@@ -1021,7 +977,7 @@ void FUnrealClassDefinitionInfo::MergeClassCategories()
 	AutoExpandCategories.Append(MoveTemp(ParentAutoExpandCategories));
 }
 
-void FUnrealClassDefinitionInfo::MergeAndValidateClassFlags(const FString& DeclaredClassName, uint32 PreviousClassFlags)
+void FUnrealClassDefinitionInfo::MergeAndValidateClassFlags(const FString& DeclaredClassName, EClassFlags PreviousClassFlags)
 {
 	UClass* Class = GetClass();
 	if (bWantsToBePlaceable)
@@ -1035,7 +991,7 @@ void FUnrealClassDefinitionInfo::MergeAndValidateClassFlags(const FString& Decla
 	}
 
 	// Now merge all remaining flags/properties
-	Class->ClassFlags |= ClassFlags;
+	Class->ClassFlags |= ParsedClassFlags;
 	Class->ClassConfigName = FName(*ConfigName);
 
 	SetAndValidateWithinClass();
@@ -1057,7 +1013,7 @@ void FUnrealClassDefinitionInfo::MergeAndValidateClassFlags(const FString& Decla
 	}
 
 	// All classes must start with a valid Unreal prefix
-	const FString ExpectedClassName = GetNameWithPrefix(GetClass());
+	const FString ExpectedClassName = GetNameWithPrefix();
 	if (DeclaredClassName != ExpectedClassName)
 	{
 		FError::Throwf(TEXT("Class name '%s' is invalid, should be identified as '%s'"), *DeclaredClassName, *ExpectedClassName);
@@ -1209,6 +1165,48 @@ void FUnrealClassDefinitionInfo::GetSparseClassDataTypes(TArray<FString>& OutSpa
 	}
 }
 
+FString FUnrealClassDefinitionInfo::GetNameWithPrefix(EEnforceInterfacePrefix EnforceInterfacePrefix) const
+{
+	const TCHAR* Prefix = 0;
+
+	if (HasAnyClassFlags(CLASS_Interface))
+	{
+		// Grab the expected prefix for interfaces (U on the first one, I on the second one)
+		switch (EnforceInterfacePrefix)
+		{
+		case EEnforceInterfacePrefix::None:
+			// For old-style files: "I" for interfaces, unless it's the actual "Interface" class, which gets "U"
+			if (GetFName() == NAME_Interface)
+			{
+				Prefix = TEXT("U");
+			}
+			else
+			{
+				Prefix = TEXT("I");
+			}
+			break;
+
+		case EEnforceInterfacePrefix::I:
+			Prefix = TEXT("I");
+			break;
+
+		case EEnforceInterfacePrefix::U:
+			Prefix = TEXT("U");
+			break;
+
+		default:
+			check(false);
+		}
+	}
+	else
+	{
+		// Get the expected class name with prefix
+		Prefix = GetPrefixCPP();
+	}
+
+	return FString::Printf(TEXT("%s%s"), Prefix, *GetName());
+}
+
 void FUnrealFunctionDefinitionInfo::AddProperty(FUnrealPropertyDefinitionInfo& PropertyDef)
 {
 	const FProperty* Prop = PropertyDef.GetProperty();
diff --git a/Engine/Source/Programs/UnrealHeaderTool/Private/UnrealTypeDefinitionInfo.h b/Engine/Source/Programs/UnrealHeaderTool/Private/UnrealTypeDefinitionInfo.h
--- a/Engine/Source/Programs/UnrealHeaderTool/Private/UnrealTypeDefinitionInfo.h
+++ b/Engine/Source/Programs/UnrealHeaderTool/Private/UnrealTypeDefinitionInfo.h
@@ -343,8 +343,6 @@ public:
 		return Outer;
 	}
 
-	static FString GetNameWithPrefix(const UClass* Class, EEnforceInterfacePrefix EnforceInterfacePrefix = EEnforceInterfacePrefix::None);
-
 protected:
 	explicit FUnrealTypeDefinitionInfo(FString&& InNameCPP)
 		: NameCPP(MoveTemp(InNameCPP))
@@ -623,6 +621,20 @@ public:
 		return GetObject()->GetFName();
 	}
 
+	/**
+	 * Returns the fully qualified pathname for this object, in the format:
+	 * 'Outermost.[Outer:]Name'
+	 *
+	 * @param	StopOuter	if specified, indicates that the output string should be relative to this object.  if StopOuter
+	 *						does not exist in this object's Outer chain, the result would be the same as passing NULL.
+	 *
+	 * @note	safe to call on NULL object pointers!
+	 */
+	FString GetPathName(const FUnrealObjectDefinitionInfo* StopOuter = nullptr) const
+	{
+		return GetObject()->GetPathName(StopOuter ? StopOuter->GetObject() : nullptr);
+	}
+
 	/**
 	 * Return the package associated with this object
 	 */
@@ -1194,6 +1206,14 @@ public:
 		return static_cast<UStruct*>(GetObject());
 	}
 
+	/**
+	 * Return the CPP prefix 
+	 */
+	const TCHAR* GetPrefixCPP() const
+	{
+		return GetStruct()->GetPrefixCPP();
+	}
+
 	/**
 	 * Returns the name used for declaring the passed in struct in C++
 	 * 
@@ -1204,8 +1224,7 @@ public:
 	 */
 	FString GetAlternateNameCPP(bool bForceInterface = false)
 	{
-		UStruct* Struct = GetStruct();
-		return FString::Printf(TEXT("%s%s"), (bForceInterface ? TEXT("I") : Struct->GetPrefixCPP()), *Struct->GetName());
+		return FString::Printf(TEXT("%s%s"), (bForceInterface ? TEXT("I") : GetPrefixCPP()), *GetName());
 	}
 
 	/**
@@ -1705,6 +1724,17 @@ public:
 		return static_cast<UClass*>(GetObject());
 	}
 
+	/**
+	 * Set the Engine instance associated with the compiler instance
+	 */
+	virtual void SetObject(UObject* InObject) override
+	{
+		UClass* Class = static_cast<UClass*>(InObject);
+		InitialEngineClassFlags = Class->ClassFlags;
+		Class->ClassFlags |= ParsedClassFlags;
+		FUnrealStructDefinitionInfo::SetObject(InObject);
+	}
+
 	FUnrealClassDefinitionInfo* GetSuperClass() const
 	{
 		if (FUnrealStructDefinitionInfo* SuperClass = GetSuperStructInfo().Struct)
@@ -1754,6 +1784,30 @@ public:
 		return bIsInterface;
 	}
 
+	/**
+	 * Return the flags that were parsed as part of the pre-parse phase
+	 */
+	EClassFlags GetClassFlags() const
+	{
+		return GetClass()->ClassFlags;
+	}
+
+	/**
+	 * Sets the given flags
+	 */
+	void SetClassFlags(EClassFlags FlagsToSet)
+	{
+		GetClass()->ClassFlags |= FlagsToSet;
+	}
+
+	/**
+	 * Clears the given flags
+	 */
+	void ClearClassFlags(EClassFlags FlagsToClear)
+	{
+		GetClass()->ClassFlags &= ~FlagsToClear;
+	}
+
 	/**
 	 * Used to safely check whether the passed in flag is set.
 	 *
@@ -1761,9 +1815,10 @@ public:
 	 * @return	true if the passed in flag is set, false otherwise
 	 *			(including no flag passed in, unless the FlagsToCheck is CLASS_AllFlags)
 	 */
-	FORCEINLINE bool HasAnyClassFlags(EClassFlags FlagsToCheck) const
+	bool HasAnyClassFlags(EClassFlags FlagsToCheck) const
 	{
-		return EnumHasAnyFlags(ClassFlags, FlagsToCheck) != 0 || GetClass()->HasAnyClassFlags(FlagsToCheck);
+
+		return EnumHasAnyFlags(ParsedClassFlags, FlagsToCheck) || GetClass()->HasAnyClassFlags(FlagsToCheck);
 	}
 
 	/**
@@ -1772,19 +1827,9 @@ public:
 	 * @param FlagsToCheck	Class flags to check for
 	 * @return true if all of the passed in flags are set (including no flags passed in), false otherwise
 	 */
-	FORCEINLINE bool HasAllClassFlags(EClassFlags FlagsToCheck) const
-	{
-		return EnumHasAllFlags(ClassFlags, FlagsToCheck) || GetClass()->HasAllClassFlags(FlagsToCheck);
-	}
-
-	/**
-	 * Gets the class flags.
-	 *
-	 * @return	The class flags.
-	 */
-	FORCEINLINE EClassFlags GetClassFlags() const
+	bool HasAllClassFlags(EClassFlags FlagsToCheck) const
 	{
-		return ClassFlags;
+		return EnumHasAllFlags(ParsedClassFlags, FlagsToCheck) || GetClass()->HasAllClassFlags(FlagsToCheck);
 	}
 
 	/**
@@ -1824,6 +1869,23 @@ public:
 		return false;
 	}
 
+	/**
+	 * Return the flags that were parsed as part of the pre-parse phase
+	 */
+	EClassFlags GetParsedClassFlags() const
+	{
+		return ParsedClassFlags;
+	}
+
+	/**
+	 * Return the initial flags from the class. These will be set for class definitions created directly from
+	 * existing engine types.
+	 */
+	EClassFlags GetInitialEngineClassFlags() const
+	{
+		return InitialEngineClassFlags;
+	}
+
 	/**
 	* Parse Class's properties to generate its declaration data.
 	*
@@ -1844,7 +1906,7 @@ public:
 	* @param	DeclaredClassName Name this class was declared with (for validation)
 	* @param	PreviousClassFlags Class flags before resetting the class (for validation)
 	*/
-	void MergeAndValidateClassFlags(const FString& DeclaredClassName, uint32 PreviousClassFlags);
+	void MergeAndValidateClassFlags(const FString& DeclaredClassName, EClassFlags PreviousClassFlags);
 
 	/**
 	 * Add the category meta data
@@ -1869,8 +1931,25 @@ public:
 		ClassWithin = InClassWithin;
 	}
 
+	/**
+	 * Get the class config name
+	 */
+	FName GetClassConfigName() const
+	{
+		return GetClass()->ClassConfigName;
+	}
+
+	/**
+	 * Set the class config name
+	 */
+	void SetClassConfigName(FName InClassConfigName)
+	{
+		GetClass()->ClassConfigName = InClassConfigName;
+	}
+
+	FString GetNameWithPrefix(EEnforceInterfacePrefix EnforceInterfacePrefix = EEnforceInterfacePrefix::None) const;
+
 public:
-	EClassFlags ClassFlags = CLASS_None;
 	TMap<FName, FString> MetaData;
 
 private:
@@ -1898,6 +1977,8 @@ private:
 	FString EnclosingDefine;
 	FString ClassWithinStr;
 	FString ConfigName;
+	EClassFlags ParsedClassFlags = CLASS_None;
+	EClassFlags InitialEngineClassFlags = CLASS_None;
 	FUnrealClassDefinitionInfo* ClassWithin = nullptr;
 	ESerializerArchiveType ArchiveType = ESerializerArchiveType::None;
 	bool bIsInterface = false;
