Remove all GetProperty calls

#rnx

diff --git a/Engine/Source/Programs/UnrealHeaderTool/Private/CodeGenerator.cpp b/Engine/Source/Programs/UnrealHeaderTool/Private/CodeGenerator.cpp
--- a/Engine/Source/Programs/UnrealHeaderTool/Private/CodeGenerator.cpp
+++ b/Engine/Source/Programs/UnrealHeaderTool/Private/CodeGenerator.cpp
@@ -177,7 +177,7 @@ namespace
 		{
 			for (FUnrealPropertyDefinitionInfo* PropertyDef : ClassDef.GetProperties())
 			{
-				if ((PropertyDef->GetProperty()->PropertyFlags & CPF_Net) != 0)
+				if (PropertyDef->HasAnyPropertyFlags(CPF_Net))
 				{
 					return true;
 				}
@@ -516,12 +516,11 @@ FParmsAndReturnProperties GetFunctionParmsAndReturn(FUnrealFunctionDefinitionInf
 	FParmsAndReturnProperties Result;
 	for (FUnrealPropertyDefinitionInfo* PropertyDef : FunctionDef.GetProperties())
 	{
-		FProperty* Field = PropertyDef->GetProperty();
-		if ((Field->PropertyFlags & (CPF_Parm | CPF_ReturnParm)) == CPF_Parm)
+		if (PropertyDef->HasSpecificPropertyFlags(CPF_Parm | CPF_ReturnParm, CPF_Parm))
 		{
 			Result.Parms.Add(PropertyDef);
 		}
-		else if (Field->PropertyFlags & CPF_ReturnParm)
+		else if (PropertyDef->HasAnyPropertyFlags(CPF_ReturnParm))
 		{
 			Result.Return = PropertyDef;
 		}
@@ -652,58 +651,10 @@ FString CreateUTF8LiteralString(const FString& Str)
 	return Result;
 }
 
-TMap<FName, FString> GenerateMetadataMapForObject(const UObject* Obj)
-{
-	check(Obj);
-	UPackage* Package = Obj->GetOutermost();
-	check(Package);
-	UMetaData* Metadata = Package->GetMetaData();
-	check(Metadata);
-
-	TMap<FName, FString>* PackageMap = Metadata->ObjectMetaDataMap.Find(Obj);
-	TMap<FName, FString> Map;
-	if (PackageMap)
-	{
-		for (const TPair<FName, FString>& MetaKeyValue : *PackageMap)
-		{
-			FString Key = MetaKeyValue.Key.ToString();
-			if (!Key.StartsWith(TEXT("/Script")))
-			{
-				Map.Add(MetaKeyValue.Key, MetaKeyValue.Value);
-			}
-		}
-	}
-	return Map;
-}
-
-TMap<FName, FString> GenerateMetadataMapForField(const FField* Field)
-{
-	TMap<FName, FString> MetaDataMap;
-	const TMap<FName, FString>* FieldMetaDataMap = Field->GetMetaDataMap();
-	if (FieldMetaDataMap)
-	{
-		MetaDataMap = *FieldMetaDataMap;
-	}
-	return MetaDataMap;
-}
-
 // Returns the METADATA_PARAMS for this output
 static FString OutputMetaDataCodeForObject(FOutputDevice& OutDeclaration, FOutputDevice& Out, FUnrealTypeDefinitionInfo& TypeDef, const TCHAR* MetaDataBlockName, const TCHAR* DeclSpaces, const TCHAR* Spaces)
 {
-	TMap<FName, FString> MetaData;
-	
-	if (FUnrealObjectDefinitionInfo* ObjectDef = UHTCast<FUnrealObjectDefinitionInfo>(TypeDef))
-	{
-		MetaData = GenerateMetadataMapForObject(ObjectDef->GetObject());
-	}
-	else if (FUnrealPropertyDefinitionInfo* PropertyDef = UHTCast<FUnrealPropertyDefinitionInfo>(TypeDef))
-	{
-		MetaData = GenerateMetadataMapForField(PropertyDef->GetProperty());
-	}
-	else
-	{
-		FError::Throwf(TEXT("Unknown type definition"));
-	}
+	TMap<FName, FString> MetaData = TypeDef.GenerateMetadataMap();
 
 	FString Result;
 	if (MetaData.Num())
@@ -759,14 +710,12 @@ void FNativeClassHeaderGenerator::ExportProperties(FOutputDevice& Out, FUnrealSt
 	// Iterate over all properties in this struct.
 	for (FUnrealPropertyDefinitionInfo* PropertyDef : StructDef.GetProperties())
 	{
-		FProperty* Property = PropertyDef->GetProperty();
-
-		WithEditorOnlyData(Property->IsEditorOnlyProperty());
+		WithEditorOnlyData(PropertyDef->IsEditorOnlyProperty());
 
 		// Export property specifiers
 		// Indent code and export CPP text.
 		FUHTStringBuilder JustPropertyDecl;
-		Property->ExportCppDeclaration( JustPropertyDecl, EExportedDeclaration::Member, PropertyDef->GetArrayDimensions());
+		PropertyDef->ExportCppDeclaration( JustPropertyDecl, EExportedDeclaration::Member, PropertyDef->GetArrayDimensions());
 		ApplyAlternatePropertyExportText(*PropertyDef, JustPropertyDecl, EExportingState::TypeEraseDelegates);
 
 		// Finish up line.
@@ -804,142 +753,83 @@ const FString& FNativeClassHeaderGenerator::GetSingletonNameFuncAddr(FUnrealFiel
 void FNativeClassHeaderGenerator::GetPropertyTag(FUHTStringBuilder& Out, FUnrealPropertyDefinitionInfo& PropDef)
 {
 	const FPropertyBase& PropertyBase = PropDef.GetPropertyBase();
-	FProperty* Prop = PropDef.GetProperty();
-
-#if UHT_ENABLE_VALUE_PROPERTY_TAG
-	if (FNumericProperty* TypedProp = CastField<FNumericProperty>(Prop))
-	{
-		if (PropertyBase.EnumDef != nullptr)
-		{
-			PropertyBase.EnumDef->GetHashTag(Out);
-		}
-		return;
-	}
 
-	if (FEnumProperty* TypedProp = CastField<FEnumProperty>(Prop))
+	// Do the value types
+	switch (PropertyBase.GetUHTPropertyType())
 	{
+#if UHT_ENABLE_VALUE_PROPERTY_TAG
+	case EUHTPropertyType::Enum:
 		PropertyBase.EnumDef->GetHashTag(Out);
-		return;
-	}
+		break;
 
-	if (FStructProperty* TypedProp = CastField<FStructProperty>(Prop))
-	{
+	case EUHTPropertyType::Struct:
 		PropertyBase.ScriptStructDef->GetHashTag(Out);
-		return;
-	}
-
-	if (FArrayProperty* TypedProp = CastField<FArrayProperty>(Prop))
-	{
-		GetPropertyTag(Out, PropDef.GetValuePropDef());
-		return;
-	}
-
-	if (FMapProperty* TypedProp = CastField<FMapProperty>(Prop))
-	{
-		GetPropertyTag(Out, PropDef.GetKeyPropDef());
-		GetPropertyTag(Out, PropDef.GetValuePropDef());
-		return;
-	}
-
-	if (FSetProperty* TypedProp = CastField<FSetProperty>(Prop))
-	{
-		GetPropertyTag(Out, PropDef.GetValuePropDef());
-		return;
-	}
+		break;
 #endif
 
 #if UHT_ENABLE_PTR_PROPERTY_TAG
-	if (FClassProperty* TypedProp = CastField<FClassProperty>(Prop))
-	{
-		PropertyBase.ClassDef->GetHashTag(Out);
-		return;
-	}
-
-	if (FObjectProperty* TypedProp = CastField<FObjectProperty>(Prop))
-	{
-		PropertyBase.ClassDef->GetHashTag(Out);
-		return;
-	}
-
-	if (FSoftClassProperty* TypedProp = CastField<FSoftClassProperty>(Prop))
-	{
-		PropertyBase.ClassDef->GetHashTag(Out);
-		return;
-	}
-
-	if (FSoftObjectProperty* TypedProp = CastField<FSoftObjectProperty>(Prop))
-	{
-		PropertyBase.ClassDef->GetHashTag(Out);
-		return;
-	}
-
-	if (FClassPtrProperty* TypedProp = CastField<FClassPtrProperty>(Prop))
-	{
-		PropertyBase.ClassDef->GetHashTag(Out);
-		return;
-	}
-
-	if (FObjectPtrProperty* TypedProp = CastField<FObjectPtrProperty>(Prop))
-	{
+	case EUHTPropertyType::ObjectReference:
+	case EUHTPropertyType::WeakObjectReference:
+	case EUHTPropertyType::LazyObjectReference:
+	case EUHTPropertyType::SoftObjectReference:
+	case EUHTPropertyType::ObjectPtrReference:
+	case EUHTPropertyType::InterfaceReference:
+	case EUHTPropertyType::InterfaceReference:
 		PropertyBase.ClassDef->GetHashTag(Out);
-		return;
-	}
+		break;
+#endif
 
-	if (FInterfaceProperty* TypedProp = CastField<FInterfaceProperty>(Prop))
-	{
-		PropertyBase.ClassDef->GetHashTag(Out);
-		return;
+	default:
+		break;
 	}
 
-	if (FWeakObjectProperty* TypedProp = CastField<FWeakObjectProperty>(Prop))
+	// Add the key value
+	if (PropDef.HasKeyPropDef())
 	{
-		PropertyBase.ClassDef->GetHashTag(Out);
-		return;
+		GetPropertyTag(Out, PropDef.GetKeyPropDef());
 	}
-
-	if (FLazyObjectProperty* TypedProp = CastField<FLazyObjectProperty>(Prop))
+	if (PropDef.HasValuePropDef())
 	{
-		PropertyBase.ClassDef->GetHashTag(Out);
-		return;
+		GetPropertyTag(Out, PropDef.GetValuePropDef());
 	}
-#endif
 }
 
-void FNativeClassHeaderGenerator::PropertyNew(FOutputDevice& DeclOut, FOutputDevice& Out, FReferenceGatherers& OutReferenceGatherers, FUnrealPropertyDefinitionInfo& PropertyDef, const TCHAR* OffsetStr, const TCHAR* Name, const TCHAR* DeclSpaces, const TCHAR* Spaces, const TCHAR* SourceStruct) const
+void FNativeClassHeaderGenerator::OutputProperty(FOutputDevice& DeclOut, FOutputDevice& Out, FReferenceGatherers& OutReferenceGatherers, const TCHAR* Scope, TArray<FPropertyNamePointerPair>& PropertyNamesAndPointers, FUnrealPropertyDefinitionInfo& PropertyDef, const TCHAR* OffsetStr, FString&& Name, const TCHAR* DeclSpaces, const TCHAR* Spaces, const TCHAR* SourceStruct) const
 {
 	const FPropertyBase& PropertyBase   = PropertyDef.GetPropertyBase();
-	FProperty*	   Prop                 = PropertyDef.GetProperty();
-	FString        PropName             = CreateUTF8LiteralString(FNativeClassHeaderGenerator::GetOverriddenName(Prop));
-	FString        PropNameDep          = Prop->HasAllPropertyFlags(CPF_Deprecated) ? Prop->GetName() + TEXT("_DEPRECATED") : Prop->GetName();
+	FString        PropName             = CreateUTF8LiteralString(FNativeClassHeaderGenerator::GetOverriddenName(&PropertyDef));
+	FString        PropNameDep          = PropertyDef.HasAllPropertyFlags(CPF_Deprecated) ? PropertyDef.GetName() + TEXT("_DEPRECATED") : PropertyDef.GetName();
 	const TCHAR*   FPropertyObjectFlags = PropertyDef.IsOwnedByDynamicType() ? TEXT("RF_Public|RF_Transient") : TEXT("RF_Public|RF_Transient|RF_MarkAsNative");
-	EPropertyFlags PropFlags            = Prop->PropertyFlags & ~CPF_ComputedFlags;
+	EPropertyFlags PropFlags            = PropertyDef.GetPropertyFlags() & ~CPF_ComputedFlags;
 
 	FUHTStringBuilder PropTag;
 	GetPropertyTag(PropTag, PropertyDef);
 
-	FString PropNotifyFunc = (Prop->RepNotifyFunc != NAME_None) ? CreateUTF8LiteralString(*Prop->RepNotifyFunc.ToString()) : TEXT("nullptr");
-
-	FString ArrayDim = (Prop->ArrayDim != 1) ? FString::Printf(TEXT("CPP_ARRAY_DIM(%s, %s)"), *PropNameDep, SourceStruct) : TEXT("1");
+	FString PropNotifyFunc = PropertyDef.GetRepNotifyFunc() != NAME_None ? CreateUTF8LiteralString(*PropertyDef.GetRepNotifyFunc().ToString()) : TEXT("nullptr");
 
-	FString MetaDataParams = OutputMetaDataCodeForObject(DeclOut, Out, PropertyDef, *FString::Printf(TEXT("%s_MetaData"), Name), DeclSpaces, Spaces);
+	FString ArrayDim = PropertyDef.IsStaticArray() ? FString::Printf(TEXT("CPP_ARRAY_DIM(%s, %s)"), *PropNameDep, SourceStruct) : TEXT("1");
 
-	FString NameWithoutScope = Name;
-	FString Scope;
-	int32 ScopeIndex = NameWithoutScope.Find(TEXT("::"), ESearchCase::CaseSensitive);
-	if (ScopeIndex != INDEX_NONE)
+	FString NameWithoutScope = *Name;
 	{
-		Scope = NameWithoutScope.Left(ScopeIndex) + TEXT("_");
-		NameWithoutScope.RightChopInline(ScopeIndex + 2, false);
+		//FString Scope;
+		int32 ScopeIndex = NameWithoutScope.Find(TEXT("::"), ESearchCase::CaseSensitive);
+		if (ScopeIndex != INDEX_NONE)
+		{
+			//Scope = NameWithoutScope.Left(ScopeIndex) + TEXT("_");
+			NameWithoutScope.RightChopInline(ScopeIndex + 2, false);
+		}
 	}
 
-	if (FByteProperty* TypedProp = CastField<FByteProperty>(Prop))
+	auto OutputByteProperty = [&]()
 	{
+		FString MetaDataParams = OutputMetaDataCodeForObject(DeclOut, Out, PropertyDef, *FString::Printf(TEXT("%s_MetaData"), *Name), DeclSpaces, Spaces);
+
 		DeclOut.Logf(TEXT("%sstatic const UECodeGen_Private::FBytePropertyParams %s;\r\n"), DeclSpaces, *NameWithoutScope);
 
 		Out.Logf(
 			TEXT("%sconst UECodeGen_Private::FBytePropertyParams %s = { %s, %s, (EPropertyFlags)0x%016llx, UECodeGen_Private::EPropertyGenFlags::Byte, %s, %s, %s, %s, %s };%s\r\n"),
 			Spaces,
-			Name,
+			*Name,
 			*PropName,
 			*PropNotifyFunc,
 			PropFlags,
@@ -950,18 +840,26 @@ void FNativeClassHeaderGenerator::PropertyNew(FOutputDevice& DeclOut, FOutputDev
 			*MetaDataParams,
 			*PropTag
 		);
+	};
 
-		return;
+	switch (PropertyBase.GetUHTPropertyType())
+	{
+	case EUHTPropertyType::Byte:
+	{
+		OutputByteProperty();
+		break;
 	}
 
-	if (FInt8Property* TypedProp = CastField<FInt8Property>(Prop))
+	case EUHTPropertyType::Int8:
 	{
+		FString MetaDataParams = OutputMetaDataCodeForObject(DeclOut, Out, PropertyDef, *FString::Printf(TEXT("%s_MetaData"), *Name), DeclSpaces, Spaces);
+
 		DeclOut.Logf(TEXT("%sstatic const UECodeGen_Private::FInt8PropertyParams %s;\r\n"), DeclSpaces, *NameWithoutScope);
 
 		Out.Logf(
 			TEXT("%sconst UECodeGen_Private::FInt8PropertyParams %s = { %s, %s, (EPropertyFlags)0x%016llx, UECodeGen_Private::EPropertyGenFlags::Int8, %s, %s, %s, %s };%s\r\n"),
 			Spaces,
-			Name,
+			*Name,
 			*PropName,
 			*PropNotifyFunc,
 			PropFlags,
@@ -971,18 +869,19 @@ void FNativeClassHeaderGenerator::PropertyNew(FOutputDevice& DeclOut, FOutputDev
 			*MetaDataParams,
 			*PropTag
 		);
-
-		return;
+		break;
 	}
 
-	if (FInt16Property* TypedProp = CastField<FInt16Property>(Prop))
+	case EUHTPropertyType::Int16:
 	{
+		FString MetaDataParams = OutputMetaDataCodeForObject(DeclOut, Out, PropertyDef, *FString::Printf(TEXT("%s_MetaData"), *Name), DeclSpaces, Spaces);
+
 		DeclOut.Logf(TEXT("%sstatic const UECodeGen_Private::FInt16PropertyParams %s;\r\n"), DeclSpaces, *NameWithoutScope);
 
 		Out.Logf(
 			TEXT("%sconst UECodeGen_Private::FInt16PropertyParams %s = { %s, %s, (EPropertyFlags)0x%016llx, UECodeGen_Private::EPropertyGenFlags::Int16, %s, %s, %s, %s };%s\r\n"),
 			Spaces,
-			Name,
+			*Name,
 			*PropName,
 			*PropNotifyFunc,
 			PropFlags,
@@ -992,12 +891,13 @@ void FNativeClassHeaderGenerator::PropertyNew(FOutputDevice& DeclOut, FOutputDev
 			*MetaDataParams,
 			*PropTag
 		);
-
-		return;
+		break;
 	}
 
-	if (FIntProperty* TypedProp = CastField<FIntProperty>(Prop))
+	case EUHTPropertyType::Int:
 	{
+		FString MetaDataParams = OutputMetaDataCodeForObject(DeclOut, Out, PropertyDef, *FString::Printf(TEXT("%s_MetaData"), *Name), DeclSpaces, Spaces);
+
 		const TCHAR* PropTypeName = PropertyDef.IsUnsized() ? TEXT("FUnsizedIntPropertyParams") : TEXT("FIntPropertyParams");
 
 		DeclOut.Logf(TEXT("%sstatic const UECodeGen_Private::%s %s;\r\n"), DeclSpaces, PropTypeName, *NameWithoutScope);
@@ -1006,7 +906,7 @@ void FNativeClassHeaderGenerator::PropertyNew(FOutputDevice& DeclOut, FOutputDev
 			TEXT("%sconst UECodeGen_Private::%s %s = { %s, %s, (EPropertyFlags)0x%016llx, UECodeGen_Private::EPropertyGenFlags::Int, %s, %s, %s, %s };%s\r\n"),
 			Spaces,
 			PropTypeName,
-			Name,
+			*Name,
 			*PropName,
 			*PropNotifyFunc,
 			PropFlags,
@@ -1016,18 +916,19 @@ void FNativeClassHeaderGenerator::PropertyNew(FOutputDevice& DeclOut, FOutputDev
 			*MetaDataParams,
 			*PropTag
 		);
-
-		return;
+		break;
 	}
 
-	if (FInt64Property* TypedProp = CastField<FInt64Property>(Prop))
+	case EUHTPropertyType::Int64:
 	{
+		FString MetaDataParams = OutputMetaDataCodeForObject(DeclOut, Out, PropertyDef, *FString::Printf(TEXT("%s_MetaData"), *Name), DeclSpaces, Spaces);
+
 		DeclOut.Logf(TEXT("%sstatic const UECodeGen_Private::FInt64PropertyParams %s;\r\n"), DeclSpaces, *NameWithoutScope);
 
 		Out.Logf(
 			TEXT("%sconst UECodeGen_Private::FInt64PropertyParams %s = { %s, %s, (EPropertyFlags)0x%016llx, UECodeGen_Private::EPropertyGenFlags::Int64, %s, %s, %s, %s };%s\r\n"),
 			Spaces,
-			Name,
+			*Name,
 			*PropName,
 			*PropNotifyFunc,
 			PropFlags,
@@ -1037,18 +938,19 @@ void FNativeClassHeaderGenerator::PropertyNew(FOutputDevice& DeclOut, FOutputDev
 			*MetaDataParams,
 			*PropTag
 		);
-
-		return;
+		break;
 	}
 
-	if (FUInt16Property* TypedProp = CastField<FUInt16Property>(Prop))
+	case EUHTPropertyType::UInt16:
 	{
+		FString MetaDataParams = OutputMetaDataCodeForObject(DeclOut, Out, PropertyDef, *FString::Printf(TEXT("%s_MetaData"), *Name), DeclSpaces, Spaces);
+
 		DeclOut.Logf(TEXT("%sstatic const UECodeGen_Private::FFInt16PropertyParams %s;\r\n"), DeclSpaces, *NameWithoutScope);
 
 		Out.Logf(
 			TEXT("%sconst UECodeGen_Private::FFInt16PropertyParams %s = { %s, %s, (EPropertyFlags)0x%016llx, UECodeGen_Private::EPropertyGenFlags::UInt16, %s, %s, %s, %s };%s\r\n"),
 			Spaces,
-			Name,
+			*Name,
 			*PropName,
 			*PropNotifyFunc,
 			PropFlags,
@@ -1058,12 +960,13 @@ void FNativeClassHeaderGenerator::PropertyNew(FOutputDevice& DeclOut, FOutputDev
 			*MetaDataParams,
 			*PropTag
 		);
-
-		return;
+		break;
 	}
 
-	if (FUInt32Property* TypedProp = CastField<FUInt32Property>(Prop))
+	case EUHTPropertyType::UInt32:
 	{
+		FString MetaDataParams = OutputMetaDataCodeForObject(DeclOut, Out, PropertyDef, *FString::Printf(TEXT("%s_MetaData"), *Name), DeclSpaces, Spaces);
+
 		const TCHAR* PropTypeName = PropertyDef.IsUnsized() ? TEXT("FUnsizedFIntPropertyParams") : TEXT("FUInt32PropertyParams");
 
 		DeclOut.Logf(TEXT("%sstatic const UECodeGen_Private::%s %s;\r\n"), DeclSpaces, PropTypeName, *NameWithoutScope);
@@ -1072,7 +975,7 @@ void FNativeClassHeaderGenerator::PropertyNew(FOutputDevice& DeclOut, FOutputDev
 			TEXT("%sconst UECodeGen_Private::%s %s = { %s, %s, (EPropertyFlags)0x%016llx, UECodeGen_Private::EPropertyGenFlags::UInt32, %s, %s, %s, %s };%s\r\n"),
 			Spaces,
 			PropTypeName,
-			Name,
+			*Name,
 			*PropName,
 			*PropNotifyFunc,
 			PropFlags,
@@ -1082,18 +985,19 @@ void FNativeClassHeaderGenerator::PropertyNew(FOutputDevice& DeclOut, FOutputDev
 			*MetaDataParams,
 			*PropTag
 		);
-
-		return;
+		break;
 	}
 
-	if (FUInt64Property* TypedProp = CastField<FUInt64Property>(Prop))
+	case EUHTPropertyType::UInt64:
 	{
+		FString MetaDataParams = OutputMetaDataCodeForObject(DeclOut, Out, PropertyDef, *FString::Printf(TEXT("%s_MetaData"), *Name), DeclSpaces, Spaces);
+
 		DeclOut.Logf(TEXT("%sstatic const UECodeGen_Private::FFInt64PropertyParams %s;\r\n"), DeclSpaces, *NameWithoutScope);
 
 		Out.Logf(
 			TEXT("%sconst UECodeGen_Private::FFInt64PropertyParams %s = { %s, %s, (EPropertyFlags)0x%016llx, UECodeGen_Private::EPropertyGenFlags::UInt64, %s, %s, %s, %s };%s\r\n"),
 			Spaces,
-			Name,
+			*Name,
 			*PropName,
 			*PropNotifyFunc,
 			PropFlags,
@@ -1103,18 +1007,19 @@ void FNativeClassHeaderGenerator::PropertyNew(FOutputDevice& DeclOut, FOutputDev
 			*MetaDataParams,
 			*PropTag
 		);
-
-		return;
+		break;
 	}
 
-	if (FFloatProperty* TypedProp = CastField<FFloatProperty>(Prop))
+	case EUHTPropertyType::Float:
 	{
+		FString MetaDataParams = OutputMetaDataCodeForObject(DeclOut, Out, PropertyDef, *FString::Printf(TEXT("%s_MetaData"), *Name), DeclSpaces, Spaces);
+
 		DeclOut.Logf(TEXT("%sstatic const UECodeGen_Private::FFloatPropertyParams %s;\r\n"), DeclSpaces, *NameWithoutScope);
 
 		Out.Logf(
 			TEXT("%sconst UECodeGen_Private::FFloatPropertyParams %s = { %s, %s, (EPropertyFlags)0x%016llx, UECodeGen_Private::EPropertyGenFlags::Float, %s, %s, %s, %s };%s\r\n"),
 			Spaces,
-			Name,
+			*Name,
 			*PropName,
 			*PropNotifyFunc,
 			PropFlags,
@@ -1124,18 +1029,19 @@ void FNativeClassHeaderGenerator::PropertyNew(FOutputDevice& DeclOut, FOutputDev
 			*MetaDataParams,
 			*PropTag
 		);
-
-		return;
+		break;
 	}
 
-	if (FDoubleProperty* TypedProp = CastField<FDoubleProperty>(Prop))
+	case EUHTPropertyType::Double:
 	{
+		FString MetaDataParams = OutputMetaDataCodeForObject(DeclOut, Out, PropertyDef, *FString::Printf(TEXT("%s_MetaData"), *Name), DeclSpaces, Spaces);
+
 		DeclOut.Logf(TEXT("%sstatic const UECodeGen_Private::FDoublePropertyParams %s;\r\n"), DeclSpaces, *NameWithoutScope);
 
 		Out.Logf(
 			TEXT("%sconst UECodeGen_Private::FDoublePropertyParams %s = { %s, %s, (EPropertyFlags)0x%016llx, UECodeGen_Private::EPropertyGenFlags::Double, %s, %s, %s, %s };%s\r\n"),
 			Spaces,
-			Name,
+			*Name,
 			*PropName,
 			*PropNotifyFunc,
 			PropFlags,
@@ -1145,18 +1051,19 @@ void FNativeClassHeaderGenerator::PropertyNew(FOutputDevice& DeclOut, FOutputDev
 			*MetaDataParams,
 			*PropTag
 		);
-
-		return;
+		break;
 	}
 
-	if (FLargeWorldCoordinatesRealProperty* TypedProp = CastField<FLargeWorldCoordinatesRealProperty>(Prop))
+	case EUHTPropertyType::LargeWorldCoordinatesReal:
 	{
+		FString MetaDataParams = OutputMetaDataCodeForObject(DeclOut, Out, PropertyDef, *FString::Printf(TEXT("%s_MetaData"), *Name), DeclSpaces, Spaces);
+
 		DeclOut.Logf(TEXT("%sstatic const UECodeGen_Private::FLargeWorldCoordinatesRealPropertyParams %s;\r\n"), DeclSpaces, *NameWithoutScope);
 
 		Out.Logf(
 			TEXT("%sconst UECodeGen_Private::FLargeWorldCoordinatesRealPropertyParams %s = { %s, %s, (EPropertyFlags)0x%016llx, UECodeGen_Private::EPropertyGenFlags::LargeWorldCoordinatesReal, %s, %s, %s, %s };%s\r\n"),
 			Spaces,
-			Name,
+			*Name,
 			*PropName,
 			*PropNotifyFunc,
 			PropFlags,
@@ -1166,18 +1073,24 @@ void FNativeClassHeaderGenerator::PropertyNew(FOutputDevice& DeclOut, FOutputDev
 			*MetaDataParams,
 			*PropTag
 		);
-
-		return;
+		break;
 	}
 
-	if (FBoolProperty* TypedProp = CastField<FBoolProperty>(Prop))
+	case EUHTPropertyType::Bool:
+	case EUHTPropertyType::Bool8:
+	case EUHTPropertyType::Bool16:
+	case EUHTPropertyType::Bool32:
+	case EUHTPropertyType::Bool64:
 	{
+		FString MetaDataParams = OutputMetaDataCodeForObject(DeclOut, Out, PropertyDef, *FString::Printf(TEXT("%s_MetaData"), *Name), DeclSpaces, Spaces);
+
+		bool bIsNativeBool = PropertyBase.Type == CPT_Bool || EnumHasAnyFlags(PropertyBase.PropertyFlags, CPF_ReturnParm);
 		FString OuterSize;
 		FString Setter;
-		if (!Prop->GetOwner<UObject>())
+		if (UHTCast<FUnrealObjectDefinitionInfo>(PropertyDef.GetOuter()) == nullptr)
 		{
 			OuterSize = TEXT("0");
-			Setter    = TEXT("nullptr");
+			Setter = TEXT("nullptr");
 		}
 		else
 		{
@@ -1185,222 +1098,231 @@ void FNativeClassHeaderGenerator::PropertyNew(FOutputDevice& DeclOut, FOutputDev
 
 			DeclOut.Logf(TEXT("%sstatic void %s_SetBit(void* Obj);\r\n"), DeclSpaces, *NameWithoutScope);
 
-			Out.Logf(TEXT("%svoid %s_SetBit(void* Obj)\r\n"), Spaces, Name);
+			Out.Logf(TEXT("%svoid %s_SetBit(void* Obj)\r\n"), Spaces, *Name);
 			Out.Logf(TEXT("%s{\r\n"), Spaces);
-			Out.Logf(TEXT("%s\t((%s*)Obj)->%s%s = 1;\r\n"), Spaces, SourceStruct, *Prop->GetName(), Prop->HasAllPropertyFlags(CPF_Deprecated) ? TEXT("_DEPRECATED") : TEXT(""));
+			Out.Logf(TEXT("%s\t((%s*)Obj)->%s%s = 1;\r\n"), Spaces, SourceStruct, *PropertyDef.GetName(), PropertyDef.HasAllPropertyFlags(CPF_Deprecated) ? TEXT("_DEPRECATED") : TEXT(""));
 			Out.Logf(TEXT("%s}\r\n"), Spaces);
 
-			Setter = FString::Printf(TEXT("&%s_SetBit"), Name);
+			Setter = FString::Printf(TEXT("&%s_SetBit"), *Name);
 		}
-
+		
 		DeclOut.Logf(TEXT("%sstatic const UECodeGen_Private::FBoolPropertyParams %s;\r\n"), DeclSpaces, *NameWithoutScope);
 
 		Out.Logf(
 			TEXT("%sconst UECodeGen_Private::FBoolPropertyParams %s = { %s, %s, (EPropertyFlags)0x%016llx, UECodeGen_Private::EPropertyGenFlags::Bool %s, %s, %s, sizeof(%s), %s, %s, %s };%s\r\n"),
 			Spaces,
-			Name,
+			*Name,
 			*PropName,
 			*PropNotifyFunc,
 			PropFlags,
-			TypedProp->IsNativeBool() ? TEXT("| UECodeGen_Private::EPropertyGenFlags::NativeBool") : TEXT(""),
+			bIsNativeBool ? TEXT("| UECodeGen_Private::EPropertyGenFlags::NativeBool") : TEXT(""),
 			FPropertyObjectFlags,
 			*ArrayDim,
-			*TypedProp->GetCPPType(nullptr, 0),
+			*PropertyDef.GetCPPType(nullptr, 0),
 			*OuterSize,
 			*Setter,
 			*MetaDataParams,
 			*PropTag
 		);
-
-		return;
+		break;
 	}
 
-	if (FSoftClassProperty* TypedProp = CastField<FSoftClassProperty>(Prop))
+	case EUHTPropertyType::ObjectReference:
 	{
-		DeclOut.Logf(TEXT("%sstatic const UECodeGen_Private::FSoftClassPropertyParams %s;\r\n"), DeclSpaces, *NameWithoutScope);
+		FString MetaDataParams = OutputMetaDataCodeForObject(DeclOut, Out, PropertyDef, *FString::Printf(TEXT("%s_MetaData"), *Name), DeclSpaces, Spaces);
 
-		Out.Logf(
-			TEXT("%sconst UECodeGen_Private::FSoftClassPropertyParams %s = { %s, %s, (EPropertyFlags)0x%016llx, UECodeGen_Private::EPropertyGenFlags::SoftClass, %s, %s, %s, %s, %s };%s\r\n"),
-			Spaces,
-			Name,
-			*PropName,
-			*PropNotifyFunc,
-			PropFlags,
-			FPropertyObjectFlags,
-			*ArrayDim,
-			OffsetStr,
-			*GetSingletonNameFuncAddr(PropertyBase.MetaClassDef, OutReferenceGatherers.UniqueCrossModuleReferences, false),
-			*MetaDataParams,
-			*PropTag
-		);
+		UClass* PropertyClass = PropertyBase.ClassDef->GetClass();
+		if (PropertyClass->IsChildOf(UClass::StaticClass()))
+		{
+			DeclOut.Logf(TEXT("%sstatic const UECodeGen_Private::FClassPropertyParams %s;\r\n"), DeclSpaces, *NameWithoutScope);
 
-		return;
+			Out.Logf(
+				TEXT("%sconst UECodeGen_Private::FClassPropertyParams %s = { %s, %s, (EPropertyFlags)0x%016llx, UECodeGen_Private::EPropertyGenFlags::Class, %s, %s, %s, %s, %s, %s };%s\r\n"),
+				Spaces,
+				*Name,
+				*PropName,
+				*PropNotifyFunc,
+				PropFlags,
+				FPropertyObjectFlags,
+				*ArrayDim,
+				OffsetStr,
+				*GetSingletonNameFuncAddr(PropertyBase.MetaClassDef, OutReferenceGatherers.UniqueCrossModuleReferences, false),
+				*GetSingletonNameFuncAddr(PropertyBase.ClassDef, OutReferenceGatherers.UniqueCrossModuleReferences, false),
+				*MetaDataParams,
+				*PropTag
+			);
+		}
+		else
+		{
+			DeclOut.Logf(TEXT("%sstatic const UECodeGen_Private::FObjectPropertyParams %s;\r\n"), DeclSpaces, *NameWithoutScope);
+
+			Out.Logf(
+				TEXT("%sconst UECodeGen_Private::FObjectPropertyParams %s = { %s, %s, (EPropertyFlags)0x%016llx, UECodeGen_Private::EPropertyGenFlags::Object, %s, %s, %s, %s, %s };%s\r\n"),
+				Spaces,
+				*Name,
+				*PropName,
+				*PropNotifyFunc,
+				PropFlags,
+				FPropertyObjectFlags,
+				*ArrayDim,
+				OffsetStr,
+				*GetSingletonNameFuncAddr(PropertyBase.ClassDef, OutReferenceGatherers.UniqueCrossModuleReferences, false),
+				*MetaDataParams,
+				*PropTag
+			);
+		}
+		break;
 	}
 
-	if (FWeakObjectProperty* TypedProp = CastField<FWeakObjectProperty>(Prop))
+	case EUHTPropertyType::ObjectPtrReference:
 	{
-		DeclOut.Logf(TEXT("%sstatic const UECodeGen_Private::FWeakObjectPropertyParams %s;\r\n"), DeclSpaces, *NameWithoutScope);
+		FString MetaDataParams = OutputMetaDataCodeForObject(DeclOut, Out, PropertyDef, *FString::Printf(TEXT("%s_MetaData"), *Name), DeclSpaces, Spaces);
 
-		Out.Logf(
-			TEXT("%sconst UECodeGen_Private::FWeakObjectPropertyParams %s = { %s, %s, (EPropertyFlags)0x%016llx, UECodeGen_Private::EPropertyGenFlags::WeakObject, %s, %s, %s, %s, %s };%s\r\n"),
-			Spaces,
-			Name,
-			*PropName,
-			*PropNotifyFunc,
-			PropFlags,
-			FPropertyObjectFlags,
-			*ArrayDim,
-			OffsetStr,
-			*GetSingletonNameFuncAddr(PropertyBase.ClassDef, OutReferenceGatherers.UniqueCrossModuleReferences, false),
-			*MetaDataParams,
-			*PropTag
-		);
+		UClass* PropertyClass = PropertyBase.ClassDef->GetClass();
+		if (PropertyClass->IsChildOf(UClass::StaticClass()))
+		{
+			DeclOut.Logf(TEXT("%sstatic const UECodeGen_Private::FClassPtrPropertyParams %s;\r\n"), DeclSpaces, *NameWithoutScope);
 
-		return;
+			Out.Logf(
+				TEXT("%sconst UECodeGen_Private::FClassPtrPropertyParams %s = { %s, %s, (EPropertyFlags)0x%016llx, UECodeGen_Private::EPropertyGenFlags::Class | UECodeGen_Private::EPropertyGenFlags::ObjectPtr, %s, %s, %s, %s, %s, %s };%s\r\n"),
+				Spaces,
+				*Name,
+				*PropName,
+				*PropNotifyFunc,
+				PropFlags,
+				FPropertyObjectFlags,
+				*ArrayDim,
+				OffsetStr,
+				*GetSingletonNameFuncAddr(PropertyBase.MetaClassDef, OutReferenceGatherers.UniqueCrossModuleReferences, false),
+				*GetSingletonNameFuncAddr(PropertyBase.ClassDef, OutReferenceGatherers.UniqueCrossModuleReferences, false),
+				*MetaDataParams,
+				*PropTag
+			);
+		}
+		else
+		{
+			DeclOut.Logf(TEXT("%sstatic const UECodeGen_Private::FObjectPtrPropertyParams %s;\r\n"), DeclSpaces, *NameWithoutScope);
+
+			Out.Logf(
+				TEXT("%sconst UECodeGen_Private::FObjectPtrPropertyParams %s = { %s, %s, (EPropertyFlags)0x%016llx, UECodeGen_Private::EPropertyGenFlags::Object | UECodeGen_Private::EPropertyGenFlags::ObjectPtr, %s, %s, %s, %s, %s };%s\r\n"),
+				Spaces,
+				*Name,
+				*PropName,
+				*PropNotifyFunc,
+				PropFlags,
+				FPropertyObjectFlags,
+				*ArrayDim,
+				OffsetStr,
+				*GetSingletonNameFuncAddr(PropertyBase.ClassDef, OutReferenceGatherers.UniqueCrossModuleReferences, false),
+				*MetaDataParams,
+				*PropTag
+			);
+		}
+		break;
 	}
 
-	if (FLazyObjectProperty* TypedProp = CastField<FLazyObjectProperty>(Prop))
+	case EUHTPropertyType::SoftObjectReference:
 	{
-		DeclOut.Logf(TEXT("%sstatic const UECodeGen_Private::FLazyObjectPropertyParams %s;\r\n"), DeclSpaces, *NameWithoutScope);
+		FString MetaDataParams = OutputMetaDataCodeForObject(DeclOut, Out, PropertyDef, *FString::Printf(TEXT("%s_MetaData"), *Name), DeclSpaces, Spaces);
 
-		Out.Logf(
-			TEXT("%sconst UECodeGen_Private::FLazyObjectPropertyParams %s = { %s, %s, (EPropertyFlags)0x%016llx, UECodeGen_Private::EPropertyGenFlags::LazyObject, %s, %s, %s, %s, %s };%s\r\n"),
-			Spaces,
-			Name,
-			*PropName,
-			*PropNotifyFunc,
-			PropFlags,
-			FPropertyObjectFlags,
-			*ArrayDim,
-			OffsetStr,
-			*GetSingletonNameFuncAddr(PropertyBase.ClassDef, OutReferenceGatherers.UniqueCrossModuleReferences, false),
-			*MetaDataParams,
-			*PropTag
-		);
+		UClass* PropertyClass = PropertyBase.ClassDef->GetClass();
+		if (PropertyClass->IsChildOf(UClass::StaticClass()))
+		{
+			DeclOut.Logf(TEXT("%sstatic const UECodeGen_Private::FSoftClassPropertyParams %s;\r\n"), DeclSpaces, *NameWithoutScope);
 
-		return;
+			Out.Logf(
+				TEXT("%sconst UECodeGen_Private::FSoftClassPropertyParams %s = { %s, %s, (EPropertyFlags)0x%016llx, UECodeGen_Private::EPropertyGenFlags::SoftClass, %s, %s, %s, %s, %s };%s\r\n"),
+				Spaces,
+				*Name,
+				*PropName,
+				*PropNotifyFunc,
+				PropFlags,
+				FPropertyObjectFlags,
+				*ArrayDim,
+				OffsetStr,
+				*GetSingletonNameFuncAddr(PropertyBase.MetaClassDef, OutReferenceGatherers.UniqueCrossModuleReferences, false),
+				*MetaDataParams,
+				*PropTag
+			);
+		}
+		else
+		{
+			DeclOut.Logf(TEXT("%sstatic const UECodeGen_Private::FSoftObjectPropertyParams %s;\r\n"), DeclSpaces, *NameWithoutScope);
+
+			Out.Logf(
+				TEXT("%sconst UECodeGen_Private::FSoftObjectPropertyParams %s = { %s, %s, (EPropertyFlags)0x%016llx, UECodeGen_Private::EPropertyGenFlags::SoftObject, %s, %s, %s, %s, %s };%s\r\n"),
+				Spaces,
+				*Name,
+				*PropName,
+				*PropNotifyFunc,
+				PropFlags,
+				FPropertyObjectFlags,
+				*ArrayDim,
+				OffsetStr,
+				*GetSingletonNameFuncAddr(PropertyBase.ClassDef, OutReferenceGatherers.UniqueCrossModuleReferences, false),
+				*MetaDataParams,
+				*PropTag
+			);
+		}
+		break;
 	}
 
-	if (FObjectPtrProperty* TypedProp = CastField<FObjectPtrProperty>(Prop))
+	case EUHTPropertyType::WeakObjectReference:
 	{
-		DeclOut.Logf(TEXT("%sstatic const UECodeGen_Private::FObjectPtrPropertyParams %s;\r\n"), DeclSpaces, *NameWithoutScope);
-
-		Out.Logf(
-			TEXT("%sconst UECodeGen_Private::FObjectPtrPropertyParams %s = { %s, %s, (EPropertyFlags)0x%016llx, UECodeGen_Private::EPropertyGenFlags::Object | UECodeGen_Private::EPropertyGenFlags::ObjectPtr, %s, %s, %s, %s, %s };%s\r\n"),
-			Spaces,
-			Name,
-			*PropName,
-			*PropNotifyFunc,
-			PropFlags,
-			FPropertyObjectFlags,
-			*ArrayDim,
-			OffsetStr,
-			*GetSingletonNameFuncAddr(PropertyBase.ClassDef, OutReferenceGatherers.UniqueCrossModuleReferences, false),
-			*MetaDataParams,
-			*PropTag
-		);
+		FString MetaDataParams = OutputMetaDataCodeForObject(DeclOut, Out, PropertyDef, *FString::Printf(TEXT("%s_MetaData"), *Name), DeclSpaces, Spaces);
 
-		return;
-	}
-
-	if (FClassPtrProperty* TypedProp = CastField<FClassPtrProperty>(Prop))
-	{
-		DeclOut.Logf(TEXT("%sstatic const UECodeGen_Private::FClassPtrPropertyParams %s;\r\n"), DeclSpaces, *NameWithoutScope);
+		DeclOut.Logf(TEXT("%sstatic const UECodeGen_Private::FWeakObjectPropertyParams %s;\r\n"), DeclSpaces, *NameWithoutScope);
 
 		Out.Logf(
-			TEXT("%sconst UECodeGen_Private::FClassPtrPropertyParams %s = { %s, %s, (EPropertyFlags)0x%016llx, UECodeGen_Private::EPropertyGenFlags::Class | UECodeGen_Private::EPropertyGenFlags::ObjectPtr, %s, %s, %s, %s, %s, %s };%s\r\n"),
+			TEXT("%sconst UECodeGen_Private::FWeakObjectPropertyParams %s = { %s, %s, (EPropertyFlags)0x%016llx, UECodeGen_Private::EPropertyGenFlags::WeakObject, %s, %s, %s, %s, %s };%s\r\n"),
 			Spaces,
-			Name,
+			*Name,
 			*PropName,
 			*PropNotifyFunc,
 			PropFlags,
 			FPropertyObjectFlags,
 			*ArrayDim,
 			OffsetStr,
-			*GetSingletonNameFuncAddr(PropertyBase.MetaClassDef, OutReferenceGatherers.UniqueCrossModuleReferences, false),
 			*GetSingletonNameFuncAddr(PropertyBase.ClassDef, OutReferenceGatherers.UniqueCrossModuleReferences, false),
 			*MetaDataParams,
 			*PropTag
 		);
-
-		return;
+		break;
 	}
 
-	if (FSoftObjectProperty* TypedProp = CastField<FSoftObjectProperty>(Prop))
+	case EUHTPropertyType::LazyObjectReference:
 	{
-		DeclOut.Logf(TEXT("%sstatic const UECodeGen_Private::FSoftObjectPropertyParams %s;\r\n"), DeclSpaces, *NameWithoutScope);
+		FString MetaDataParams = OutputMetaDataCodeForObject(DeclOut, Out, PropertyDef, *FString::Printf(TEXT("%s_MetaData"), *Name), DeclSpaces, Spaces);
 
-		Out.Logf(
-			TEXT("%sconst UECodeGen_Private::FSoftObjectPropertyParams %s = { %s, %s, (EPropertyFlags)0x%016llx, UECodeGen_Private::EPropertyGenFlags::SoftObject, %s, %s, %s, %s, %s };%s\r\n"),
-			Spaces,
-			Name,
-			*PropName,
-			*PropNotifyFunc,
-			PropFlags,
-			FPropertyObjectFlags,
-			*ArrayDim,
-			OffsetStr,
-			*GetSingletonNameFuncAddr(PropertyBase.ClassDef, OutReferenceGatherers.UniqueCrossModuleReferences, false),
-			*MetaDataParams,
-			*PropTag
-		);
-
-		return;
-	}
-
-	if (FClassProperty* TypedProp = CastField<FClassProperty>(Prop))
-	{
-		DeclOut.Logf(TEXT("%sstatic const UECodeGen_Private::FClassPropertyParams %s;\r\n"), DeclSpaces, *NameWithoutScope);
+		DeclOut.Logf(TEXT("%sstatic const UECodeGen_Private::FLazyObjectPropertyParams %s;\r\n"), DeclSpaces, *NameWithoutScope);
 
 		Out.Logf(
-			TEXT("%sconst UECodeGen_Private::FClassPropertyParams %s = { %s, %s, (EPropertyFlags)0x%016llx, UECodeGen_Private::EPropertyGenFlags::Class, %s, %s, %s, %s, %s, %s };%s\r\n"),
+			TEXT("%sconst UECodeGen_Private::FLazyObjectPropertyParams %s = { %s, %s, (EPropertyFlags)0x%016llx, UECodeGen_Private::EPropertyGenFlags::LazyObject, %s, %s, %s, %s, %s };%s\r\n"),
 			Spaces,
-			Name,
+			*Name,
 			*PropName,
 			*PropNotifyFunc,
 			PropFlags,
 			FPropertyObjectFlags,
 			*ArrayDim,
 			OffsetStr,
-			*GetSingletonNameFuncAddr(PropertyBase.MetaClassDef, OutReferenceGatherers.UniqueCrossModuleReferences, false),
 			*GetSingletonNameFuncAddr(PropertyBase.ClassDef, OutReferenceGatherers.UniqueCrossModuleReferences, false),
 			*MetaDataParams,
 			*PropTag
 		);
-
-		return;
+		break;
 	}
 
-	if (FObjectProperty* TypedProp = CastField<FObjectProperty>(Prop))
+	case EUHTPropertyType::Interface:
 	{
-		DeclOut.Logf(TEXT("%sstatic const UECodeGen_Private::FObjectPropertyParams %s;\r\n"), DeclSpaces, *NameWithoutScope);
+		FString MetaDataParams = OutputMetaDataCodeForObject(DeclOut, Out, PropertyDef, *FString::Printf(TEXT("%s_MetaData"), *Name), DeclSpaces, Spaces);
 
-		Out.Logf(
-			TEXT("%sconst UECodeGen_Private::FObjectPropertyParams %s = { %s, %s, (EPropertyFlags)0x%016llx, UECodeGen_Private::EPropertyGenFlags::Object, %s, %s, %s, %s, %s };%s\r\n"),
-			Spaces,
-			Name,
-			*PropName,
-			*PropNotifyFunc,
-			PropFlags,
-			FPropertyObjectFlags,
-			*ArrayDim,
-			OffsetStr,
-			*GetSingletonNameFuncAddr(PropertyBase.ClassDef, OutReferenceGatherers.UniqueCrossModuleReferences, false),
-			*MetaDataParams,
-			*PropTag
-		);
-
-		return;
-	}
-
-	if (FInterfaceProperty* TypedProp = CastField<FInterfaceProperty>(Prop))
-	{
 		DeclOut.Logf(TEXT("%sstatic const UECodeGen_Private::FInterfacePropertyParams %s;\r\n"), DeclSpaces, *NameWithoutScope);
 
 		Out.Logf(
 			TEXT("%sconst UECodeGen_Private::FInterfacePropertyParams %s = { %s, %s, (EPropertyFlags)0x%016llx, UECodeGen_Private::EPropertyGenFlags::Interface, %s, %s, %s, %s, %s };%s\r\n"),
 			Spaces,
-			Name,
+			*Name,
 			*PropName,
 			*PropNotifyFunc,
 			PropFlags,
@@ -1411,18 +1333,19 @@ void FNativeClassHeaderGenerator::PropertyNew(FOutputDevice& DeclOut, FOutputDev
 			*MetaDataParams,
 			*PropTag
 		);
-
-		return;
+		break;
 	}
 
-	if (FNameProperty* TypedProp = CastField<FNameProperty>(Prop))
+	case EUHTPropertyType::Name:
 	{
+		FString MetaDataParams = OutputMetaDataCodeForObject(DeclOut, Out, PropertyDef, *FString::Printf(TEXT("%s_MetaData"), *Name), DeclSpaces, Spaces);
+
 		DeclOut.Logf(TEXT("%sstatic const UECodeGen_Private::FNamePropertyParams %s;\r\n"), DeclSpaces, *NameWithoutScope);
 
 		Out.Logf(
 			TEXT("%sconst UECodeGen_Private::FNamePropertyParams %s = { %s, %s, (EPropertyFlags)0x%016llx, UECodeGen_Private::EPropertyGenFlags::Name, %s, %s, %s, %s };%s\r\n"),
 			Spaces,
-			Name,
+			*Name,
 			*PropName,
 			*PropNotifyFunc,
 			PropFlags,
@@ -1432,18 +1355,19 @@ void FNativeClassHeaderGenerator::PropertyNew(FOutputDevice& DeclOut, FOutputDev
 			*MetaDataParams,
 			*PropTag
 		);
-
-		return;
+		break;
 	}
 
-	if (FStrProperty* TypedProp = CastField<FStrProperty>(Prop))
+	case EUHTPropertyType::String:
 	{
+		FString MetaDataParams = OutputMetaDataCodeForObject(DeclOut, Out, PropertyDef, *FString::Printf(TEXT("%s_MetaData"), *Name), DeclSpaces, Spaces);
+
 		DeclOut.Logf(TEXT("%sstatic const UECodeGen_Private::FStrPropertyParams %s;\r\n"), DeclSpaces, *NameWithoutScope);
 
 		Out.Logf(
 			TEXT("%sconst UECodeGen_Private::FStrPropertyParams %s = { %s, %s, (EPropertyFlags)0x%016llx, UECodeGen_Private::EPropertyGenFlags::Str, %s, %s, %s, %s };%s\r\n"),
 			Spaces,
-			Name,
+			*Name,
 			*PropName,
 			*PropNotifyFunc,
 			PropFlags,
@@ -1453,18 +1377,25 @@ void FNativeClassHeaderGenerator::PropertyNew(FOutputDevice& DeclOut, FOutputDev
 			*MetaDataParams,
 			*PropTag
 		);
-
-		return;
+		break;
 	}
 
-	if (FArrayProperty* TypedProp = CastField<FArrayProperty>(Prop))
-	{
+	case EUHTPropertyType::DynamicArray:
+	{
+		FUnrealPropertyDefinitionInfo& ValueDef = PropertyDef.GetValuePropDef();
+
+		FString ValueVariableName = FString::Printf(TEXT("%sNewProp_%s_Inner"), Scope, *ValueDef.GetName());
+
+		OutputProperty(DeclOut, Out, OutReferenceGatherers, Scope, PropertyNamesAndPointers, ValueDef, TEXT("0"), MoveTemp(ValueVariableName), DeclSpaces, Spaces, nullptr);
+
+		FString MetaDataParams = OutputMetaDataCodeForObject(DeclOut, Out, PropertyDef, *FString::Printf(TEXT("%s_MetaData"), *Name), DeclSpaces, Spaces);
+
 		DeclOut.Logf(TEXT("%sstatic const UECodeGen_Private::FArrayPropertyParams %s;\r\n"), DeclSpaces, *NameWithoutScope);
 
 		Out.Logf(
 			TEXT("%sconst UECodeGen_Private::FArrayPropertyParams %s = { %s, %s, (EPropertyFlags)0x%016llx, UECodeGen_Private::EPropertyGenFlags::Array, %s, %s, %s, %s, %s };%s\r\n"),
 			Spaces,
-			Name,
+			*Name,
 			*PropName,
 			*PropNotifyFunc,
 			PropFlags,
@@ -1475,18 +1406,28 @@ void FNativeClassHeaderGenerator::PropertyNew(FOutputDevice& DeclOut, FOutputDev
 			*MetaDataParams,
 			*PropTag
 		);
-
-		return;
+		break;
 	}
 
-	if (FMapProperty* TypedProp = CastField<FMapProperty>(Prop))
+	case EUHTPropertyType::Map:
 	{
+		FUnrealPropertyDefinitionInfo& KeyDef = PropertyDef.GetKeyPropDef();
+		FUnrealPropertyDefinitionInfo& ValueDef = PropertyDef.GetValuePropDef();
+
+		FString KeyVariableName = FString::Printf(TEXT("%sNewProp_%s_KeyProp"), Scope, *KeyDef.GetName());
+		FString ValueVariableName = FString::Printf(TEXT("%sNewProp_%s_ValueProp"), Scope, *ValueDef.GetName());
+
+		OutputProperty(DeclOut, Out, OutReferenceGatherers, Scope, PropertyNamesAndPointers, ValueDef, TEXT("1"), MoveTemp(ValueVariableName), DeclSpaces, Spaces, nullptr);
+		OutputProperty(DeclOut, Out, OutReferenceGatherers, Scope, PropertyNamesAndPointers, KeyDef, TEXT("0"), MoveTemp(KeyVariableName), DeclSpaces, Spaces, nullptr);
+
+		FString MetaDataParams = OutputMetaDataCodeForObject(DeclOut, Out, PropertyDef, *FString::Printf(TEXT("%s_MetaData"), *Name), DeclSpaces, Spaces);
+
 		DeclOut.Logf(TEXT("%sstatic const UECodeGen_Private::FMapPropertyParams %s;\r\n"), DeclSpaces, *NameWithoutScope);
 
 		Out.Logf(
 			TEXT("%sconst UECodeGen_Private::FMapPropertyParams %s = { %s, %s, (EPropertyFlags)0x%016llx, UECodeGen_Private::EPropertyGenFlags::Map, %s, %s, %s, %s, %s };%s\r\n"),
 			Spaces,
-			Name,
+			*Name,
 			*PropName,
 			*PropNotifyFunc,
 			PropFlags,
@@ -1497,18 +1438,25 @@ void FNativeClassHeaderGenerator::PropertyNew(FOutputDevice& DeclOut, FOutputDev
 			*MetaDataParams,
 			*PropTag
 		);
-
-		return;
+		break;
 	}
 
-	if (FSetProperty* TypedProp = CastField<FSetProperty>(Prop))
+	case EUHTPropertyType::Set:
 	{
+		FUnrealPropertyDefinitionInfo& ValueDef = PropertyDef.GetValuePropDef();
+
+		FString ValueVariableName = FString::Printf(TEXT("%sNewProp_%s_ElementProp"), Scope, *ValueDef.GetName());
+
+		OutputProperty(DeclOut, Out, OutReferenceGatherers, Scope, PropertyNamesAndPointers, ValueDef, TEXT("0"), MoveTemp(ValueVariableName), DeclSpaces, Spaces, nullptr);
+
+		FString MetaDataParams = OutputMetaDataCodeForObject(DeclOut, Out, PropertyDef, *FString::Printf(TEXT("%s_MetaData"), *Name), DeclSpaces, Spaces);
+
 		DeclOut.Logf(TEXT("%sstatic const UECodeGen_Private::FSetPropertyParams %s;\r\n"), DeclSpaces, *NameWithoutScope);
 
 		Out.Logf(
 			TEXT("%sconst UECodeGen_Private::FSetPropertyParams %s = { %s, %s, (EPropertyFlags)0x%016llx, UECodeGen_Private::EPropertyGenFlags::Set, %s, %s, %s, %s };%s\r\n"),
 			Spaces,
-			Name,
+			*Name,
 			*PropName,
 			*PropNotifyFunc,
 			PropFlags,
@@ -1518,18 +1466,19 @@ void FNativeClassHeaderGenerator::PropertyNew(FOutputDevice& DeclOut, FOutputDev
 			*MetaDataParams,
 			*PropTag
 		);
-
-		return;
+		break;
 	}
 
-	if (FStructProperty* TypedProp = CastField<FStructProperty>(Prop))
+	case EUHTPropertyType::Struct:
 	{
+		FString MetaDataParams = OutputMetaDataCodeForObject(DeclOut, Out, PropertyDef, *FString::Printf(TEXT("%s_MetaData"), *Name), DeclSpaces, Spaces);
+
 		DeclOut.Logf(TEXT("%sstatic const UECodeGen_Private::FStructPropertyParams %s;\r\n"), DeclSpaces, *NameWithoutScope);
 
 		Out.Logf(
 			TEXT("%sconst UECodeGen_Private::FStructPropertyParams %s = { %s, %s, (EPropertyFlags)0x%016llx, UECodeGen_Private::EPropertyGenFlags::Struct, %s, %s, %s, %s, %s };%s\r\n"),
 			Spaces,
-			Name,
+			*Name,
 			*PropName,
 			*PropNotifyFunc,
 			PropFlags,
@@ -1540,18 +1489,19 @@ void FNativeClassHeaderGenerator::PropertyNew(FOutputDevice& DeclOut, FOutputDev
 			*MetaDataParams,
 			*PropTag
 		);
-
-		return;
+		break;
 	}
 
-	if (FDelegateProperty* TypedProp = CastField<FDelegateProperty>(Prop))
+	case EUHTPropertyType::Delegate:
 	{
+		FString MetaDataParams = OutputMetaDataCodeForObject(DeclOut, Out, PropertyDef, *FString::Printf(TEXT("%s_MetaData"), *Name), DeclSpaces, Spaces);
+
 		DeclOut.Logf(TEXT("%sstatic const UECodeGen_Private::FDelegatePropertyParams %s;\r\n"), DeclSpaces, *NameWithoutScope);
 
 		Out.Logf(
 			TEXT("%sconst UECodeGen_Private::FDelegatePropertyParams %s = { %s, %s, (EPropertyFlags)0x%016llx, UECodeGen_Private::EPropertyGenFlags::Delegate, %s, %s, %s, %s, %s };%s\r\n"),
 			Spaces,
-			Name,
+			*Name,
 			*PropName,
 			*PropNotifyFunc,
 			PropFlags,
@@ -1562,22 +1512,25 @@ void FNativeClassHeaderGenerator::PropertyNew(FOutputDevice& DeclOut, FOutputDev
 			*MetaDataParams,
 			*PropTag
 		);
-
-		return;
+		break;
 	}
 
-	if (FMulticastDelegateProperty* TypedProp = CastField<FMulticastDelegateProperty>(Prop))
+	case EUHTPropertyType::MulticastDelegate:
 	{
+		bool bIsSparse = PropertyBase.FunctionDef->GetFunction()->IsA<USparseDelegateFunction>();
+
+		FString MetaDataParams = OutputMetaDataCodeForObject(DeclOut, Out, PropertyDef, *FString::Printf(TEXT("%s_MetaData"), *Name), DeclSpaces, Spaces);
+
 		DeclOut.Logf(TEXT("%sstatic const UECodeGen_Private::FMulticastDelegatePropertyParams %s;\r\n"), DeclSpaces, *NameWithoutScope);
 
 		Out.Logf(
 			TEXT("%sconst UECodeGen_Private::FMulticastDelegatePropertyParams %s = { %s, %s, (EPropertyFlags)0x%016llx, UECodeGen_Private::EPropertyGenFlags::%sMulticastDelegate, %s, %s, %s, %s, %s };%s\r\n"),
 			Spaces,
-			Name,
+			*Name,
 			*PropName,
 			*PropNotifyFunc,
 			PropFlags,
-			(TypedProp->IsA<FMulticastInlineDelegateProperty>() ? TEXT("Inline") : TEXT("Sparse")),
+			!bIsSparse ? TEXT("Inline") : TEXT("Sparse"),
 			FPropertyObjectFlags,
 			*ArrayDim,
 			OffsetStr,
@@ -1585,18 +1538,19 @@ void FNativeClassHeaderGenerator::PropertyNew(FOutputDevice& DeclOut, FOutputDev
 			*MetaDataParams,
 			*PropTag
 		);
-
-		return;
+		break;
 	}
 
-	if (FTextProperty* TypedProp = CastField<FTextProperty>(Prop))
+	case EUHTPropertyType::Text:
 	{
+		FString MetaDataParams = OutputMetaDataCodeForObject(DeclOut, Out, PropertyDef, *FString::Printf(TEXT("%s_MetaData"), *Name), DeclSpaces, Spaces);
+
 		DeclOut.Logf(TEXT("%sstatic const UECodeGen_Private::FTextPropertyParams %s;\r\n"), DeclSpaces, *NameWithoutScope);
 
 		Out.Logf(
 			TEXT("%sconst UECodeGen_Private::FTextPropertyParams %s = { %s, %s, (EPropertyFlags)0x%016llx, UECodeGen_Private::EPropertyGenFlags::Text, %s, %s, %s, %s };%s\r\n"),
 			Spaces,
-			Name,
+			*Name,
 			*PropName,
 			*PropNotifyFunc,
 			PropFlags,
@@ -1606,76 +1560,113 @@ void FNativeClassHeaderGenerator::PropertyNew(FOutputDevice& DeclOut, FOutputDev
 			*MetaDataParams,
 			*PropTag
 		);
-
-		return;
+		break;
 	}
 
-	if (FEnumProperty* TypedProp = CastField<FEnumProperty>(Prop))
+	case EUHTPropertyType::Enum:
 	{
-		DeclOut.Logf(TEXT("%sstatic const UECodeGen_Private::FEnumPropertyParams %s;\r\n"), DeclSpaces, *NameWithoutScope);
-
-		Out.Logf(
-			TEXT("%sconst UECodeGen_Private::FEnumPropertyParams %s = { %s, %s, (EPropertyFlags)0x%016llx, UECodeGen_Private::EPropertyGenFlags::Enum, %s, %s, %s, %s, %s };%s\r\n"),
-			Spaces,
-			Name,
-			*PropName,
-			*PropNotifyFunc,
-			PropFlags,
-			FPropertyObjectFlags,
-			*ArrayDim,
-			OffsetStr,
-			*GetSingletonNameFuncAddr(PropertyBase.EnumDef, OutReferenceGatherers.UniqueCrossModuleReferences),
-			*MetaDataParams,
-			*PropTag
-		);
-
-		return;
+		if (PropertyBase.EnumDef->GetEnum()->GetCppForm() != UEnum::ECppForm::EnumClass)
+		{
+			OutputByteProperty();
+		}
+		else
+		{
+			// Output the underlying property
+			FString PropVarName = FString::Printf(TEXT("%s_Underlying"), *Name);
+			OutputProperty(DeclOut, Out, OutReferenceGatherers, Scope, PropertyNamesAndPointers, PropertyDef.GetValuePropDef(), TEXT("0"), *PropVarName, DeclSpaces, Spaces, nullptr);
+
+			FString MetaDataParams = OutputMetaDataCodeForObject(DeclOut, Out, PropertyDef, *FString::Printf(TEXT("%s_MetaData"), *Name), DeclSpaces, Spaces);
+
+			DeclOut.Logf(TEXT("%sstatic const UECodeGen_Private::FEnumPropertyParams %s;\r\n"), DeclSpaces, *NameWithoutScope);
+
+			Out.Logf(
+				TEXT("%sconst UECodeGen_Private::FEnumPropertyParams %s = { %s, %s, (EPropertyFlags)0x%016llx, UECodeGen_Private::EPropertyGenFlags::Enum, %s, %s, %s, %s, %s };%s\r\n"),
+				Spaces,
+				*Name,
+				*PropName,
+				*PropNotifyFunc,
+				PropFlags,
+				FPropertyObjectFlags,
+				*ArrayDim,
+				OffsetStr,
+				*GetSingletonNameFuncAddr(PropertyBase.EnumDef, OutReferenceGatherers.UniqueCrossModuleReferences),
+				*MetaDataParams,
+				*PropTag
+			);
+		}
+		break;
 	}
 
-	if (FFieldPathProperty* TypedProp = CastField<FFieldPathProperty>(Prop))
+	case EUHTPropertyType::FieldPath:
 	{
+		FString MetaDataParams = OutputMetaDataCodeForObject(DeclOut, Out, PropertyDef, *FString::Printf(TEXT("%s_MetaData"), *Name), DeclSpaces, Spaces);
+
 		DeclOut.Logf(TEXT("%sstatic const UECodeGen_Private::FFieldPathPropertyParams %s;\r\n"), DeclSpaces, *NameWithoutScope);
 
 		Out.Logf(
 			TEXT("%sconst UECodeGen_Private::FFieldPathPropertyParams %s = { %s, %s, (EPropertyFlags)0x%016llx, UECodeGen_Private::EPropertyGenFlags::FieldPath, %s, %s, %s, %s, %s };%s\r\n"),
 			Spaces,
-			Name,
+			*Name,
 			*PropName,
 			*PropNotifyFunc,
 			PropFlags,
 			FPropertyObjectFlags,
 			*ArrayDim,
 			OffsetStr,
-			*FString::Printf(TEXT("&F%s::StaticClass"), *TypedProp->PropertyClass->GetName()),
+			*FString::Printf(TEXT("&F%s::StaticClass"), *PropertyBase.PropertyPathClass->GetName()),
 			*MetaDataParams,
 			*PropTag
 		);
+		break;
+	}
 
-		return;
+	default:
+		check(false);
 	}
 
-	// Unhandled type
-	check(false);
+	PropertyNamesAndPointers.Emplace(MoveTemp(Name), PropertyDef);
 }
 
-bool IsEditorOnlyDataProperty(FProperty* Prop)
+bool IsEditorOnlyDataProperty(FUnrealPropertyDefinitionInfo& PropDef)
 {
-	while (Prop)
+	for (FUnrealPropertyDefinitionInfo* TestPropDef = &PropDef; TestPropDef; TestPropDef = UHTCast<FUnrealPropertyDefinitionInfo>(TestPropDef->GetOuter()))
 	{
-		if (Prop->IsEditorOnlyProperty())
+		if (TestPropDef->IsEditorOnlyProperty())
 		{
 			return true;
 		}
+	}
+	return false;
+}
 
-		Prop = Prop->GetOwner<FProperty>();
+FString GetEventStructParamsName(FUnrealObjectDefinitionInfo& OuterDef, const TCHAR* FunctionName)
+{
+	FString OuterName;
+	if (FUnrealClassDefinitionInfo* ClassDef = UHTCast<FUnrealClassDefinitionInfo>(OuterDef))
+	{
+		OuterName = ClassDef->GetName();
+	}
+	else if (FUnrealPackageDefinitionInfo* PackageDef = UHTCast<FUnrealPackageDefinitionInfo>(OuterDef))
+	{
+		OuterName = PackageDef->GetPackage()->GetName();
+		OuterName.ReplaceInline(TEXT("/"), TEXT("_"), ESearchCase::CaseSensitive);
+	}
+	else
+	{
+		FError::Throwf(TEXT("Unrecognized outer type"));
 	}
 
-	return false;
+	FString Result = FString::Printf(TEXT("%s_event%s_Parms"), *OuterName, FunctionName);
+	if (Result.Len() && FChar::IsDigit(Result[0]))
+	{
+		Result.InsertAt(0, TCHAR('_'));
+	}
+	return Result;
 }
 
-TTuple<FString, FString> FNativeClassHeaderGenerator::OutputProperties(FOutputDevice& DeclOut, FOutputDevice& Out, FReferenceGatherers& OutReferenceGatherers, const TCHAR* Scope, const TArray<FUnrealPropertyDefinitionInfo*>& PropertyDefs, const TCHAR* DeclSpaces, const TCHAR* Spaces) const
+TTuple<FString, FString> FNativeClassHeaderGenerator::OutputProperties(FOutputDevice& DeclOut, FOutputDevice& Out, FReferenceGatherers& OutReferenceGatherers, const TCHAR* Scope, FUnrealStructDefinitionInfo& StructDef, const TCHAR* DeclSpaces, const TCHAR* Spaces) const
 {
-	if (PropertyDefs.Num() == 0)
+	if (StructDef.GetProperties().Num() == 0)
 	{
 		return TTuple<FString, FString>(TEXT("nullptr"), TEXT("0"));
 	}
@@ -1684,12 +1675,32 @@ TTuple<FString, FString> FNativeClassHeaderGenerator::OutputProperties(FOutputDe
 	bool bHasAllEditorOnlyDataProperties = true;
 
 	{
+		FString SourceStruct;
+		if (FUnrealFunctionDefinitionInfo* FunctionDef = UHTCast<FUnrealFunctionDefinitionInfo>(StructDef))
+		{
+			while (FunctionDef->GetSuperFunction())
+			{
+				FunctionDef = FunctionDef->GetSuperFunction();
+			}
+			FString FunctionName = FunctionDef->GetName();
+			if (FunctionDef->HasAnyFunctionFlags(FUNC_Delegate))
+			{
+				FunctionName.LeftChopInline(HEADER_GENERATED_DELEGATE_SIGNATURE_SUFFIX_LENGTH, false);
+			}
+
+			SourceStruct = GetEventStructParamsName(*FunctionDef->GetOuter(), *FunctionName);
+		}
+		else 
+		{
+			SourceStruct = StructDef.GetAlternateNameCPP();
+		}
+
 		FMacroBlockEmitter WithEditorOnlyMacroEmitter(Out, TEXT("WITH_EDITORONLY_DATA"));
 		FMacroBlockEmitter WithEditorOnlyMacroEmitterDecl(DeclOut, TEXT("WITH_EDITORONLY_DATA"));
 
-		for (FUnrealPropertyDefinitionInfo* PropDef : PropertyDefs)
+		for (FUnrealPropertyDefinitionInfo* PropertyDef : StructDef.GetProperties())
 		{
-			bool bRequiresHasEditorOnlyMacro = IsEditorOnlyDataProperty(PropDef->GetProperty());
+			bool bRequiresHasEditorOnlyMacro = IsEditorOnlyDataProperty(*PropertyDef);
 			if (!bRequiresHasEditorOnlyMacro)
 			{
 				bHasAllEditorOnlyDataProperties = false;
@@ -1697,7 +1708,18 @@ TTuple<FString, FString> FNativeClassHeaderGenerator::OutputProperties(FOutputDe
 
 			WithEditorOnlyMacroEmitter(bRequiresHasEditorOnlyMacro);
 			WithEditorOnlyMacroEmitterDecl(bRequiresHasEditorOnlyMacro);
-			OutputProperty(DeclOut, Out, OutReferenceGatherers, Scope, PropertyNamesAndPointers, *PropDef, DeclSpaces, Spaces);
+
+			FString PropName = PropertyDef->GetName();
+			FString PropVariableName = FString::Printf(TEXT("%sNewProp_%s"), Scope, *PropName);
+
+			if (PropertyDef->HasAllPropertyFlags(CPF_Deprecated))
+			{
+				PropName += TEXT("_DEPRECATED");
+			}
+
+			FString PropMacroOuterClass = FString::Printf(TEXT("STRUCT_OFFSET(%s, %s)"), *SourceStruct, *PropName);
+
+			OutputProperty(DeclOut, Out, OutReferenceGatherers, Scope, PropertyNamesAndPointers, *PropertyDef, *PropMacroOuterClass, MoveTemp(PropVariableName), DeclSpaces, Spaces, *SourceStruct);
 		}
 
 		WithEditorOnlyMacroEmitter(bHasAllEditorOnlyDataProperties);
@@ -1707,7 +1729,7 @@ TTuple<FString, FString> FNativeClassHeaderGenerator::OutputProperties(FOutputDe
 
 		for (const FPropertyNamePointerPair& PropNameAndPtr : PropertyNamesAndPointers)
 		{
-			bool bRequiresHasEditorOnlyMacro = IsEditorOnlyDataProperty(PropNameAndPtr.PropDef->GetProperty());
+			bool bRequiresHasEditorOnlyMacro = IsEditorOnlyDataProperty(*PropNameAndPtr.PropDef);
 
 			WithEditorOnlyMacroEmitter(bRequiresHasEditorOnlyMacro);
 			WithEditorOnlyMacroEmitterDecl(bRequiresHasEditorOnlyMacro);
@@ -1735,136 +1757,6 @@ TTuple<FString, FString> FNativeClassHeaderGenerator::OutputProperties(FOutputDe
 	}
 }
 
-inline FString GetEventStructParamsName(FUnrealObjectDefinitionInfo& OuterDef, const TCHAR* FunctionName)
-{
-	FString OuterName;
-	if (FUnrealClassDefinitionInfo* ClassDef = UHTCast<FUnrealClassDefinitionInfo>(OuterDef))
-	{
-		OuterName = ClassDef->GetName();
-	}
-	else if (FUnrealPackageDefinitionInfo* PackageDef = UHTCast<FUnrealPackageDefinitionInfo>(OuterDef))
-	{
-		OuterName = PackageDef->GetPackage()->GetName();
-		OuterName.ReplaceInline(TEXT("/"), TEXT("_"), ESearchCase::CaseSensitive);
-	}
-	else
-	{
-		FError::Throwf(TEXT("Unrecognized outer type"));
-	}
-
-	FString Result = FString::Printf(TEXT("%s_event%s_Parms"), *OuterName, FunctionName);
-	if (Result.Len() && FChar::IsDigit(Result[0]))
-	{
-		Result.InsertAt(0, TCHAR('_'));
-	}
-	return Result;
-}
-
-void FNativeClassHeaderGenerator::OutputProperty(FOutputDevice& DeclOut, FOutputDevice& Out, FReferenceGatherers& OutReferenceGatherers, const TCHAR* Scope, TArray<FPropertyNamePointerPair>& PropertyNamesAndPointers, FUnrealPropertyDefinitionInfo& PropertyDef, const TCHAR* DeclSpaces, const TCHAR* Spaces) const
-{
-	FProperty* Prop = PropertyDef.GetProperty();
-
-	// Helper to handle the creation of the underlying properties if they're enum properties
-	auto HandleUnderlyingEnumProperty = [this, &PropertyNamesAndPointers, &DeclOut, &Out, &OutReferenceGatherers, DeclSpaces, Spaces](FUnrealPropertyDefinitionInfo& LocalPropDef, FString&& InOuterName)
-	{
-		FProperty* LocalProp = LocalPropDef.GetProperty();
-		if (FEnumProperty* EnumProp = CastField<FEnumProperty>(LocalProp))
-		{
-			FString PropVarName = InOuterName + TEXT("_Underlying");
-
-			FUnrealPropertyDefinitionInfo& UnderlyingDef = LocalPropDef.GetValuePropDef();
-			FProperty* Underlying = UnderlyingDef.GetProperty();
-
-			checkSlow(Underlying == EnumProp->UnderlyingProp);
-			PropertyNew(DeclOut, Out, OutReferenceGatherers, UnderlyingDef, TEXT("0"), *PropVarName, DeclSpaces, Spaces);
-			PropertyNamesAndPointers.Emplace(MoveTemp(PropVarName), UnderlyingDef);
-		}
-
-		PropertyNamesAndPointers.Emplace(MoveTemp(InOuterName), LocalPropDef);
-	};
-
-	if (FArrayProperty* ArrayProperty = CastField<FArrayProperty>(Prop))
-	{
-		FUnrealPropertyDefinitionInfo& InnerDef = PropertyDef.GetValuePropDef();
-		FProperty* Inner = InnerDef.GetProperty();
-
-		FString InnerVariableName = FString::Printf(TEXT("%sNewProp_%s_Inner"), Scope, *Inner->GetName());
-
-		checkSlow(PropertyDef.GetValuePropDef().GetProperty() == ArrayProperty->Inner);
-
-		HandleUnderlyingEnumProperty(InnerDef, CopyTemp(InnerVariableName));
-		PropertyNew(DeclOut, Out, OutReferenceGatherers, PropertyDef.GetValuePropDef(), TEXT("0"), *InnerVariableName, DeclSpaces, Spaces);
-	}
-
-	else if (FMapProperty* MapProperty = CastField<FMapProperty>(Prop))
-	{
-		FUnrealPropertyDefinitionInfo& KeyDef = PropertyDef.GetKeyPropDef();
-		FProperty* Key = KeyDef.GetProperty();
-		FUnrealPropertyDefinitionInfo& ValueDef = PropertyDef.GetValuePropDef();
-		FProperty* Value = ValueDef.GetProperty();
-
-		checkSlow(PropertyDef.GetValuePropDef().GetProperty() == MapProperty->ValueProp);
-		checkSlow(PropertyDef.GetKeyPropDef().GetProperty() == MapProperty->KeyProp);
-
-		FString KeyVariableName   = FString::Printf(TEXT("%sNewProp_%s_KeyProp"), Scope, *Key->GetName());
-		FString ValueVariableName = FString::Printf(TEXT("%sNewProp_%s_ValueProp"), Scope, *Value->GetName());
-
-		HandleUnderlyingEnumProperty(ValueDef, CopyTemp(ValueVariableName));
-		PropertyNew(DeclOut, Out, OutReferenceGatherers, ValueDef, TEXT("1"), *ValueVariableName, DeclSpaces, Spaces);
-
-		HandleUnderlyingEnumProperty(KeyDef, CopyTemp(KeyVariableName));
-		PropertyNew(DeclOut, Out, OutReferenceGatherers, KeyDef, TEXT("0"), *KeyVariableName, DeclSpaces, Spaces);
-	}
-
-	else if (FSetProperty* SetProperty = CastField<FSetProperty>(Prop))
-	{
-		FUnrealPropertyDefinitionInfo& InnerDef = PropertyDef.GetValuePropDef();
-		FProperty* Inner = InnerDef.GetProperty();
-
-		FString ElementVariableName = FString::Printf(TEXT("%sNewProp_%s_ElementProp"), Scope, *Inner->GetName());
-
-		checkSlow(PropertyDef.GetValuePropDef().GetProperty() == SetProperty->ElementProp);
-
-		HandleUnderlyingEnumProperty(InnerDef, CopyTemp(ElementVariableName));
-		PropertyNew(DeclOut, Out, OutReferenceGatherers, InnerDef, TEXT("0"), *ElementVariableName, DeclSpaces, Spaces);
-	}
-
-	{
-		FString SourceStruct;
-		if (FUnrealFunctionDefinitionInfo* FunctionDef = UHTCast<FUnrealFunctionDefinitionInfo>(PropertyDef.GetOuter()))
-		{
-			while (FunctionDef->GetSuperFunction())
-			{
-				FunctionDef = FunctionDef->GetSuperFunction();
-			}
-			FString FunctionName = FunctionDef->GetName();
-			if (FunctionDef->HasAnyFunctionFlags(FUNC_Delegate))
-			{
-				FunctionName.LeftChopInline(HEADER_GENERATED_DELEGATE_SIGNATURE_SUFFIX_LENGTH, false);
-			}
-
-			SourceStruct = GetEventStructParamsName(*FunctionDef->GetOuter(), *FunctionName);
-		}
-		else
-		{
-			SourceStruct = FNameLookupCPP::GetNameCPP(CastChecked<UStruct>(Prop->GetOwner<UObject>()));
-		}
-
-		FString PropName = Prop->GetName();
-		FString PropVariableName = FString::Printf(TEXT("%sNewProp_%s"), Scope, *PropName);
-
-		if (Prop->HasAllPropertyFlags(CPF_Deprecated))
-		{
-			PropName += TEXT("_DEPRECATED");
-		}
-
-		FString PropMacroOuterClass = FString::Printf(TEXT("STRUCT_OFFSET(%s, %s)"), *SourceStruct, *PropName);
-
-		HandleUnderlyingEnumProperty(PropertyDef, CopyTemp(PropVariableName));
-		PropertyNew(DeclOut, Out, OutReferenceGatherers, PropertyDef, *PropMacroOuterClass, *PropVariableName, DeclSpaces, Spaces, *SourceStruct);
-	}
-}
-
 static bool IsAlwaysAccessible(FUnrealScriptStructDefinitionInfo& ScriptDef)
 {
 	FName ToTest = ScriptDef.GetFName();
@@ -2175,7 +2067,7 @@ void FNativeClassHeaderGenerator::ExportNativeGeneratedInitCode(FOutputDevice& O
 
 		FString MetaDataParams = OutputMetaDataCodeForObject(GeneratedClassRegisterFunctionText, StaticDefinitions, ClassDef, *FString::Printf(TEXT("%s::Class_MetaDataParams"), *StaticsStructName), TEXT("\t\t"), TEXT("\t"));
 
-		TTuple<FString, FString> PropertyRange = OutputProperties(GeneratedClassRegisterFunctionText, StaticDefinitions, OutReferenceGatherers, *FString::Printf(TEXT("%s::"), *StaticsStructName), ClassDef.GetProperties(), TEXT("\t\t"), TEXT("\t"));
+		TTuple<FString, FString> PropertyRange = OutputProperties(GeneratedClassRegisterFunctionText, StaticDefinitions, OutReferenceGatherers, *FString::Printf(TEXT("%s::"), *StaticsStructName), ClassDef, TEXT("\t\t"), TEXT("\t"));
 
 		const TCHAR* InterfaceArray;
 		const TCHAR* InterfaceCount;
@@ -2335,8 +2227,7 @@ void FNativeClassHeaderGenerator::ExportNativeGeneratedInitCode(FOutputDevice& O
 			GeneratedClassRegisterFunctionText.Logf(TEXT("%s\r\n"), *SparseScriptStructDef->GetName());
 			for (FUnrealPropertyDefinitionInfo* ChildDef : TUHTFieldRange<FUnrealPropertyDefinitionInfo>(*SparseScriptStructDef))
 			{
-				FProperty* Child = ChildDef->GetProperty();
-				GeneratedClassRegisterFunctionText.Logf(TEXT("%s %s\r\n"), *Child->GetCPPType(), *Child->GetNameCPP());
+				GeneratedClassRegisterFunctionText.Logf(TEXT("%s %s\r\n"), *ChildDef->GetCPPType(), *ChildDef->GetPropertyNameCPP());
 			}
 		}
 	}
@@ -2497,7 +2388,7 @@ void FNativeClassHeaderGenerator::ExportFunction(FOutputDevice& Out, FReferenceG
 	USparseDelegateFunction* SparseDelegateFunction = Cast<USparseDelegateFunction>(FunctionDef.GetFunction());
 	const TCHAR* UFunctionObjectFlags = FunctionDef.IsOwnedByDynamicType() ? TEXT("RF_Public|RF_Transient") : TEXT("RF_Public|RF_Transient|RF_MarkAsNative");
 
-	TTuple<FString, FString> PropertyRange = OutputProperties(CurrentFunctionText, StaticDefinitions, OutReferenceGatherers, *FString::Printf(TEXT("%s::"), *StaticsStructName), FunctionDef.GetProperties(), TEXT("\t\t"), TEXT("\t"));
+	TTuple<FString, FString> PropertyRange = OutputProperties(CurrentFunctionText, StaticDefinitions, OutReferenceGatherers, *FString::Printf(TEXT("%s::"), *StaticsStructName), FunctionDef, TEXT("\t\t"), TEXT("\t"));
 
 	const FFuncInfo&     FunctionData = FunctionDef.GetFunctionData();
 	const bool           bIsNet       = FunctionDef.HasAnyFunctionFlags(FUNC_NetRequest | FUNC_NetResponse);
@@ -2660,7 +2551,7 @@ void FNativeClassHeaderGenerator::ExportInterfaceCallFunctions(FOutputDevice& Ou
 		// code to populate Parms struct
 		for (FUnrealPropertyDefinitionInfo* ParamDef : Parameters.Parms)
 		{
-			const FString ParamName = ParamDef->GetProperty()->GetName();
+			const FString ParamName = ParamDef->GetName();
 			OutCpp.Logf(TEXT("\t\t\tParms.%s=%s;") LINE_TERMINATOR, *ParamName, *ParamName);
 		}
 
@@ -2669,10 +2560,9 @@ void FNativeClassHeaderGenerator::ExportInterfaceCallFunctions(FOutputDevice& Ou
 
 		for (FUnrealPropertyDefinitionInfo* ParamDef : Parameters.Parms)
 		{
-			FProperty* Param = ParamDef->GetProperty();
-			if( Param->HasAllPropertyFlags(CPF_OutParm) && !Param->HasAnyPropertyFlags(CPF_ConstParm|CPF_ReturnParm))
+			if(ParamDef->HasAllPropertyFlags(CPF_OutParm) && !ParamDef->HasAnyPropertyFlags(CPF_ConstParm|CPF_ReturnParm))
 			{
-				const FString ParamName = Param->GetName();
+				const FString ParamName = ParamDef->GetName();
 				OutCpp.Logf(TEXT("\t\t\t%s=Parms.%s;") LINE_TERMINATOR, *ParamName, *ParamName);
 			}
 		}
@@ -2703,7 +2593,7 @@ void FNativeClassHeaderGenerator::ExportInterfaceCallFunctions(FOutputDevice& Ou
 				}
 				bFirst = false;
 
-				OutCpp.Logf(TEXT("%s"), *ParamDef->GetProperty()->GetName());
+				OutCpp.Logf(TEXT("%s"), *ParamDef->GetName());
 			}
 
 			OutCpp.Logf(TEXT(");") LINE_TERMINATOR);
@@ -3808,7 +3698,7 @@ void FNativeClassHeaderGenerator::ExportGeneratedStructBodyMacros(FOutputDevice&
 		NewStructOps = TEXT("nullptr");
 	}
 
-	TTuple<FString, FString> PropertyRange = OutputProperties(GeneratedStructRegisterFunctionText, StaticDefinitions, OutReferenceGatherers, *FString::Printf(TEXT("%s::"), *StaticsStructName), ScriptStructDef.GetProperties(), TEXT("\t\t"), TEXT("\t"));
+	TTuple<FString, FString> PropertyRange = OutputProperties(GeneratedStructRegisterFunctionText, StaticDefinitions, OutReferenceGatherers, *FString::Printf(TEXT("%s::"), *StaticsStructName), ScriptStructDef, TEXT("\t\t"), TEXT("\t"));
 
 	GeneratedStructRegisterFunctionText.Log (TEXT("\t\tstatic const UECodeGen_Private::FStructParams ReturnStructParams;\r\n"));
 
@@ -4156,7 +4046,7 @@ void FNativeClassHeaderGenerator::ExportMirrorsForNoexportStruct(FOutputDevice&
 bool FNativeClassHeaderGenerator::WillExportEventParms(FUnrealFunctionDefinitionInfo& FunctionDef)
 {
 	const TArray<FUnrealPropertyDefinitionInfo*>& Properties = FunctionDef.GetProperties();
-	return Properties.Num() > 0 && (Properties[0]->GetProperty()->PropertyFlags & CPF_Parm) != 0;
+	return Properties.Num() > 0 && Properties[0]->HasAnyPropertyFlags(CPF_Parm);
 }
 
 void WriteEventFunctionPrologue(FOutputDevice& Output, int32 Indent, const FParmsAndReturnProperties& Parameters, FUnrealObjectDefinitionInfo& FunctionOuterDef, const TCHAR* FunctionName)
@@ -4175,16 +4065,15 @@ void WriteEventFunctionPrologue(FOutputDevice& Output, int32 Indent, const FParm
 	// Declare a parameter struct for this event/delegate and assign the struct members using the values passed into the event/delegate call.
 	for (FUnrealPropertyDefinitionInfo* PropDef : Parameters.Parms)
 	{
-		FProperty* Prop = PropDef->GetProperty();
-		const FString PropertyName = Prop->GetName();
-		if (Prop->ArrayDim > 1)
+		const FString PropertyName = PropDef->GetName();
+		if (PropDef->IsStaticArray())
 		{
 			Output.Logf(TEXT("%sFMemory::Memcpy(Parms.%s,%s,sizeof(Parms.%s));\r\n"), FCString::Tab(Indent + 1), *PropertyName, *PropertyName, *PropertyName);
 		}
 		else
 		{
 			FString ValueAssignmentText = PropertyName;
-			if (Prop->IsA<FBoolProperty>())
+			if (PropDef->IsBooleanOrBooleanStaticArray())
 			{
 				ValueAssignmentText += TEXT(" ? true : false");
 			}
@@ -4199,11 +4088,10 @@ void WriteEventFunctionEpilogue(FOutputDevice& Output, int32 Indent, const FParm
 	// Out parm copying.
 	for (FUnrealPropertyDefinitionInfo* PropDef : Parameters.Parms)
 	{
-		FProperty* Prop = PropDef->GetProperty();
-		if ((Prop->PropertyFlags & (CPF_OutParm | CPF_ConstParm)) == CPF_OutParm)
+		if (PropDef->HasSpecificPropertyFlags(CPF_OutParm | CPF_ConstParm, CPF_OutParm))
 		{
-			const FString PropertyName = Prop->GetName();
-			if ( Prop->ArrayDim > 1 )
+			const FString PropertyName = PropDef->GetName();
+			if (PropDef->IsStaticArray())
 			{
 				Output.Logf(TEXT("%sFMemory::Memcpy(&%s,&Parms.%s,sizeof(%s));\r\n"), FCString::Tab(Indent + 1), *PropertyName, *PropertyName, *PropertyName);
 			}
@@ -4218,9 +4106,8 @@ void WriteEventFunctionEpilogue(FOutputDevice& Output, int32 Indent, const FParm
 	if (Parameters.Return)
 	{
 		// Make sure uint32 -> bool is supported
-		FProperty* Return = Parameters.Return->GetProperty();
-		bool bBoolProperty = Return->IsA(FBoolProperty::StaticClass());
-		Output.Logf(TEXT("%sreturn %sParms.%s;\r\n"), FCString::Tab(Indent + 1), bBoolProperty ? TEXT("!!") : TEXT(""), *Return->GetName());
+		bool bBoolProperty = Parameters.Return->IsBooleanOrBooleanStaticArray();
+		Output.Logf(TEXT("%sreturn %sParms.%s;\r\n"), FCString::Tab(Indent + 1), bBoolProperty ? TEXT("!!") : TEXT(""), *Parameters.Return->GetName());
 	}
 	Output.Logf(TEXT("%s}\r\n"), FCString::Tab(Indent));
 }
@@ -4332,24 +4219,25 @@ void FNativeClassHeaderGenerator::ExportEventParm(FUHTStringBuilder& Out, TSet<F
 
 	for (FUnrealPropertyDefinitionInfo* PropDef : FunctionDef.GetProperties())
 	{
-		FProperty* Prop = PropDef->GetProperty();
-		if (!(Prop->PropertyFlags & CPF_Parm))
+		FPropertyBase& PropertyBase = PropDef->GetPropertyBase();
+		if (!PropDef->HasAnyPropertyFlags(CPF_Parm))
 		{
 			continue;
 		}
 
-		PropertyFwd.Add(Prop->GetCPPTypeForwardDeclaration());
+		PropertyFwd.Add(PropDef->GetCPPTypeForwardDeclaration());
 
 		FUHTStringBuilder PropertyText;
 		PropertyText.Log(FCString::Tab(Indent + 1));
 
-		bool bEmitConst = Prop->HasAnyPropertyFlags(CPF_ConstParm) && Prop->IsA<FObjectProperty>();
+		bool bEmitConst = PropDef->HasAnyPropertyFlags(CPF_ConstParm) && PropertyBase.IsObjectRefOrObjectRefStaticArray();
 
 		//@TODO: UCREMOVAL: This is awful code duplication to avoid a double-const
 		{
+			//@TODO: bEmitConst will only be true if we have an object, so checking interface here doesn't do anything.
 			// export 'const' for parameters
-			const bool bIsConstParam = (Prop->IsA(FInterfaceProperty::StaticClass()) && !Prop->HasAllPropertyFlags(CPF_OutParm)); //@TODO: This should be const once that flag exists
-			const bool bIsOnConstClass = (Prop->IsA(FObjectProperty::StaticClass()) && ((FObjectProperty*)Prop)->PropertyClass != NULL && ((FObjectProperty*)Prop)->PropertyClass->HasAnyClassFlags(CLASS_Const));
+			const bool bIsConstParam = PropertyBase.IsInterfaceOrInterfaceStaticArray() && !PropDef->HasAllPropertyFlags(CPF_OutParm); //@TODO: This should be const once that flag exists
+			const bool bIsOnConstClass = PropertyBase.IsObjectRefOrObjectRefStaticArray() && PropertyBase.ClassDef->HasAnyClassFlags(CLASS_Const);
 
 			if (bIsConstParam || bIsOnConstClass)
 			{
@@ -4362,7 +4250,7 @@ void FNativeClassHeaderGenerator::ExportEventParm(FUHTStringBuilder& Out, TSet<F
 			PropertyText.Logf(TEXT("const "));
 		}
 
-		Prop->ExportCppDeclaration(PropertyText, EExportedDeclaration::Local, PropDef->GetArrayDimensions());
+		PropDef->ExportCppDeclaration(PropertyText, EExportedDeclaration::Local, PropDef->GetArrayDimensions());
 		ApplyAlternatePropertyExportText(*PropDef, PropertyText, ExportingState);
 
 		PropertyText.Log(TEXT(";\r\n"));
@@ -4373,36 +4261,34 @@ void FNativeClassHeaderGenerator::ExportEventParm(FUHTStringBuilder& Out, TSet<F
 	FUnrealPropertyDefinitionInfo* PropDef = FunctionDef.GetReturn();
 	if (PropDef && bOutputConstructor)
 	{
-		FProperty* Prop = PropDef->GetProperty();
+		const FPropertyBase& PropertyBase = PropDef->GetPropertyBase();
 		FUHTStringBuilder InitializationAr;
 
-		FStructProperty* InnerStruct = CastField<FStructProperty>(Prop);
 		bool bNeedsOutput = true;
-		if (InnerStruct)
-		{
-			bNeedsOutput = InnerStruct->HasNoOpConstructor();
-		}
-		else if (
-			CastField<FNameProperty>(Prop) ||
-			CastField<FDelegateProperty>(Prop) ||
-			CastField<FMulticastDelegateProperty>(Prop) ||
-			CastField<FStrProperty>(Prop) ||
-			CastField<FTextProperty>(Prop) ||
-			CastField<FArrayProperty>(Prop) ||
-			CastField<FMapProperty>(Prop) ||
-			CastField<FSetProperty>(Prop) ||
-			CastField<FInterfaceProperty>(Prop) ||
-			CastField<FFieldPathProperty>(Prop)
-			)
+		switch (PropertyBase.GetUHTPropertyType())
 		{
+		case EUHTPropertyType::Struct:
+			bNeedsOutput = PropDef->HasNoOpConstructor();
+			break;
+		case EUHTPropertyType::Name:
+		case EUHTPropertyType::Delegate:
+		case EUHTPropertyType::MulticastDelegate:
+		case EUHTPropertyType::String:
+		case EUHTPropertyType::Text:
+		case EUHTPropertyType::DynamicArray:
+		case EUHTPropertyType::Map:
+		case EUHTPropertyType::Set:
+		case EUHTPropertyType::Interface:
+		case EUHTPropertyType::FieldPath:
 			bNeedsOutput = false;
 		}
+
 		if (bNeedsOutput)
 		{
-			check(Prop->ArrayDim == 1); // can't return arrays
+			check(!PropDef->IsStaticArray()); // can't return arrays
 			Out.Logf(TEXT("\r\n%s/** Constructor, initializes return property only **/\r\n"), FCString::Tab(Indent + 1));
 			Out.Logf(TEXT("%s%s()\r\n"), FCString::Tab(Indent + 1), *EventParmStructName);
-			Out.Logf(TEXT("%s%s %s(%s)\r\n"), FCString::Tab(Indent + 2), TEXT(":"), *Prop->GetName(), *GetNullParameterValue(*PropDef, true));
+			Out.Logf(TEXT("%s%s %s(%s)\r\n"), FCString::Tab(Indent + 2), TEXT(":"), *PropDef->GetName(), *GetNullParameterValue(*PropDef, true));
 			Out.Logf(TEXT("%s{\r\n"), FCString::Tab(Indent + 1));
 			Out.Logf(TEXT("%s}\r\n"), FCString::Tab(Indent + 1));
 		}
@@ -4418,86 +4304,89 @@ void FNativeClassHeaderGenerator::ExportEventParm(FUHTStringBuilder& Out, TSet<F
  *
  * @return	the intrinsic null value for the property (0 for ints, TEXT("") for strings, etc.)
  */
-FString FNativeClassHeaderGenerator::GetNullParameterValue(FUnrealPropertyDefinitionInfo& PropertyDef, bool bInitializer/*=false*/ )
+FString FNativeClassHeaderGenerator::GetNullParameterValue(FUnrealPropertyDefinitionInfo& PropertyDef, bool bInitializer/*=false*/)
 {
-	FProperty* Prop = PropertyDef.GetProperty();
-	FFieldClass* PropClass = Prop->GetClass();
-	FObjectPropertyBase* ObjectProperty = CastField<FObjectPropertyBase>(Prop);
-	if (PropClass == FByteProperty::StaticClass())
-	{
-		FByteProperty* ByteProp = (FByteProperty*)Prop;
-
-		// if it's an enum class then we need an explicit cast
-		if( ByteProp->Enum && ByteProp->Enum->GetCppForm() == UEnum::ECppForm::EnumClass )
-		{
-			return FString::Printf(TEXT("(%s)0"), *ByteProp->GetCPPType());
-		}
-
+	const FPropertyBase& PropertyBase = PropertyDef.GetPropertyBase();
+	switch (PropertyBase.GetUHTPropertyType())
+	{
+	case EUHTPropertyType::Byte:
+	case EUHTPropertyType::Int16:
+	case EUHTPropertyType::Int:
+	case EUHTPropertyType::Int64:
+	case EUHTPropertyType::UInt16:
+	case EUHTPropertyType::UInt32:
+	case EUHTPropertyType::UInt64:
+	case EUHTPropertyType::Float:
+	case EUHTPropertyType::Double:
+	case EUHTPropertyType::LargeWorldCoordinatesReal:
 		return TEXT("0");
-	}
-	else if (PropClass == FEnumProperty::StaticClass())
-	{
-		FEnumProperty* EnumProp = (FEnumProperty*)Prop;
 
-		return FString::Printf(TEXT("(%s)0"), *EnumProp->Enum->GetName());
-	}
-	else if ( PropClass == FBoolProperty::StaticClass() )
-	{
+	case EUHTPropertyType::Bool:
+	case EUHTPropertyType::Bool8:
+	case EUHTPropertyType::Bool16:
+	case EUHTPropertyType::Bool32:
+	case EUHTPropertyType::Bool64:
 		return TEXT("false");
-	}
-	else if ( PropClass->IsChildOf(FNumericProperty::StaticClass()) )
-	{
-		return TEXT("0");
-	}
-	else if ( PropClass == FNameProperty::StaticClass() )
-	{
+
+	case EUHTPropertyType::ObjectReference:
+	case EUHTPropertyType::ObjectPtrReference:
+	case EUHTPropertyType::SoftObjectReference:
+	case EUHTPropertyType::WeakObjectReference:
+	case EUHTPropertyType::LazyObjectReference:
+		return TEXT("NULL");
+
+	case EUHTPropertyType::Interface:
+		return TEXT("NULL");
+
+	case EUHTPropertyType::Name:
 		return TEXT("NAME_None");
-	}
-	else if ( PropClass == FStrProperty::StaticClass() )
-	{
+
+	case EUHTPropertyType::String:
 		return TEXT("TEXT(\"\")");
-	}
-	else if ( PropClass == FTextProperty::StaticClass() )
-	{
-		return TEXT("FText::GetEmpty()");
-	}
-	else if ( PropClass == FArrayProperty::StaticClass()
-		||    PropClass == FMapProperty::StaticClass()
-		||    PropClass == FSetProperty::StaticClass()
-		||    PropClass == FDelegateProperty::StaticClass()
-		||    PropClass == FMulticastDelegateProperty::StaticClass() )
+
+	case EUHTPropertyType::DynamicArray:
+	case EUHTPropertyType::Map:
+	case EUHTPropertyType::Set:
+	case EUHTPropertyType::Delegate:
+	case EUHTPropertyType::MulticastDelegate:
 	{
 		FString Type, ExtendedType;
-		Type = Prop->GetCPPType(&ExtendedType,CPPF_OptionalValue);
+		Type = PropertyDef.GetCPPType(&ExtendedType, CPPF_OptionalValue);
 		return Type + ExtendedType + TEXT("()");
 	}
-	else if ( PropClass == FStructProperty::StaticClass() )
+
+	case EUHTPropertyType::Struct:
 	{
-		bool bHasNoOpConstuctor = CastFieldChecked<FStructProperty>(Prop)->HasNoOpConstructor();
+		bool bHasNoOpConstuctor = PropertyDef.HasNoOpConstructor();
 		if (bInitializer && bHasNoOpConstuctor)
 		{
 			return TEXT("ForceInit");
 		}
 
 		FString Type, ExtendedType;
-		Type = Prop->GetCPPType(&ExtendedType,CPPF_OptionalValue);
+		Type = PropertyDef.GetCPPType(&ExtendedType, CPPF_OptionalValue);
 		return Type + ExtendedType + (bHasNoOpConstuctor ? TEXT("(ForceInit)") : TEXT("()"));
 	}
-	else if (ObjectProperty)
-	{
-		return TEXT("NULL");
-	}
-	else if ( PropClass == FInterfaceProperty::StaticClass() )
+
+	case EUHTPropertyType::Text:
+		return TEXT("FText::GetEmpty()");
+
+	case EUHTPropertyType::Enum:
 	{
-		return TEXT("NULL");
+		if (PropertyBase.EnumDef->GetCppForm() != UEnum::ECppForm::EnumClass)
+		{
+			return TEXT("0");
+		}
+		return FString::Printf(TEXT("(%s)0"), *PropertyDef.GetCPPType());
 	}
-	else if (PropClass == FFieldPathProperty::StaticClass())
-	{
+
+	case EUHTPropertyType::FieldPath:
 		return TEXT("nullptr");
-	}
 
-	UE_LOG(LogCompile, Fatal,TEXT("GetNullParameterValue - Unhandled property type '%s': %s"), *Prop->GetClass()->GetName(), *Prop->GetPathName());
-	return TEXT("");
+	default:
+		UE_LOG(LogCompile, Fatal, TEXT("GetNullParameterValue - Unhandled property type '%s': %s"), *PropertyDef.GetEngineClassName(), *PropertyDef.GetPathName());
+		return TEXT("");
+	}
 }
 
 
@@ -4507,10 +4396,9 @@ FString FNativeClassHeaderGenerator::GetFunctionReturnString(FUnrealFunctionDefi
 
 	if (FUnrealPropertyDefinitionInfo* ReturnDef = FunctionDef.GetReturn())
 	{
-		FProperty* Return = ReturnDef->GetProperty();
 		FString ExtendedReturnType;
-		OutReferenceGatherers.ForwardDeclarations.Add(Return->GetCPPTypeForwardDeclaration());
-		FString ReturnType = Return->GetCPPType(&ExtendedReturnType, CPPF_ArgumentOrReturnValue);
+		OutReferenceGatherers.ForwardDeclarations.Add(ReturnDef->GetCPPTypeForwardDeclaration());
+		FString ReturnType = ReturnDef->GetCPPType(&ExtendedReturnType, CPPF_ArgumentOrReturnValue);
 		FUHTStringBuilder ReplacementText;
 		ReplacementText += MoveTemp(ReturnType);
 		ApplyAlternatePropertyExportText(*ReturnDef, ReplacementText, EExportingState::Normal);
@@ -4714,17 +4602,16 @@ void FNativeClassHeaderGenerator::ExportNativeFunctionHeader(
 	}
 
 	FUnrealPropertyDefinitionInfo* ReturnPropertyDef = FunctionDef.GetReturn();
-	FProperty* ReturnProperty = ReturnPropertyDef ? ReturnPropertyDef->GetProperty() : nullptr;
-	if (ReturnProperty != nullptr)
+	if (ReturnPropertyDef != nullptr)
 	{
-		if (ReturnProperty->HasAnyPropertyFlags(EPropertyFlags::CPF_ConstParm))
+		if (ReturnPropertyDef->HasAnyPropertyFlags(EPropertyFlags::CPF_ConstParm))
 		{
 			Out.Log(TEXT("const "));
 		}
 
 		FString ExtendedReturnType;
-		FString ReturnType = ReturnProperty->GetCPPType(&ExtendedReturnType, (FunctionHeaderStyle == EExportFunctionHeaderStyle::Definition && (FunctionType != EExportFunctionType::Interface) ? CPPF_Implementation : 0) | CPPF_ArgumentOrReturnValue);
-		OutFwdDecls.Add(ReturnProperty->GetCPPTypeForwardDeclaration());
+		FString ReturnType = ReturnPropertyDef->GetCPPType(&ExtendedReturnType, (FunctionHeaderStyle == EExportFunctionHeaderStyle::Definition && (FunctionType != EExportFunctionType::Interface) ? CPPF_Implementation : 0) | CPPF_ArgumentOrReturnValue);
+		OutFwdDecls.Add(ReturnPropertyDef->GetCPPTypeForwardDeclaration());
 		FUHTStringBuilder ReplacementText;
 		ReplacementText += ReturnType;
 		ApplyAlternatePropertyExportText(*ReturnPropertyDef, ReplacementText, EExportingState::Normal);
@@ -4767,13 +4654,12 @@ void FNativeClassHeaderGenerator::ExportNativeFunctionHeader(
 
 	for (FUnrealPropertyDefinitionInfo* PropertyDef : FunctionDef.GetProperties())
 	{
-		FProperty* Property = PropertyDef->GetProperty();
-		if ((Property->PropertyFlags & (CPF_Parm | CPF_ReturnParm)) != CPF_Parm)
+		if (!PropertyDef->HasSpecificPropertyFlags(CPF_Parm | CPF_ReturnParm, CPF_Parm))
 		{
 			continue;
 		}
 
-		OutFwdDecls.Add(Property->GetCPPTypeForwardDeclaration());
+		OutFwdDecls.Add(PropertyDef->GetCPPTypeForwardDeclaration());
 
 		if( ParmCount++ )
 		{
@@ -4782,7 +4668,7 @@ void FNativeClassHeaderGenerator::ExportNativeFunctionHeader(
 
 		FUHTStringBuilder PropertyText;
 
-		Property->ExportCppDeclaration( PropertyText, EExportedDeclaration::Parameter, PropertyDef->GetArrayDimensions() );
+		PropertyDef->ExportCppDeclaration( PropertyText, EExportedDeclaration::Parameter, PropertyDef->GetArrayDimensions() );
 		ApplyAlternatePropertyExportText(*PropertyDef, PropertyText, EExportingState::Normal);
 
 		Out.Logf(TEXT("%s"), *PropertyText);
@@ -4803,12 +4689,11 @@ void FNativeClassHeaderGenerator::ExportNativeFunctionHeader(
 			{
 				// For BlueprintNativeEvent methods we emit a stub implementation. This allows Blueprints that implement the interface class to be nativized.
 				FString ReturnValue;
-				if (ReturnProperty != nullptr)
+				if (ReturnPropertyDef != nullptr)
 				{
-					FByteProperty* ByteProperty = CastField<FByteProperty>(ReturnProperty);
-					if (ByteProperty != nullptr && ByteProperty->Enum != nullptr && ByteProperty->Enum->GetCppForm() != UEnum::ECppForm::EnumClass)
-					{
-						ReturnValue = FString::Printf(TEXT(" return TEnumAsByte<%s>(%s); "), *ByteProperty->Enum->CppType, *GetNullParameterValue(*ReturnPropertyDef, false));
+					if (ReturnPropertyDef->IsByteEnumOrByteEnumStaticArray())
+					{ 
+						ReturnValue = FString::Printf(TEXT(" return TEnumAsByte<%s>(%s); "), *ReturnPropertyDef->GetPropertyBase().AsEnum()->GetCppType(), *GetNullParameterValue(*ReturnPropertyDef, false));
 					}
 					else
 					{
@@ -4844,8 +4729,8 @@ void FNativeClassHeaderGenerator::ExportFunctionThunk(FUHTStringBuilder& RPCWrap
 	for (int32 ParameterIndex = 0; ParameterIndex < ParameterDefs.Num(); ParameterIndex++)
 	{
 		FUnrealPropertyDefinitionInfo* ParamDef = ParameterDefs[ParameterIndex];
-		FProperty* Param = ParamDef->GetProperty();
-		OutReferenceGatherers.ForwardDeclarations.Add(Param->GetCPPTypeForwardDeclaration());
+		const FPropertyBase& PropertyBase = ParamDef->GetPropertyBase();
+		OutReferenceGatherers.ForwardDeclarations.Add(ParamDef->GetCPPTypeForwardDeclaration());
 
 		FString EvalBaseText = TEXT("P_GET_");	// e.g. P_GET_STR
 		FString EvalModifierText;				// e.g. _REF
@@ -4853,32 +4738,27 @@ void FNativeClassHeaderGenerator::ExportFunctionThunk(FUHTStringBuilder& RPCWrap
 
 		FString TypeText;
 
-		if (Param->ArrayDim > 1)
+		if (ParamDef->IsStaticArray())
 		{
 			EvalBaseText += TEXT("ARRAY");
-			TypeText = Param->GetCPPType();
+			TypeText = ParamDef->GetCPPType();
 		}
 		else
 		{
-			EvalBaseText += Param->GetCPPMacroType(TypeText);
+			EvalBaseText += ParamDef->GetCPPMacroType(TypeText);
 
-			FArrayProperty* ArrayProperty = CastField<FArrayProperty>(Param);
-			if (ArrayProperty)
+			if (PropertyBase.ArrayType == EArrayType::Dynamic && PropertyBase.Type == CPT_Interface)
 			{
-				FInterfaceProperty* InterfaceProperty = CastField<FInterfaceProperty>(ArrayProperty->Inner);
-				if (InterfaceProperty)
-				{
-					FString InterfaceTypeText;
-					InterfaceProperty->GetCPPMacroType(InterfaceTypeText);
-					TypeText += FString::Printf(TEXT("<%s>"), *InterfaceTypeText);
-				}
+				FString InterfaceTypeText;
+				ParamDef->GetValuePropDef().GetCPPMacroType(InterfaceTypeText);
+				TypeText += FString::Printf(TEXT("<%s>"), *InterfaceTypeText);
 			}
 		}
 
-		bool bPassAsNoPtr = Param->HasAllPropertyFlags(CPF_UObjectWrapper | CPF_OutParm) && Param->IsA(FClassProperty::StaticClass());
+		bool bPassAsNoPtr = ParamDef->HasAllPropertyFlags(CPF_UObjectWrapper | CPF_OutParm) && ParamDef->IsClassRefOrClassRefStaticArray();
 		if (bPassAsNoPtr)
 		{
-			TypeText = Param->GetCPPType();
+			TypeText = ParamDef->GetCPPType();
 		}
 
 		FUHTStringBuilder ReplacementText;
@@ -4891,7 +4771,7 @@ void FNativeClassHeaderGenerator::ExportFunctionThunk(FUHTStringBuilder& RPCWrap
 		FString ParamPrefix = TEXT("Z_Param_");
 
 		// if this property is an out parm, add the REF tag
-		if (Param->PropertyFlags & CPF_OutParm)
+		if (ParamDef->HasAnyPropertyFlags(CPF_OutParm))
 		{
 			if (!bPassAsNoPtr)
 			{
@@ -4912,7 +4792,7 @@ void FNativeClassHeaderGenerator::ExportFunctionThunk(FUHTStringBuilder& RPCWrap
 			TypeText += TCHAR(',');
 		}
 
-		FString ParamName = ParamPrefix + Param->GetName();
+		FString ParamName = ParamPrefix + ParamDef->GetName();
 
 		EvalParameterText = FString::Printf(TEXT("(%s%s%s)"), *TypeText, *ParamName, *DefaultValueText);
 
@@ -4924,54 +4804,37 @@ void FNativeClassHeaderGenerator::ExportFunctionThunk(FUHTStringBuilder& RPCWrap
 			ParameterList += TCHAR(',');
 		}
 
+		if (PropertyBase.Type == CPT_Delegate && PropertyBase.IsPrimitiveOrPrimitiveStaticArray())
 		{
-			FDelegateProperty* DelegateProp = CastField< FDelegateProperty >(Param);
-			if (DelegateProp != NULL)
-			{
-				// For delegates, add an explicit conversion to the specific type of delegate before passing it along
-				const FString FunctionName = DelegateProp->SignatureFunction->GetName().LeftChop(HEADER_GENERATED_DELEGATE_SIGNATURE_SUFFIX_LENGTH);
-				ParamName = FString::Printf(TEXT("F%s(%s)"), *FunctionName, *ParamName);
-			}
-		}
-
-		{
-			FMulticastDelegateProperty* MulticastDelegateProp = CastField< FMulticastDelegateProperty >(Param);
-			if (MulticastDelegateProp != NULL)
-			{
-				// For delegates, add an explicit conversion to the specific type of delegate before passing it along
-				const FString FunctionName = MulticastDelegateProp->SignatureFunction->GetName().LeftChop(HEADER_GENERATED_DELEGATE_SIGNATURE_SUFFIX_LENGTH);
-				ParamName = FString::Printf(TEXT("F%s(%s)"), *FunctionName, *ParamName);
-			}
+			// For delegates, add an explicit conversion to the specific type of delegate before passing it along
+			const FString FunctionName = PropertyBase.FunctionDef->GetName().LeftChop(HEADER_GENERATED_DELEGATE_SIGNATURE_SUFFIX_LENGTH);
+			ParamName = FString::Printf(TEXT("F%s(%s)"), *FunctionName, *ParamName);
 		}
 
-		UEnum* Enum = nullptr;
-		FByteProperty* ByteProp = CastField<FByteProperty>(Param);
-		if (ByteProp && ByteProp->Enum)
+		if (PropertyBase.Type == CPT_MulticastDelegate && PropertyBase.IsPrimitiveOrPrimitiveStaticArray())
 		{
-			Enum = ByteProp->Enum;
-		}
-		else if (Param->IsA<FEnumProperty>())
-		{
-			Enum = ((FEnumProperty*)Param)->Enum;
+			// For delegates, add an explicit conversion to the specific type of delegate before passing it along
+			const FString FunctionName = PropertyBase.FunctionDef->GetName().LeftChop(HEADER_GENERATED_DELEGATE_SIGNATURE_SUFFIX_LENGTH);
+			ParamName = FString::Printf(TEXT("F%s(%s)"), *FunctionName, *ParamName);
 		}
 
-		if (Enum)
+		if (FUnrealEnumDefinitionInfo* EnumDef = PropertyBase.IsPrimitiveOrPrimitiveStaticArray() ? PropertyBase.AsEnum() : nullptr)
 		{
 			// For enums, add an explicit conversion
-			if (!(Param->PropertyFlags & CPF_OutParm))
+			if (!ParamDef->HasAnyPropertyFlags(CPF_OutParm))
 			{
-				ParamName = FString::Printf(TEXT("%s(%s)"), *Enum->CppType, *ParamName);
+				ParamName = FString::Printf(TEXT("%s(%s)"), *EnumDef->GetCppType(), *ParamName);
 			}
 			else
 			{
-				if (Enum->GetCppForm() == UEnum::ECppForm::EnumClass)
+				if (EnumDef->GetCppForm() == UEnum::ECppForm::EnumClass)
 				{
 					// If we're an enum class don't require the wrapper
-					ParamName = FString::Printf(TEXT("(%s&)(%s)"), *Enum->CppType, *ParamName);
+					ParamName = FString::Printf(TEXT("(%s&)(%s)"), *EnumDef->GetCppType(), *ParamName);
 				}
 				else
 				{
-					ParamName = FString::Printf(TEXT("(TEnumAsByte<%s>&)(%s)"), *Enum->CppType, *ParamName);
+					ParamName = FString::Printf(TEXT("(TEnumAsByte<%s>&)(%s)"), *EnumDef->GetCppType(), *ParamName);
 				}
 			}
 		}
@@ -5025,16 +4888,15 @@ void FNativeClassHeaderGenerator::ExportFunctionThunk(FUHTStringBuilder& RPCWrap
 	RPCWrappers.Log(TEXT("\t\t"));
 	if (ReturnDef)
 	{
-		FProperty* Return = ReturnDef->GetProperty();
-		OutReferenceGatherers.ForwardDeclarations.Add(Return->GetCPPTypeForwardDeclaration());
+		OutReferenceGatherers.ForwardDeclarations.Add(ReturnDef->GetCPPTypeForwardDeclaration());
 
 		FUHTStringBuilder ReplacementText;
 		FString ReturnExtendedType;
-		ReplacementText += Return->GetCPPType(&ReturnExtendedType);
+		ReplacementText += ReturnDef->GetCPPType(&ReturnExtendedType);
 		ApplyAlternatePropertyExportText(*ReturnDef, ReplacementText, EExportingState::Normal);
 
 		FString ReturnType = ReplacementText;
-		if (Return->HasAnyPropertyFlags(CPF_ConstParm) && CastField<FObjectProperty>(Return))
+		if (ReturnDef->HasAnyPropertyFlags(CPF_ConstParm) && ReturnDef->IsObjectRefOrObjectRefStaticArray())
 		{
 			ReturnType = TEXT("const ") + ReturnType;
 		}
@@ -5058,12 +4920,16 @@ FString FNativeClassHeaderGenerator::GetFunctionParameterString(FUnrealFunctionD
 	FString ParameterList;
 	FUHTStringBuilder PropertyText;
 
+	if (FunctionDef.GetReturn() != nullptr)
+	{
+		ParameterList = ParameterList;
+	}
+
 	for (FUnrealPropertyDefinitionInfo* PropertyDef : FunctionDef.GetProperties())
 	{
-		FProperty* Property = PropertyDef->GetProperty();
-		OutReferenceGatherers.ForwardDeclarations.Add(Property->GetCPPTypeForwardDeclaration());
+		OutReferenceGatherers.ForwardDeclarations.Add(PropertyDef->GetCPPTypeForwardDeclaration());
 
-		if ((Property->PropertyFlags & (CPF_Parm | CPF_ReturnParm)) != CPF_Parm)
+		if (!PropertyDef->HasSpecificPropertyFlags(CPF_Parm | CPF_ReturnParm, CPF_Parm))
 		{
 			break;
 		}
@@ -5073,7 +4939,7 @@ FString FNativeClassHeaderGenerator::GetFunctionParameterString(FUnrealFunctionD
 			ParameterList += TEXT(", ");
 		}
 
-		Property->ExportCppDeclaration(PropertyText, EExportedDeclaration::Parameter, PropertyDef->GetArrayDimensions(), 0, true);
+		PropertyDef->ExportCppDeclaration(PropertyText, EExportedDeclaration::Parameter, PropertyDef->GetArrayDimensions(), 0, true);
 		ApplyAlternatePropertyExportText(*PropertyDef, PropertyText, EExportingState::Normal);
 
 		ParameterList += PropertyText;
@@ -5126,23 +4992,22 @@ void FNativeClassHeaderGenerator::ExportNativeFunctions(FOutputDevice& OutGenera
 		{
 			for (FUnrealPropertyDefinitionInfo* PropertyDef : SparseClassDataStructDef->GetProperties())
 			{
-				const FProperty* Child = PropertyDef->GetProperty();
 				FString ReturnExtendedType;
-				FString VarType = Child->GetCPPType(&ReturnExtendedType, EPropertyExportCPPFlags::CPPF_ArgumentOrReturnValue | EPropertyExportCPPFlags::CPPF_Implementation);
+				FString VarType = PropertyDef->GetCPPType(&ReturnExtendedType, EPropertyExportCPPFlags::CPPF_ArgumentOrReturnValue | EPropertyExportCPPFlags::CPPF_Implementation);
 				if (!ReturnExtendedType.IsEmpty())
 				{
 					VarType.Append(ReturnExtendedType);
 				}
-				FString VarName = Child->GetName();
+				FString VarName = PropertyDef->GetName();
 				FString CleanVarName = VarName;
-				if (CastField<FBoolProperty>(Child) && VarName.StartsWith(TEXT("b"), ESearchCase::CaseSensitive))
+				if (PropertyDef->IsBooleanOrBooleanStaticArray() && VarName.StartsWith(TEXT("b"), ESearchCase::CaseSensitive))
 				{
 					CleanVarName = VarName.RightChop(1);
 				}
 
-				if (!Child->HasMetaData(NAME_NoGetter))
+				if (!PropertyDef->HasMetaData(NAME_NoGetter))
 				{
-					if (Child->HasMetaData(NAME_GetByRef))
+					if (PropertyDef->HasMetaData(NAME_GetByRef))
 					{
 						RuntimeStringBuilders.AutogeneratedStaticDataFuncs.Logf(TEXT("const %s& Get%s()\r\n"), *VarType, *CleanVarName);
 					}
@@ -5154,7 +5019,7 @@ void FNativeClassHeaderGenerator::ExportNativeFunctions(FOutputDevice& OutGenera
 					RuntimeStringBuilders.AutogeneratedStaticDataFuncs.Logf(TEXT("\treturn Get%s()->%s;\r\n"), *SparseClassDataString, *VarName);
 					RuntimeStringBuilders.AutogeneratedStaticDataFuncs.Logf(TEXT("}\r\n"));
 
-					if (Child->HasMetaData(NAME_GetByRef))
+					if (PropertyDef->HasMetaData(NAME_GetByRef))
 					{
 						RuntimeStringBuilders.AutogeneratedStaticDataFuncs.Logf(TEXT("const %s& Get%s() const\r\n"), *VarType, *CleanVarName);
 					}
@@ -5465,33 +5330,29 @@ void FNativeClassHeaderGenerator::ExportCallbackFunctions(
  */
 void FNativeClassHeaderGenerator::ApplyAlternatePropertyExportText(FUnrealPropertyDefinitionInfo& PropertyDef, FUHTStringBuilder& PropertyText, EExportingState ExportingState)
 {
-	FProperty* Prop = PropertyDef.GetProperty();
-	FArrayProperty* ArrayProperty = CastField<FArrayProperty>(Prop);
-	FProperty* InnerProperty = ArrayProperty ? ArrayProperty->Inner : nullptr;
-	if (InnerProperty && (
-			(InnerProperty->IsA<FByteProperty>() && ((FByteProperty*)InnerProperty)->Enum && FUnrealTypeDefinitionInfo::IsDynamic(static_cast<UField*>(((FByteProperty*)InnerProperty)->Enum))) ||
-			(InnerProperty->IsA<FEnumProperty>()                                          && FUnrealTypeDefinitionInfo::IsDynamic(static_cast<UField*>(((FEnumProperty*)InnerProperty)->Enum)))
-		)
-	)
+	const FPropertyBase& PropertyBase = PropertyDef.GetPropertyBase();
+	if (FUnrealEnumDefinitionInfo* EnumDef = PropertyBase.AsEnum())
 	{
-		const FString Original = InnerProperty->GetCPPType();
-		const FString RawByte = InnerProperty->GetCPPType(nullptr, EPropertyExportCPPFlags::CPPF_BlueprintCppBackend);
-		if (Original != RawByte)
+		if (PropertyBase.ArrayType == EArrayType::Static && EnumDef->IsDynamic())
 		{
-			PropertyText.ReplaceInline(*Original, *RawByte, ESearchCase::CaseSensitive);
+			FUnrealPropertyDefinitionInfo& InnerPropertyDef = PropertyDef.GetValuePropDef();
+			const FString Original = InnerPropertyDef.GetCPPType();
+			const FString RawByte = InnerPropertyDef.GetCPPType(nullptr, EPropertyExportCPPFlags::CPPF_BlueprintCppBackend);
+			if (Original != RawByte)
+			{
+				PropertyText.ReplaceInline(*Original, *RawByte, ESearchCase::CaseSensitive);
+			}
 		}
 		return;
 	}
 
 	if (ExportingState == EExportingState::TypeEraseDelegates)
 	{
-		FDelegateProperty* DelegateProperty = CastField<FDelegateProperty>(Prop);
-		FMulticastDelegateProperty* MulticastDelegateProperty = CastField<FMulticastDelegateProperty>(Prop);
-		if (DelegateProperty || MulticastDelegateProperty)
+		if (PropertyBase.Type == CPT_Delegate || PropertyBase.Type == CPT_MulticastDelegate)
 		{
-			FString Original = Prop->GetCPPType();
+			FString Original = PropertyDef.GetCPPType();
 			const TCHAR* PlaceholderOfSameSizeAndAlignemnt;
-			if (DelegateProperty)
+			if (PropertyBase.Type == CPT_Delegate)
 			{
 				PlaceholderOfSameSizeAndAlignemnt = TEXT("FScriptDelegate");
 			}
diff --git a/Engine/Source/Programs/UnrealHeaderTool/Private/HeaderParser.cpp b/Engine/Source/Programs/UnrealHeaderTool/Private/HeaderParser.cpp
--- a/Engine/Source/Programs/UnrealHeaderTool/Private/HeaderParser.cpp
+++ b/Engine/Source/Programs/UnrealHeaderTool/Private/HeaderParser.cpp
@@ -674,7 +674,7 @@ namespace
 					if (FuncDef->HasAnyFunctionFlags(FUNC_Static))
 					{
 						// Determine if it's a function library
-						FUnrealClassDefinitionInfo* ClassDef = FuncDef->GetOuterClass();
+						FUnrealClassDefinitionInfo* ClassDef = FuncDef->GetOwnerClass();
 						for (; ClassDef != nullptr && ClassDef->GetSuperClass() != GUObjectDef; ClassDef = ClassDef->GetSuperClass())
 						{
 						}
@@ -695,30 +695,29 @@ namespace
 					bool bGoodParams = (FuncDef->GetProperties().Num() == 3);
 					if (bGoodParams)
 					{
-						FProperty* FirstParam = nullptr;
-						FProperty* SecondParam = nullptr;
-						FProperty* ReturnValue = nullptr;
+						FUnrealPropertyDefinitionInfo* FirstParam = nullptr;
+						FUnrealPropertyDefinitionInfo* SecondParam = nullptr;
+						FUnrealPropertyDefinitionInfo* ReturnValue = nullptr;
 						for (FUnrealPropertyDefinitionInfo* PropertyDef : FuncDef->GetProperties())
 						{
-							FProperty* Property = PropertyDef->GetProperty();
-							if (Property->HasAnyPropertyFlags(CPF_ReturnParm))
+							if (PropertyDef->HasAnyPropertyFlags(CPF_ReturnParm))
 							{
-								ReturnValue = Property;
+								ReturnValue = PropertyDef;
 							}
 							else
 							{
 								if (FirstParam == nullptr)
 								{
-									FirstParam = Property;
+									FirstParam = PropertyDef;
 								}
 								else if (SecondParam == nullptr)
 								{
-									SecondParam = Property;
+									SecondParam = PropertyDef;
 								}
 							}
 						}
 
-						if (ReturnValue == nullptr || SecondParam == nullptr || !SecondParam->SameType(FirstParam))
+						if (ReturnValue == nullptr || SecondParam == nullptr || !SecondParam->SameType(*FirstParam))
 						{
 							bGoodParams = false;
 						}
@@ -741,7 +740,7 @@ namespace
 					TArray<FString> RawGroupings;
 					InValue.ParseIntoArray(RawGroupings, TEXT(","), false);
 
-					FProperty* FirstInput = nullptr;
+					FUnrealPropertyDefinitionInfo* FirstInputDef = nullptr;
 					for (const FString& RawGroup : RawGroupings)
 					{
 						TArray<FString> IndividualEntries;
@@ -754,25 +753,25 @@ namespace
 								continue;
 							}
 							
-							FField* FoundField = FHeaderParser::FindProperty(*FuncDef, *Entry, false);
-							if (!FoundField)
+							FUnrealPropertyDefinitionInfo* FoundFieldDef = FHeaderParser::FindProperty(*FuncDef, *Entry, false);
+							if (!FoundFieldDef)
 							{
 								UE_LOG_ERROR_UHT(TEXT("Function does not have a parameter named '%s'"), *Entry);
 							}
-							else if (FProperty* Prop = CastField<FProperty>(FoundField))
+							else
 							{
-								if (!Prop->HasAnyPropertyFlags(CPF_ReturnParm) &&
+								if (!FoundFieldDef->HasAnyPropertyFlags(CPF_ReturnParm) &&
 
-								    (!Prop->HasAnyPropertyFlags(CPF_OutParm) ||
-									Prop->HasAnyPropertyFlags(CPF_ReferenceParm)))
+								    (!FoundFieldDef->HasAnyPropertyFlags(CPF_OutParm) ||
+										FoundFieldDef->HasAnyPropertyFlags(CPF_ReferenceParm)))
 								{
-									if (!FirstInput)
+									if (!FirstInputDef)
 									{
-										FirstInput = Prop;
+										FirstInputDef = FoundFieldDef;
 									}
 									else
 									{
-										UE_LOG_ERROR_UHT(TEXT("Function already specified an ExpandEnumAsExec input (%s), but '%s' is also an input parameter. Only one is permitted."), *FirstInput->GetName(), *Entry);
+										UE_LOG_ERROR_UHT(TEXT("Function already specified an ExpandEnumAsExec input (%s), but '%s' is also an input parameter. Only one is permitted."), *FirstInputDef->GetName(), *Entry);
 									}
 								}
 							}
@@ -798,8 +797,7 @@ namespace
 				// Check for numeric property
 				if (FUnrealPropertyDefinitionInfo* PropDef = TypeDef.AsProperty())
 				{
-					FProperty* Property = PropDef->GetProperty();
-					if (!Property->IsA<FNumericProperty>() && !Property->IsA<FStructProperty>())
+					if (!PropDef->IsNumericOrNumericStaticArray() && !PropDef->IsStructOrStructStaticArray())
 					{
 						FError::Throwf(TEXT("'Units' meta data can only be applied to numeric and struct properties"));
 					}
@@ -807,14 +805,7 @@ namespace
 
 				if (!FUnitConversion::UnitFromString(*InValue))
 				{
-					if (TypeDef.AsField() != nullptr)
-					{
-						FError::Throwf(TEXT("Unrecognized units (%s) specified for property '%s'"), *InValue, *TypeDef.AsFieldChecked().GetField()->GetFullName());
-					}
-					else
-					{
-						FError::Throwf(TEXT("Unrecognized units (%s) specified for property '%s'"), *InValue, *TypeDef.AsPropertyChecked().GetProperty()->GetFullName());
-					}
+					FError::Throwf(TEXT("Unrecognized units (%s) specified for property '%s'"), *InValue, *TypeDef.GetFullName());
 				}
 			}
 			break;
@@ -1059,23 +1050,21 @@ FUnrealClassDefinitionInfo* FHeaderParser::GetQualifiedClass(const TCHAR* Thing)
 -----------------------------------------------------------------------------*/
 
 /**
- * Find a field in the specified context.  Starts with the specified scope, then iterates
+ * Find a function in the specified context.  Starts with the specified scope, then iterates
  * through the Outer chain until the field is found.
  * 
  * @param	InScope				scope to start searching for the field in 
  * @param	InIdentifier		name of the field we're searching for
  * @param	bIncludeParents		whether to allow searching in the scope of a parent struct
- * @param	FieldClass			class of the field to search for.  used to e.g. search for functions only
  * @param	Thing				hint text that will be used in the error message if an error is encountered
  *
  * @return	a pointer to a UField with a name matching InIdentifier, or NULL if it wasn't found
  */
-UField* FHeaderParser::FindField
+FUnrealFunctionDefinitionInfo* FHeaderParser::FindFunction
 (
 	FUnrealStructDefinitionInfo&  InScope,
 	const TCHAR*	InIdentifier,
 	bool			bIncludeParents,
-	UClass*			FieldClass,
 	const TCHAR*	Thing
 )
 {
@@ -1087,31 +1076,20 @@ UField* FHeaderParser::FindField
 		{
 			for (FUnrealPropertyDefinitionInfo* PropertyDef : TUHTFieldRange<FUnrealPropertyDefinitionInfo>(*Scope))
 			{
-				FProperty* Property = PropertyDef->GetProperty();
-				if (Property->GetFName() == InName)
+				if (PropertyDef->GetFName() == InName)
 				{
 					if (Thing)
 					{
-						FError::Throwf(TEXT("%s: expecting @todo: FProp, got %s"), Thing, /**FieldClass->GetName(),*/ *Property->GetClass()->GetName());
+						FError::Throwf(TEXT("%s: expecting @todo: FProp, got %s"), Thing, /**FieldClass->GetName(),*/ *PropertyDef->GetEngineClassName());
 					}
 					return nullptr;
 				}
-
 			}
 			for (FUnrealFunctionDefinitionInfo* FunctionDef : TUHTFieldRange<FUnrealFunctionDefinitionInfo>(*Scope))
 			{
-				UFunction* Function = FunctionDef->GetFunction();
-				if (Function->GetFName() == InName)
+				if (FunctionDef->GetFName() == InName)
 				{
-					if (!Function->IsA(FieldClass))
-					{
-						if (Thing)
-						{
-							FError::Throwf(TEXT("%s: expecting %s, got %s"), Thing, *FieldClass->GetName(), *Function->GetClass()->GetName());
-						}
-						return nullptr;
-					}
-					return Function;
+					return FunctionDef;
 				}
 			}
 
@@ -1125,7 +1103,7 @@ UField* FHeaderParser::FindField
 	return nullptr;
 }
 
-FField* FHeaderParser::FindProperty(FUnrealStructDefinitionInfo& InScope, const TCHAR* InIdentifier, bool bIncludeParents, FFieldClass* FieldClass, const TCHAR* Thing)
+FUnrealPropertyDefinitionInfo* FHeaderParser::FindProperty(FUnrealStructDefinitionInfo& InScope, const TCHAR* InIdentifier, bool bIncludeParents, const TCHAR* Thing)
 {
 	check(InIdentifier);
 	FName InName(InIdentifier, FNAME_Find);
@@ -1135,30 +1113,20 @@ FField* FHeaderParser::FindProperty(FUnrealStructDefinitionInfo& InScope, const
 		{
 			for (FUnrealFunctionDefinitionInfo* FunctionDef : TUHTFieldRange<FUnrealFunctionDefinitionInfo>(*Scope))
 			{
-				UFunction* Function = FunctionDef->GetFunction();
-				if (Function->GetFName() == InName)
+				if (FunctionDef->GetFName() == InName)
 				{
 					if (Thing)
 					{
-						FError::Throwf(TEXT("%s: expecting a property, got %s"), Thing, *Function->GetClass()->GetName());
+						FError::Throwf(TEXT("%s: expecting a property, got %s"), Thing, *FunctionDef->GetFunction()->GetClass()->GetName());
 					}
 					return nullptr;
 				}
 			}
 			for (FUnrealPropertyDefinitionInfo* PropertyDef : TUHTFieldRange<FUnrealPropertyDefinitionInfo>(*Scope))
 			{
-				FProperty* Property = PropertyDef->GetProperty();
-				if (Property->GetFName() == InName)
+				if (PropertyDef->GetFName() == InName)
 				{
-					if (!Property->IsA(FieldClass))
-					{
-						if (Thing)
-						{
-							FError::Throwf(TEXT("%s: expecting %s: FProp, got %s"), Thing, *FieldClass->GetName(), *Property->GetClass()->GetName());
-						}
-						return nullptr;
-					}
-					return Property;
+					return PropertyDef;
 				}
 			}
 
@@ -2560,155 +2528,147 @@ void FHeaderParser::FixupDelegateProperties(FUnrealStructDefinitionInfo& StructD
 {
 	for (FUnrealPropertyDefinitionInfo* PropertyDef : StructDef.GetProperties())
 	{
-		FProperty* Property = PropertyDef->GetProperty();
-		FDelegateProperty* DelegateProperty = CastField<FDelegateProperty>(Property);
-		FMulticastDelegateProperty* MulticastDelegateProperty = CastField<FMulticastDelegateProperty>(Property);
-		if ( DelegateProperty == NULL && MulticastDelegateProperty == NULL )
+		FPropertyBase& PropertyBase = PropertyDef->GetPropertyBase();
+
+		// Ignore any containers types but static arrays
+		if (PropertyBase.ArrayType != EArrayType::None && PropertyBase.ArrayType != EArrayType::Static)
 		{
-			// if this is an array property, see if the array's type is a delegate
-			FArrayProperty* ArrayProp = CastField<FArrayProperty>(Property);
-			if ( ArrayProp != NULL )
-			{
-				DelegateProperty = CastField<FDelegateProperty>(ArrayProp->Inner);
-				MulticastDelegateProperty = CastField<FMulticastDelegateProperty>(ArrayProp->Inner);
-			}
+			continue;
 		}
-		if (DelegateProperty != nullptr || MulticastDelegateProperty != nullptr)
+
+		// At this point, we can check to see if we are a delegate or multicast delegate directly off the property definition
+		bool bIsDelegate = PropertyBase.Type == CPT_Delegate;
+		bool bIsMulticastDelegate = PropertyBase.Type == CPT_MulticastDelegate;
+		if (!bIsDelegate && !bIsMulticastDelegate)
 		{
-			// this FDelegateProperty corresponds to an actual delegate variable (i.e. delegate<SomeDelegate> Foo); we need to lookup the token data for
-			// this property and verify that the delegate property's "type" is an actual delegate function
-			FPropertyBase& DelegatePropertyToken = PropertyDef->GetPropertyBase();
+			continue;
+		}
+
+		// this FDelegateProperty corresponds to an actual delegate variable (i.e. delegate<SomeDelegate> Foo); we need to lookup the token data for
+		// this property and verify that the delegate property's "type" is an actual delegate function
+		FPropertyBase& DelegatePropertyBase = PropertyDef->GetPropertyBase();
 
-			// attempt to find the delegate function in the map of functions we've already found
-			FUnrealFunctionDefinitionInfo* SourceDelegateFunctionDef = DelegateCache.FindRef(DelegatePropertyToken.DelegateName);
-			if (SourceDelegateFunctionDef == nullptr)
+		// attempt to find the delegate function in the map of functions we've already found
+		FUnrealFunctionDefinitionInfo* SourceDelegateFunctionDef = DelegateCache.FindRef(DelegatePropertyBase.DelegateName);
+		if (SourceDelegateFunctionDef == nullptr)
+		{
+			FString NameOfDelegateFunction = DelegatePropertyBase.DelegateName.ToString() + FString( HEADER_GENERATED_DELEGATE_SIGNATURE_SUFFIX );
+			if ( !NameOfDelegateFunction.Contains(TEXT(".")) )
 			{
-				FString NameOfDelegateFunction = DelegatePropertyToken.DelegateName.ToString() + FString( HEADER_GENERATED_DELEGATE_SIGNATURE_SUFFIX );
-				if ( !NameOfDelegateFunction.Contains(TEXT(".")) )
+				// an unqualified delegate function name - search for a delegate function by this name within the current scope
+				if (FUnrealFieldDefinitionInfo* FoundDef = Scope.FindTypeByName(*NameOfDelegateFunction))
 				{
-					// an unqualified delegate function name - search for a delegate function by this name within the current scope
-					if (FUnrealFieldDefinitionInfo* FoundDef = Scope.FindTypeByName(*NameOfDelegateFunction))
-					{
-						SourceDelegateFunctionDef = FoundDef->AsFunction();
-					}
-					if (SourceDelegateFunctionDef == nullptr)
-					{
-						// Try to find in other packages.
-						UObject* DelegateSignatureOuter = DelegatePropertyToken.DelegateSignatureOwnerClass 
-							? ((UObject*)DelegatePropertyToken.DelegateSignatureOwnerClass) 
-							: ((UObject*)ANY_PACKAGE);
-						UFunction* SourceDelegateFunction = FEngineAPI::FindObject<UFunction>(DelegateSignatureOuter, *NameOfDelegateFunction);
-
-						if (SourceDelegateFunction == nullptr)
-						{
-							// convert this into a fully qualified path name for the error message.
-							NameOfDelegateFunction = Scope.GetName().ToString() + TEXT(".") + NameOfDelegateFunction;
-						}
-						else
-						{
-							SourceDelegateFunctionDef = &GTypeDefinitionInfoMap.FindChecked<FUnrealFunctionDefinitionInfo>(SourceDelegateFunction);
-						}
-					}
+					SourceDelegateFunctionDef = FoundDef->AsFunction();
 				}
-				else
+				if (SourceDelegateFunctionDef == nullptr)
 				{
-					FString DelegateClassName, DelegateName;
-					NameOfDelegateFunction.Split(TEXT("."), &DelegateClassName, &DelegateName);
-
-					// verify that we got a valid string for the class name
-					if ( DelegateClassName.Len() == 0 )
-					{
-						UngetToken(PropertyDef->GetLineNumber(), PropertyDef->GetParsePosition());
-						FError::Throwf(TEXT("Invalid scope specified in delegate property function reference: '%s'"), *NameOfDelegateFunction);
-					}
-
-					// verify that we got a valid string for the name of the function
-					if ( DelegateName.Len() == 0 )
-					{
-						UngetToken(PropertyDef->GetLineNumber(), PropertyDef->GetParsePosition());
-						FError::Throwf(TEXT("Invalid delegate name specified in delegate property function reference '%s'"), *NameOfDelegateFunction);
-					}
+					// Try to find in other packages.
+					UObject* DelegateSignatureOuter = DelegatePropertyBase.DelegateSignatureOwnerClass
+						? ((UObject*)DelegatePropertyBase.DelegateSignatureOwnerClass)
+						: ((UObject*)ANY_PACKAGE);
+					UFunction* SourceDelegateFunction = FEngineAPI::FindObject<UFunction>(DelegateSignatureOuter, *NameOfDelegateFunction);
 
-					// make sure that the class that contains the delegate can be referenced here
-					FUnrealClassDefinitionInfo* DelegateOwnerClassDef = FUnrealClassDefinitionInfo::FindScriptClassOrThrow(DelegateClassName);
-					if (DelegateOwnerClassDef->GetScope()->FindTypeByName(*DelegateName) != nullptr)
+					if (SourceDelegateFunction == nullptr)
 					{
-						FError::Throwf(TEXT("Inaccessible type: '%s'"), *DelegateOwnerClassDef->GetPathName());
+						// convert this into a fully qualified path name for the error message.
+						NameOfDelegateFunction = Scope.GetName().ToString() + TEXT(".") + NameOfDelegateFunction;
 					}
-					UFunction* SourceDelegateFunction = Cast<UFunction>(FindField(*DelegateOwnerClassDef, *DelegateName, false, UFunction::StaticClass(), NULL));
-					if (SourceDelegateFunction != nullptr)
+					else
 					{
 						SourceDelegateFunctionDef = &GTypeDefinitionInfoMap.FindChecked<FUnrealFunctionDefinitionInfo>(SourceDelegateFunction);
 					}
 				}
+			}
+			else
+			{
+				FString DelegateClassName, DelegateName;
+				NameOfDelegateFunction.Split(TEXT("."), &DelegateClassName, &DelegateName);
 
-				if (SourceDelegateFunctionDef == NULL )
+				// verify that we got a valid string for the class name
+				if ( DelegateClassName.Len() == 0 )
 				{
 					UngetToken(PropertyDef->GetLineNumber(), PropertyDef->GetParsePosition());
-					FError::Throwf(TEXT("Failed to find delegate function '%s'"), *NameOfDelegateFunction);
+					FError::Throwf(TEXT("Invalid scope specified in delegate property function reference: '%s'"), *NameOfDelegateFunction);
 				}
-				else if (!SourceDelegateFunctionDef->HasAnyFunctionFlags(FUNC_Delegate))
+
+				// verify that we got a valid string for the name of the function
+				if ( DelegateName.Len() == 0 )
 				{
 					UngetToken(PropertyDef->GetLineNumber(), PropertyDef->GetParsePosition());
-					FError::Throwf(TEXT("Only delegate functions can be used as the type for a delegate property; '%s' is not a delegate."), *NameOfDelegateFunction);
+					FError::Throwf(TEXT("Invalid delegate name specified in delegate property function reference '%s'"), *NameOfDelegateFunction);
 				}
+
+				// make sure that the class that contains the delegate can be referenced here
+				FUnrealClassDefinitionInfo* DelegateOwnerClassDef = FUnrealClassDefinitionInfo::FindScriptClassOrThrow(DelegateClassName);
+				if (DelegateOwnerClassDef->GetScope()->FindTypeByName(*DelegateName) != nullptr)
+				{
+					FError::Throwf(TEXT("Inaccessible type: '%s'"), *DelegateOwnerClassDef->GetPathName());
+				}
+				SourceDelegateFunctionDef = FindFunction(*DelegateOwnerClassDef, *DelegateName, false, nullptr);
+			}
+
+			if (SourceDelegateFunctionDef == NULL )
+			{
+				UngetToken(PropertyDef->GetLineNumber(), PropertyDef->GetParsePosition());
+				FError::Throwf(TEXT("Failed to find delegate function '%s'"), *NameOfDelegateFunction);
+			}
+			else if (!SourceDelegateFunctionDef->HasAnyFunctionFlags(FUNC_Delegate))
+			{
+				UngetToken(PropertyDef->GetLineNumber(), PropertyDef->GetParsePosition());
+				FError::Throwf(TEXT("Only delegate functions can be used as the type for a delegate property; '%s' is not a delegate."), *NameOfDelegateFunction);
 			}
+		}
 
-			// successfully found the delegate function that this delegate property corresponds to
+		// successfully found the delegate function that this delegate property corresponds to
 
-				// save this into the delegate cache for faster lookup later
-			DelegateCache.Add(DelegatePropertyToken.DelegateName, SourceDelegateFunctionDef);
+			// save this into the delegate cache for faster lookup later
+		DelegateCache.Add(DelegatePropertyBase.DelegateName, SourceDelegateFunctionDef);
 
-			// bind it to the delegate property
-			if( DelegateProperty != NULL )
+		// bind it to the delegate property
+		if (PropertyBase.Type == CPT_Delegate)
+		{
+			if( !SourceDelegateFunctionDef->HasAnyFunctionFlags( FUNC_MulticastDelegate ) )
 			{
-				if( !SourceDelegateFunctionDef->HasAnyFunctionFlags( FUNC_MulticastDelegate ) )
-				{
-					DelegateProperty->SignatureFunction = SourceDelegateFunctionDef->GetFunction();
-					DelegatePropertyToken.FunctionDef = SourceDelegateFunctionDef;
-				}
-				else
-				{
-					FError::Throwf(TEXT("Unable to declare a single-cast delegate property for a multi-cast delegate type '%s'.  Either add a 'multicast' qualifier to the property or change the delegate type to be single-cast as well."), *SourceDelegateFunctionDef->GetName());
-				}
+				PropertyDef->SetDelegateFunctionSignature(*SourceDelegateFunctionDef);
 			}
-			else if( MulticastDelegateProperty != NULL )
+			else
 			{
-				if(SourceDelegateFunctionDef->HasAnyFunctionFlags( FUNC_MulticastDelegate ) )
-				{
-					MulticastDelegateProperty->SignatureFunction = SourceDelegateFunctionDef->GetFunction();
-					DelegatePropertyToken.FunctionDef = SourceDelegateFunctionDef;
+				FError::Throwf(TEXT("Unable to declare a single-cast delegate property for a multi-cast delegate type '%s'.  Either add a 'multicast' qualifier to the property or change the delegate type to be single-cast as well."), *SourceDelegateFunctionDef->GetName());
+			}
+		}
+		else if (PropertyBase.Type == CPT_MulticastDelegate)
+		{
+			if (SourceDelegateFunctionDef->HasAnyFunctionFlags( FUNC_MulticastDelegate ))
+			{
+				DelegatePropertyBase.FunctionDef = SourceDelegateFunctionDef;
 
-					if(MulticastDelegateProperty->HasAnyPropertyFlags(CPF_BlueprintAssignable | CPF_BlueprintCallable))
+				if (PropertyDef->HasAnyPropertyFlags(CPF_BlueprintAssignable | CPF_BlueprintCallable))
+				{
+					for (FUnrealPropertyDefinitionInfo* FuncParamDef : SourceDelegateFunctionDef->GetProperties())
 					{
-						for (FUnrealPropertyDefinitionInfo* FuncParamDef : SourceDelegateFunctionDef->GetProperties())
+						if (!FPropertyTraits::IsSupportedByBlueprint(*FuncParamDef, false))
 						{
-							FProperty* FuncParam = FuncParamDef->GetProperty();
-
-							if (!FPropertyTraits::IsSupportedByBlueprint(*FuncParamDef, false))
-							{
-								FString ExtendedCPPType;
-								FString CPPType = FuncParam->GetCPPType(&ExtendedCPPType);
-								UE_LOG_ERROR_UHT(TEXT("Type '%s%s' is not supported by blueprint. %s.%s"), *CPPType, *ExtendedCPPType, *SourceDelegateFunctionDef->GetName(), *FuncParam->GetName());
-							}
+							FString ExtendedCPPType;
+							FString CPPType = FuncParamDef->GetCPPType(&ExtendedCPPType);
+							UE_LOG_ERROR_UHT(TEXT("Type '%s%s' is not supported by blueprint. %s.%s"), *CPPType, *ExtendedCPPType, *SourceDelegateFunctionDef->GetName(), *FuncParamDef->GetName());
+						}
 
-							if(FuncParam->HasAllPropertyFlags(CPF_OutParm) && !FuncParam->HasAllPropertyFlags(CPF_ConstParm)  )
+						if(FuncParamDef->HasAllPropertyFlags(CPF_OutParm) && !FuncParamDef->HasAllPropertyFlags(CPF_ConstParm)  )
+						{
+							const bool bClassGeneratedFromBP = StructDef.IsDynamic();
+							const bool bAllowedArrayRefFromBP = bClassGeneratedFromBP && FuncParamDef->IsDynamicArray();
+							if (!bAllowedArrayRefFromBP)
 							{
-								const bool bClassGeneratedFromBP = StructDef.IsDynamic();
-								const bool bAllowedArrayRefFromBP = bClassGeneratedFromBP && FuncParam->IsA<FArrayProperty>();
-								if (!bAllowedArrayRefFromBP)
-								{
-									UE_LOG_ERROR_UHT(TEXT("BlueprintAssignable delegates do not support non-const references at the moment. Function: %s Parameter: '%s'"), *SourceDelegateFunctionDef->GetName(), *FuncParam->GetName());
-								}
+								UE_LOG_ERROR_UHT(TEXT("BlueprintAssignable delegates do not support non-const references at the moment. Function: %s Parameter: '%s'"), *SourceDelegateFunctionDef->GetName(), *FuncParamDef->GetName());
 							}
 						}
 					}
-
-				}
-				else
-				{
-					FError::Throwf(TEXT("Unable to declare a multi-cast delegate property for a single-cast delegate type '%s'.  Either remove the 'multicast' qualifier from the property or change the delegate type to be 'multicast' as well."), *SourceDelegateFunctionDef->GetName());
 				}
+
+			}
+			else
+			{
+				FError::Throwf(TEXT("Unable to declare a multi-cast delegate property for a single-cast delegate type '%s'.  Either remove the 'multicast' qualifier from the property or change the delegate type to be 'multicast' as well."), *SourceDelegateFunctionDef->GetName());
 			}
 		}
 	}
@@ -2771,22 +2731,21 @@ void FHeaderParser::CheckSparseClassData(const FUnrealStructDefinitionInfo& Stru
 		// check the data struct for invalid properties
 		for (FUnrealPropertyDefinitionInfo* PropertyDef : TUHTFieldRange<FUnrealPropertyDefinitionInfo>(*SparseClassDataStructDef))
 		{
-			FProperty* Property = PropertyDef->GetProperty();
-			if (Property->HasAnyPropertyFlags(CPF_BlueprintAssignable))
+			if (PropertyDef->HasAnyPropertyFlags(CPF_BlueprintAssignable))
 			{
-				FError::Throwf(TEXT("Sparse class data types can not contain blueprint assignable delegates. Type '%s' Delegate '%s'"), *SparseClassDataStructDef->GetName(), *Property->GetName());
+				FError::Throwf(TEXT("Sparse class data types can not contain blueprint assignable delegates. Type '%s' Delegate '%s'"), *SparseClassDataStructDef->GetName(), *PropertyDef->GetName());
 			}
 
 			// all sparse properties should have EditDefaultsOnly
-			if (!Property->HasAllPropertyFlags(CPF_Edit | CPF_DisableEditOnInstance))
+			if (!PropertyDef->HasAllPropertyFlags(CPF_Edit | CPF_DisableEditOnInstance))
 			{
-				FError::Throwf(TEXT("Sparse class data types must be EditDefaultsOnly. Type '%s' Property '%s'"), *SparseClassDataStructDef->GetName(), *Property->GetName());
+				FError::Throwf(TEXT("Sparse class data types must be EditDefaultsOnly. Type '%s' Property '%s'"), *SparseClassDataStructDef->GetName(), *PropertyDef->GetName());
 			}
 
 			// no sparse properties should have BlueprintReadWrite
-			if (Property->HasAllPropertyFlags(CPF_BlueprintVisible) && !Property->HasAllPropertyFlags(CPF_BlueprintReadOnly))
+			if (PropertyDef->HasAllPropertyFlags(CPF_BlueprintVisible) && !PropertyDef->HasAllPropertyFlags(CPF_BlueprintReadOnly))
 			{
-				FError::Throwf(TEXT("Sparse class data types must not be BlueprintReadWrite. Type '%s' Property '%s'"), *SparseClassDataStructDef->GetName(), *Property->GetName());
+				FError::Throwf(TEXT("Sparse class data types must not be BlueprintReadWrite. Type '%s' Property '%s'"), *SparseClassDataStructDef->GetName(), *PropertyDef->GetName());
 			}
 		}
 
@@ -2821,7 +2780,6 @@ void FHeaderParser::ValidateClassFlags(const FUnrealClassDefinitionInfo& ToValid
 void FHeaderParser::VerifyBlueprintPropertyGetter(FUnrealPropertyDefinitionInfo& PropertyDef, FUnrealFunctionDefinitionInfo* TargetFuncDef)
 {
 	check(TargetFuncDef);
-	FProperty* Prop = PropertyDef.GetProperty();
 
 	const TArray<FUnrealPropertyDefinitionInfo*>& Properties = TargetFuncDef->GetProperties();
 	FUnrealPropertyDefinitionInfo* ReturnPropDef = TargetFuncDef->GetReturn();
@@ -2830,10 +2788,10 @@ void FHeaderParser::VerifyBlueprintPropertyGetter(FUnrealPropertyDefinitionInfo&
 		UE_LOG_ERROR_UHT(TEXT("Blueprint Property getter function %s must not have parameters."), *TargetFuncDef->GetName());
 	}
 
-	if (ReturnPropDef == nullptr || !Prop->SameType(ReturnPropDef->GetProperty()))
+	if (ReturnPropDef == nullptr || !PropertyDef.SameType(*ReturnPropDef))
 	{
 		FString ExtendedCPPType;
-		FString CPPType = Prop->GetCPPType(&ExtendedCPPType);
+		FString CPPType = PropertyDef.GetCPPType(&ExtendedCPPType);
 		UE_LOG_ERROR_UHT(TEXT("Blueprint Property getter function %s must have return value of type %s%s."), *TargetFuncDef->GetName(), *CPPType, *ExtendedCPPType);
 	}
 
@@ -2850,7 +2808,6 @@ void FHeaderParser::VerifyBlueprintPropertyGetter(FUnrealPropertyDefinitionInfo&
 void FHeaderParser::VerifyBlueprintPropertySetter(FUnrealPropertyDefinitionInfo& PropertyDef, FUnrealFunctionDefinitionInfo* TargetFuncDef)
 {
 	check(TargetFuncDef);
-	FProperty* Prop = PropertyDef.GetProperty();
 	FUnrealPropertyDefinitionInfo* ReturnPropDef = TargetFuncDef->GetReturn();
 
 	if (ReturnPropDef)
@@ -2860,10 +2817,10 @@ void FHeaderParser::VerifyBlueprintPropertySetter(FUnrealPropertyDefinitionInfo&
 	else
 	{
 		const TArray<FUnrealPropertyDefinitionInfo*>& Properties = TargetFuncDef->GetProperties();
-		if (TargetFuncDef->GetProperties().Num() != 1 || !Prop->SameType(Properties[0]->GetProperty()))
+		if (TargetFuncDef->GetProperties().Num() != 1 || !PropertyDef.SameType(*Properties[0]))
 		{
 			FString ExtendedCPPType;
-			FString CPPType = Prop->GetCPPType(&ExtendedCPPType);
+			FString CPPType = PropertyDef.GetCPPType(&ExtendedCPPType);
 			UE_LOG_ERROR_UHT(TEXT("Blueprint Property setter function %s must have exactly one parameter of type %s%s."), *TargetFuncDef->GetName(), *CPPType, *ExtendedCPPType);
 		}
 	}
@@ -2884,7 +2841,6 @@ void FHeaderParser::VerifyBlueprintPropertySetter(FUnrealPropertyDefinitionInfo&
 
 void FHeaderParser::VerifyRepNotifyCallback(FUnrealPropertyDefinitionInfo& PropertyDef, FUnrealFunctionDefinitionInfo* TargetFuncDef)
 {
-	FProperty* Prop = PropertyDef.GetProperty();
 	if (TargetFuncDef)
 	{
 		if (TargetFuncDef->GetReturn())
@@ -2892,7 +2848,7 @@ void FHeaderParser::VerifyRepNotifyCallback(FUnrealPropertyDefinitionInfo& Prope
 			UE_LOG_ERROR_UHT(TEXT("Replication notification function %s must not have return value."), *TargetFuncDef->GetName());
 		}
 
-		const bool bIsArrayProperty = ( Prop->ArrayDim > 1 || CastField<FArrayProperty>(Prop) );
+		const bool bIsArrayProperty = PropertyDef.IsStaticArray() || PropertyDef.IsDynamicArray();
 		const int32 MaxParms = bIsArrayProperty ? 2 : 1;
 
 		const TArray<FUnrealPropertyDefinitionInfo*>& Properties = TargetFuncDef->GetProperties();
@@ -2901,25 +2857,24 @@ void FHeaderParser::VerifyRepNotifyCallback(FUnrealPropertyDefinitionInfo& Prope
 			UE_LOG_ERROR_UHT(TEXT("Replication notification function %s has too many parameters."), *TargetFuncDef->GetName());
 		}
 
-
 		if (Properties.Num() >= 1)
 		{
-			const FProperty* Parm = Properties[0]->GetProperty();
+			const FUnrealPropertyDefinitionInfo* ParmDef = Properties[0];
 			// First parameter is always the old value:
-			if ( !Prop->SameType(Parm) )
+			if (!PropertyDef.SameType(*ParmDef))
 			{
 				FString ExtendedCPPType;
-				FString CPPType = Prop->GetCPPType(&ExtendedCPPType);
-				UE_LOG_ERROR_UHT(TEXT("Replication notification function %s has invalid parameter for property %s. First (optional) parameter must be of type %s%s."), *TargetFuncDef->GetName(), *Prop->GetName(), *CPPType, *ExtendedCPPType);
+				FString CPPType = ParmDef->GetCPPType(&ExtendedCPPType);
+				UE_LOG_ERROR_UHT(TEXT("Replication notification function %s has invalid parameter for property %s. First (optional) parameter must be of type %s%s."), *TargetFuncDef->GetName(), *ParmDef->GetName(), *CPPType, *ExtendedCPPType);
 			}
 		}
 
 		if (TargetFuncDef->GetProperties().Num() >= 2)
 		{
-			FProperty* Parm = Properties[1]->GetProperty();
+			FUnrealPropertyDefinitionInfo* ParmDef = Properties[1];
+			FPropertyBase& ParmBase = ParmDef->GetPropertyBase();
 			// A 2nd parameter for arrays can be specified as a const TArray<uint8>&. This is a list of element indices that have changed
-			FArrayProperty *ArrayProp = CastField<FArrayProperty>(Parm);
-			if (!(ArrayProp && CastField<FByteProperty>(ArrayProp->Inner)) || !(Parm->GetPropertyFlags() & CPF_ConstParm) || !(Parm->GetPropertyFlags() & CPF_ReferenceParm))
+			if (!(ParmBase.Type == CPT_Byte && ParmBase.ArrayType == EArrayType::Dynamic && ParmDef->HasAllPropertyFlags(CPF_ConstParm | CPF_ReferenceParm)))
 			{
 				UE_LOG_ERROR_UHT(TEXT("Replication notification function %s (optional) second parameter must be of type 'const TArray<uint8>&'"), *TargetFuncDef->GetName());
 			}
@@ -2928,7 +2883,7 @@ void FHeaderParser::VerifyRepNotifyCallback(FUnrealPropertyDefinitionInfo& Prope
 	else
 	{
 		// Couldn't find a valid function...
-		UE_LOG_ERROR_UHT(TEXT("Replication notification function %s not found"), *Prop->RepNotifyFunc.ToString() );
+		UE_LOG_ERROR_UHT(TEXT("Replication notification function %s not found"), *PropertyDef.GetRepNotifyFunc().ToString() );
 	}
 }
 void FHeaderParser::VerifyPropertyMarkups(FUnrealClassDefinitionInfo& TargetClassDef)
@@ -2936,8 +2891,6 @@ void FHeaderParser::VerifyPropertyMarkups(FUnrealClassDefinitionInfo& TargetClas
 	// Iterate over all properties, looking for those flagged as CPF_RepNotify
 	for (FUnrealPropertyDefinitionInfo* PropertyDef : TargetClassDef.GetProperties())
 	{
-		FProperty* Prop = PropertyDef->GetProperty();
-
 		auto FindTargetFunction = [&](const FName FuncName)
 		{
 			// Search through this class and its super classes looking for the specified callback
@@ -2960,14 +2913,14 @@ void FHeaderParser::VerifyPropertyMarkups(FUnrealClassDefinitionInfo& TargetClas
 		TGuardValue<int32> GuardedInputPos(InputPos, PropertyDef->GetParsePosition());
 		TGuardValue<int32> GuardedInputLine(InputLine, PropertyDef->GetLineNumber());
 
-		if (Prop->HasAnyPropertyFlags(CPF_RepNotify))
+		if (PropertyDef->HasAnyPropertyFlags(CPF_RepNotify))
 		{
-			VerifyRepNotifyCallback(*PropertyDef, FindTargetFunction(Prop->RepNotifyFunc));
+			VerifyRepNotifyCallback(*PropertyDef, FindTargetFunction(PropertyDef->GetRepNotifyFunc()));
 		}
 
-		if (Prop->HasAnyPropertyFlags(CPF_BlueprintVisible))
+		if (PropertyDef->HasAnyPropertyFlags(CPF_BlueprintVisible))
 		{
-			const FString& GetterFuncName = Prop->GetMetaData(NAME_BlueprintGetter);
+			const FString& GetterFuncName = PropertyDef->GetMetaData(NAME_BlueprintGetter);
 			if (!GetterFuncName.IsEmpty())
 			{
 				if (FUnrealFunctionDefinitionInfo* TargetFuncDef = FindTargetFunction(*GetterFuncName))
@@ -2981,9 +2934,9 @@ void FHeaderParser::VerifyPropertyMarkups(FUnrealClassDefinitionInfo& TargetClas
 				}
 			}
 
-			if (!Prop->HasAnyPropertyFlags(CPF_BlueprintReadOnly))
+			if (!PropertyDef->HasAnyPropertyFlags(CPF_BlueprintReadOnly))
 			{
-				const FString& SetterFuncName = Prop->GetMetaData(NAME_BlueprintSetter);
+				const FString& SetterFuncName = PropertyDef->GetMetaData(NAME_BlueprintSetter);
 				if (!SetterFuncName.IsEmpty())
 				{
 					if (FUnrealFunctionDefinitionInfo* TargetFuncDef = FindTargetFunction(*SetterFuncName))
@@ -4664,7 +4617,7 @@ void FHeaderParser::GetVarType(
 		}
 	}
 
-	if ( VarProperty.IsObject() && VarProperty.Type != CPT_SoftObjectReference && VarProperty.MetaClassDef == nullptr && (VarProperty.PropertyFlags&CPF_Config) != 0 )
+	if ( VarProperty.IsObjectOrInterface() && VarProperty.Type != CPT_SoftObjectReference && VarProperty.MetaClassDef == nullptr && (VarProperty.PropertyFlags&CPF_Config) != 0 )
 	{
 		FError::Throwf(TEXT("Not allowed to use 'config' with object variables"));
 	}
@@ -4823,14 +4776,14 @@ FUnrealPropertyDefinitionInfo& FHeaderParser::GetVarNameAndDim
 
 	// Make sure it doesn't conflict.
 	int32 OuterContextCount = 0;
-	UField* ExistingField = FindField(ParentStruct, Identifier, true, UField::StaticClass(), NULL);
-	FField* ExistingProperty = FindProperty(ParentStruct, Identifier, true, FField::StaticClass(), NULL);
+	FUnrealFunctionDefinitionInfo* ExistingFunctionDef = FindFunction(ParentStruct, Identifier, true, nullptr);
+	FUnrealPropertyDefinitionInfo* ExistingPropertyDef = FindProperty(ParentStruct, Identifier, true);
 
-	if (ExistingField != nullptr || ExistingProperty != nullptr)
+	if (ExistingFunctionDef != nullptr || ExistingPropertyDef != nullptr)
 	{
 		bool bErrorDueToShadowing = true;
 
-		if (ExistingField && ExistingField->IsA(UFunction::StaticClass()) && (VariableCategory != EVariableCategory::Member))
+		if (ExistingFunctionDef && (VariableCategory != EVariableCategory::Member))
 		{
 			// A function parameter with the same name as a method is allowed
 			bErrorDueToShadowing = false;
@@ -4838,10 +4791,9 @@ FUnrealPropertyDefinitionInfo& FHeaderParser::GetVarNameAndDim
 
 		//@TODO: This exception does not seem sound either, but there is enough existing code that it will need to be
 		// fixed up first before the exception it is removed.
-		if (ExistingProperty)
+		if (ExistingPropertyDef)
  		{
- 			FProperty* ExistingProp = CastField<FProperty>(ExistingProperty);
- 			const bool bExistingPropDeprecated = (ExistingProp != nullptr) && ExistingProp->HasAnyPropertyFlags(CPF_Deprecated);
+ 			const bool bExistingPropDeprecated = ExistingPropertyDef->HasAnyPropertyFlags(CPF_Deprecated);
  			const bool bNewPropDeprecated = (VariableCategory == EVariableCategory::Member) && ((VarProperty.PropertyFlags & CPF_Deprecated) != 0);
  			if (bNewPropDeprecated || bExistingPropDeprecated)
  			{
@@ -4856,7 +4808,7 @@ FUnrealPropertyDefinitionInfo& FHeaderParser::GetVarNameAndDim
 				HintText,
 				Identifier,
 				*ParentStruct.GetName(),
-				ExistingField ? *ExistingField->GetOuter()->GetName() : *ExistingProperty->GetOwnerVariant().GetFullName());
+				ExistingFunctionDef ? *ExistingFunctionDef->GetOuter()->GetName() : *ExistingPropertyDef->GetFullName());
 		}
 	}
 
@@ -5004,24 +4956,11 @@ FUnrealPropertyDefinitionInfo& FHeaderParser::GetVarNameAndDim
 	// Create the property
 	// NOTE: Only FProperyBase part of VarProperty is used.  
 	FName PropertyName(Identifier);
-	//VarProperty.StartLine = InputLine;
-	//VarProperty.StartPos = InputPos;
 	FUnrealPropertyDefinitionInfo& PropDef = FPropertyTraits::CreateProperty(VarProperty, ParentStruct, PropertyName, ObjectFlags, VariableCategory, Dimensions.String, SourceFile, InputLine, InputPos);
-	FProperty* Property = PropDef.GetProperty();
 
 	// Add the property to the parent
 	ParentStruct.AddProperty(PropDef);
 
-	// Add to the end of the properties slist
-	FField** Prev = &ParentStruct.GetStruct()->ChildProperties;
-	for (; *Prev != nullptr; Prev = &(*Prev)->Next)
-	{
-		// No body
-	}
-	check(*Prev == nullptr);
-	Property->Next = nullptr;
-	*Prev = Property;
-
 	// if we had any metadata, add it to the class
 	AddMetaDataToClassData(PropDef, TMap<FName,FString>(PropDef.GetPropertyBase().MetaData));
 	return PropDef;
@@ -5640,14 +5579,13 @@ void PostParsingClassSetup(FUnrealClassDefinitionInfo& ClassDef)
 	{
 		for (FUnrealPropertyDefinitionInfo* PropertyDef : ClassDef.GetProperties())
 		{
-			FProperty* Property = PropertyDef->GetProperty();
-			if ((Property->PropertyFlags & CPF_Config) != 0)
+			if (PropertyDef->HasAnyPropertyFlags(CPF_Config))
 			{
 				ClassDef.SetClassFlags(CLASS_Config);
 				if (HasAllOptimizationClassFlags()) break;
 			}
 
-			if (Property->ContainsInstancedObjectProperty())
+			if (PropertyDef->ContainsInstancedObjectProperty())
 			{
 				ClassDef.SetClassFlags(CLASS_HasInstancedReference);
 				if (HasAllOptimizationClassFlags()) break;
@@ -6137,10 +6075,9 @@ void FHeaderParser::ParseRigVMMethodParameters(FUnrealStructDefinitionInfo& Stru
 	// validate the property types for this struct
 	for (FUnrealPropertyDefinitionInfo* PropertyDef : TUHTFieldRange<FUnrealPropertyDefinitionInfo>(StructDef))
 	{
-		FProperty const* const Prop = PropertyDef->GetProperty();
 		FString MemberCPPType;
 		FString ExtendedCPPType;
-		MemberCPPType = Prop->GetCPPType(&ExtendedCPPType);
+		MemberCPPType = PropertyDef->GetCPPType(&ExtendedCPPType);
 
 		if (ExtendedCPPType.IsEmpty() && MemberCPPType.StartsWith(TEnumAsByteText))
 		{
@@ -6148,26 +6085,18 @@ void FHeaderParser::ParseRigVMMethodParameters(FUnrealStructDefinitionInfo& Stru
 		}
 
 		FRigVMParameter Parameter;
-		Parameter.Name = Prop->GetName();
+		Parameter.Name = PropertyDef->GetName();
 		Parameter.Type = MemberCPPType + ExtendedCPPType;
-		Parameter.bConstant = Prop->HasMetaData(NAME_ConstantText);
-		Parameter.bInput = Prop->HasMetaData(NAME_InputText);
-		Parameter.bOutput = Prop->HasMetaData(NAME_OutputText);
-		Parameter.ArraySize = Prop->GetMetaData(NAME_ArraySizeText);
+		Parameter.bConstant = PropertyDef->HasMetaData(NAME_ConstantText);
+		Parameter.bInput = PropertyDef->HasMetaData(NAME_InputText);
+		Parameter.bOutput = PropertyDef->HasMetaData(NAME_OutputText);
+		Parameter.ArraySize = PropertyDef->GetMetaData(NAME_ArraySizeText);
 		Parameter.Getter = GetRefText;
-		Parameter.bEditorOnly = Prop->IsEditorOnlyProperty();
-		Parameter.bSingleton = Prop->HasMetaData(NAME_SingletonText);
+		Parameter.bEditorOnly = PropertyDef->IsEditorOnlyProperty();
+		Parameter.bSingleton = PropertyDef->HasMetaData(NAME_SingletonText);
+		Parameter.bIsEnum = PropertyDef->GetPropertyBase().IsEnum();
 
-		if (const FEnumProperty* EnumProperty = CastField<FEnumProperty>(Prop))
-		{
-			Parameter.bIsEnum = true;
-		}
-		else if (const FByteProperty* ByteProperty = CastField<FByteProperty>(Prop))
-		{
-			Parameter.bIsEnum = ByteProperty->Enum != nullptr;
-		}
-
-		if (Prop->HasMetaData(NAME_VisibleText))
+		if (PropertyDef->HasMetaData(NAME_VisibleText))
 		{
 			Parameter.bConstant = true;
 			Parameter.bInput = true;
@@ -6249,13 +6178,12 @@ void FHeaderParser::ParseParameterList(FUnrealFunctionDefinitionInfo& FunctionDe
 		}
 
 		FUnrealPropertyDefinitionInfo& PropDef = GetVarNameAndDim(FunctionDef, Property, VariableCategory);
-		FProperty* Prop = PropDef.GetProperty();
 
 		FunctionDef.GetFunction()->NumParms++;
 
-		if( AdvancedDisplay.CanMarkMore() && AdvancedDisplay.ShouldMarkParameter(Prop->GetName()) )
+		if( AdvancedDisplay.CanMarkMore() && AdvancedDisplay.ShouldMarkParameter(PropDef.GetName()) )
 		{
-			Prop->PropertyFlags |= CPF_AdvancedDisplay;
+			PropDef.SetPropertyFlags(CPF_AdvancedDisplay);
 		}
 
 		// Check parameters.
@@ -6296,17 +6224,17 @@ void FHeaderParser::ParseParameterList(FUnrealFunctionDefinitionInfo& FunctionDe
 					UE_LOG_ERROR_UHT(TEXT("Only service request functions cannot contain NoReplication parameters"));
 				}
 
-				if ((Prop->GetCastFlags() & CASTCLASS_FDelegateProperty) != 0)
+				if ((PropDef.GetCastFlags() & CASTCLASS_FDelegateProperty) != 0)
 				{
 					UE_LOG_ERROR_UHT(TEXT("Replicated functions cannot contain delegate parameters (this would be insecure)"));
 				}
 
-				if (Property.Type == CPT_String && Property.RefQualifier != ERefQualifier::ConstRef && Prop->ArrayDim == 1)
+				if (Property.Type == CPT_String && Property.RefQualifier != ERefQualifier::ConstRef && !PropDef.IsStaticArray())
 				{
 					UE_LOG_ERROR_UHT(TEXT("Replicated FString parameters must be passed by const reference"));
 				}
 
-				if (Property.ArrayType == EArrayType::Dynamic && Property.RefQualifier != ERefQualifier::ConstRef && Prop->ArrayDim == 1)
+				if (Property.ArrayType == EArrayType::Dynamic && Property.RefQualifier != ERefQualifier::ConstRef && !PropDef.IsStaticArray())
 				{
 					UE_LOG_ERROR_UHT(TEXT("Replicated TArray parameters must be passed by const reference"));
 				}
@@ -6318,7 +6246,7 @@ void FHeaderParser::ParseParameterList(FUnrealFunctionDefinitionInfo& FunctionDe
 					UE_LOG_ERROR_UHT(TEXT("Service request functions cannot contain out parameters, unless marked NotReplicated"));
 				}
 
-				if (!(Property.PropertyFlags & CPF_RepSkip) && (Prop->GetCastFlags() & CASTCLASS_FDelegateProperty) != 0)
+				if (!(Property.PropertyFlags & CPF_RepSkip) && (PropDef.GetCastFlags() & CASTCLASS_FDelegateProperty) != 0)
 				{
 					UE_LOG_ERROR_UHT(TEXT("Service request functions cannot contain delegate parameters, unless marked NotReplicated"));
 				}
@@ -6326,12 +6254,12 @@ void FHeaderParser::ParseParameterList(FUnrealFunctionDefinitionInfo& FunctionDe
 		}
 		if (FunctionDef.HasAnyFunctionFlags(FUNC_BlueprintEvent|FUNC_BlueprintCallable))
 		{
-			if (Property.Type == CPT_Byte)
+			if (Property.Type == CPT_Byte && Property.IsPrimitiveOrPrimitiveStaticArray())
 			{
-				if (FEnumProperty* EnumProperty = CastField<FEnumProperty>(Prop))
+				if (FUnrealEnumDefinitionInfo* EnumDef = Property.AsEnum())
 				{
-					FProperty* InnerType = EnumProperty->GetUnderlyingProperty();
-					if (InnerType && !InnerType->IsA<FByteProperty>())
+					if (EnumDef->GetUnderlyingType() != EUnderlyingEnumType::uint8 &&
+						EnumDef->GetUnderlyingType() != EUnderlyingEnumType::Unspecified)
 					{
 						FError::Throwf(TEXT("Invalid enum param for Blueprints - currently only uint8 supported"));
 					}
@@ -6394,7 +6322,7 @@ void FHeaderParser::ParseParameterList(FUnrealFunctionDefinitionInfo& FunctionDe
 			{
 				FString DefaultArgText(EndPos - StartPos, Input + StartPos);
 				FString Key(TEXT("CPP_Default_"));
-				Key += Prop->GetName();
+				Key += PropDef.GetName();
 				FName KeyName = FName(*Key);
 				if (!MetaData->Contains(KeyName))
 				{
@@ -6402,11 +6330,11 @@ void FHeaderParser::ParseParameterList(FUnrealFunctionDefinitionInfo& FunctionDe
 					const bool bDefaultValueParsed = FPropertyTraits::DefaultValueStringCppFormatToInnerFormat(PropDef, DefaultArgText, InnerDefaultValue);
 					if (!bDefaultValueParsed)
 					{
-						FError::Throwf(TEXT("C++ Default parameter not parsed: %s \"%s\" "), *Prop->GetName(), *DefaultArgText);
+						FError::Throwf(TEXT("C++ Default parameter not parsed: %s \"%s\" "), *PropDef.GetName(), *DefaultArgText);
 					}
 
 					MetaData->Add(KeyName, InnerDefaultValue);
-						UE_LOG(LogCompile, Verbose, TEXT("C++ Default parameter parsed: %s \"%s\" -> \"%s\" "), *Prop->GetName(), *DefaultArgText, *InnerDefaultValue );
+						UE_LOG(LogCompile, Verbose, TEXT("C++ Default parameter parsed: %s \"%s\" -> \"%s\" "), *PropDef.GetName(), *DefaultArgText, *InnerDefaultValue );
 				}
 			}
 		}
@@ -6958,10 +6886,9 @@ void FHeaderParser::CompileFunctionDeclaration()
 	bool bHasAnyOutputs = bHasReturnValue;
 	if (!bHasAnyOutputs)
 	{
-		for (FUnrealPropertyDefinitionInfo* ProperyDef : FuncDef.GetProperties())
+		for (FUnrealPropertyDefinitionInfo* PropertyDef : FuncDef.GetProperties())
 		{
-			FProperty const* const Param = ProperyDef->GetProperty();
-			if (!(Param->PropertyFlags & CPF_ReturnParm) && (Param->PropertyFlags & CPF_OutParm))
+			if (PropertyDef->HasSpecificPropertyFlags(CPF_ReturnParm | CPF_OutParm, CPF_OutParm))
 			{
 				bHasAnyOutputs = true;
 				break;
@@ -7089,17 +7016,16 @@ void FHeaderParser::CompileFunctionDeclaration()
 	{
 		for (FUnrealPropertyDefinitionInfo* PropertyDef : FuncDef.GetProperties())
 		{
-			FProperty const* const Param = PropertyDef->GetProperty();
-			if (Param->ArrayDim > 1)
+			if (PropertyDef->IsStaticArray())
 			{
-				FError::Throwf(TEXT("Static array cannot be exposed to blueprint. Function: %s Parameter %s\n"), *TopFunction->GetName(), *Param->GetName());
+				FError::Throwf(TEXT("Static array cannot be exposed to blueprint. Function: %s Parameter %s\n"), *TopFunction->GetName(), *PropertyDef->GetName());
 			}
 
 			if (!FPropertyTraits::IsSupportedByBlueprint(*PropertyDef, false))
 			{
 				FString ExtendedCPPType;
-				FString CPPType = Param->GetCPPType(&ExtendedCPPType);
-				UE_LOG_ERROR_UHT(TEXT("Type '%s%s' is not supported by blueprint. %s.%s"), *CPPType, *ExtendedCPPType, *TopFunction->GetName(), *Param->GetName());
+				FString CPPType = PropertyDef->GetCPPType(&ExtendedCPPType);
+				UE_LOG_ERROR_UHT(TEXT("Type '%s%s' is not supported by blueprint. %s.%s"), *CPPType, *ExtendedCPPType, *TopFunction->GetName(), *PropertyDef->GetName());
 			}
 		}
 	}
@@ -7296,24 +7222,23 @@ bool FHeaderParser::ValidateScriptStructOkForNet(const FString& OriginStructName
 
 	for (FUnrealPropertyDefinitionInfo* PropertyDef : InStructDef.GetProperties())
 	{
-		FProperty* ChildProp = PropertyDef->GetProperty();
-		if (const FSetProperty* const SetProp = CastField<FSetProperty>(ChildProp))
+		if (PropertyDef->IsSet())
 		{
-			if (!(SetProp->PropertyFlags & CPF_RepSkip))
+			if (!PropertyDef->HasAnyPropertyFlags(CPF_RepSkip))
 			{
 				bIsStructValid = false;
-				UE_LOG_ERROR_UHT(TEXT("Sets are not supported for Replication or RPCs.  Set %s in %s.  Origin %s"), *ChildProp->GetName(), *ChildProp->Owner.GetName(), *OriginStructName);
+				UE_LOG_ERROR_UHT(TEXT("Sets are not supported for Replication or RPCs.  Set %s in %s.  Origin %s"), *PropertyDef->GetName(), *PropertyDef->GetOuter()->GetName(), *OriginStructName);
 			}
 		}
-		else if (const FMapProperty* const MapProp = CastField<FMapProperty>(ChildProp))
+		else if (PropertyDef->IsMap())
 		{
-			if (!(MapProp->PropertyFlags & CPF_RepSkip))
+			if (!PropertyDef->HasAnyPropertyFlags(CPF_RepSkip))
 			{
 				bIsStructValid = false;
-				UE_LOG_ERROR_UHT(TEXT("Maps are not supported for Replication or RPCs.  Map %s in %s.  Origin %s"), *ChildProp->GetName(), *ChildProp->Owner.GetName(), *OriginStructName);
+				UE_LOG_ERROR_UHT(TEXT("Maps are not supported for Replication or RPCs.  Map %s in %s.  Origin %s"), *PropertyDef->GetName(), *PropertyDef->GetOuter()->GetName(), *OriginStructName);
 			}
 		}
-		else if (const FStructProperty* const StructProperty = CastField<FStructProperty>(ChildProp))
+		else if (PropertyDef->IsStructOrStructStaticArray())
 		{
 			if (!ValidateScriptStructOkForNet(OriginStructName, *PropertyDef->GetPropertyBase().ScriptStructDef))
 			{
@@ -7446,12 +7371,11 @@ void FHeaderParser::CompileVariableDeclaration(FUnrealStructDefinitionInfo& Stru
 	}
 
 	// Process all variables of this type.
-	TArray<FProperty*> NewProperties;
+	TArray<FUnrealPropertyDefinitionInfo*> NewProperties;
 	for (;;)
 	{
 		FPropertyBase Property = OriginalProperty;
 		FUnrealPropertyDefinitionInfo& NewPropDef = GetVarNameAndDim(StructDef, Property, EVariableCategory::Member, LayoutMacroType);
-		FProperty* NewProperty = NewPropDef.GetProperty();
 
 		// Optionally consume the :1 at the end of a bitfield boolean declaration
 		if (Property.IsBool())
@@ -7461,7 +7385,7 @@ void FHeaderParser::CompileVariableDeclaration(FUnrealStructDefinitionInfo& Stru
 				int32 BitfieldSize = 0;
 				if (!GetConstInt(/*out*/ BitfieldSize) || (BitfieldSize != 1))
 				{
-					FError::Throwf(TEXT("Bad or missing bitfield size for '%s', must be 1."), *NewProperty->GetName());
+					FError::Throwf(TEXT("Bad or missing bitfield size for '%s', must be 1."), *NewPropDef.GetName());
 				}
 			}
 		}
@@ -7473,44 +7397,44 @@ void FHeaderParser::CompileVariableDeclaration(FUnrealStructDefinitionInfo& Stru
 		{
 			if (NewProperties.Num())
 			{
-				FError::Throwf(TEXT("Comma delimited properties cannot be converted %s.%s\n"), *StructDef.GetName(), *NewProperty->GetName());
+				FError::Throwf(TEXT("Comma delimited properties cannot be converted %s.%s\n"), *StructDef.GetName(), *NewPropDef.GetName());
 			}
 		}
 
-		NewProperties.Add( NewProperty );
+		NewProperties.Add(&NewPropDef);
 		// we'll need any metadata tags we parsed later on when we call ConvertEOLCommentToTooltip() so the tags aren't clobbered
 		OriginalProperty.MetaData = Property.MetaData;
 
-		if (NewProperty->HasAnyPropertyFlags(CPF_RepNotify))
+		if (NewPropDef.HasAnyPropertyFlags(CPF_RepNotify))
 		{
-			NewProperty->RepNotifyFunc = OriginalProperty.RepNotifyName;
+			NewPropDef.SetRepNotifyFunc(OriginalProperty.RepNotifyName);
 		}
 
 		if (UScriptStruct* StructBeingBuilt = Cast<UScriptStruct>(Struct))
 		{
-			if (NewProperty->ContainsInstancedObjectProperty())
+			if (NewPropDef.ContainsInstancedObjectProperty())
 			{
 				StructBeingBuilt->StructFlags = EStructFlags(StructBeingBuilt->StructFlags | STRUCT_HasInstancedReference);
 			}
 		}
 
-		if (NewProperty->HasAnyPropertyFlags(CPF_BlueprintVisible))
+		if (NewPropDef.HasAnyPropertyFlags(CPF_BlueprintVisible))
 		{
 			if (Struct->IsA<UScriptStruct>() && !StructDef.GetBoolMetaDataHierarchical(FHeaderParserNames::NAME_BlueprintType))
 			{
-				UE_LOG_ERROR_UHT(TEXT("Cannot expose property to blueprints in a struct that is not a BlueprintType. %s.%s"), *StructDef.GetName(), *NewProperty->GetName());
+				UE_LOG_ERROR_UHT(TEXT("Cannot expose property to blueprints in a struct that is not a BlueprintType. %s.%s"), *StructDef.GetName(), *NewPropDef.GetName());
 			}
 
-			if (NewProperty->ArrayDim > 1)
+			if (NewPropDef.IsStaticArray())
 			{
-				UE_LOG_ERROR_UHT(TEXT("Static array cannot be exposed to blueprint %s.%s"), *StructDef.GetName(), *NewProperty->GetName());
+				UE_LOG_ERROR_UHT(TEXT("Static array cannot be exposed to blueprint %s.%s"), *StructDef.GetName(), *NewPropDef.GetName());
 			}
 
 			if (!FPropertyTraits::IsSupportedByBlueprint(NewPropDef, true))
 			{
 				FString ExtendedCPPType;
-				FString CPPType = NewProperty->GetCPPType(&ExtendedCPPType);
-				UE_LOG_ERROR_UHT(TEXT("Type '%s%s' is not supported by blueprint. %s.%s"), *CPPType, *ExtendedCPPType, *StructDef.GetName(), *NewProperty->GetName());
+				FString CPPType = NewPropDef.GetCPPType(&ExtendedCPPType);
+				UE_LOG_ERROR_UHT(TEXT("Type '%s%s' is not supported by blueprint. %s.%s"), *CPPType, *ExtendedCPPType, *StructDef.GetName(), *NewPropDef.GetName());
 			}
 		}
 
@@ -7647,12 +7571,12 @@ void FHeaderParser::ComputeFunctionParametersSize(FUnrealClassDefinitionInfo& Cl
 		{
 			for (FUnrealPropertyDefinitionInfo* PropertyDef : FunctionDef->GetProperties())
 			{
-				FProperty* Param = PropertyDef->GetProperty();
-				if (FStructProperty* StructProp = CastField<FStructProperty>(Param))
+				if (PropertyDef->IsStructOrStructStaticArray())
 				{
-					if (StructProp->Struct->StructFlags & STRUCT_HasInstancedReference)
+					FUnrealScriptStructDefinitionInfo& StructDef = UHTCastChecked<FUnrealScriptStructDefinitionInfo>(PropertyDef->GetPropertyBase().ClassDef);
+					if (StructDef.HasAnyStructFlags(STRUCT_HasInstancedReference))
 					{
-						StructProp->PropertyFlags |= CPF_ContainsInstancedReference;
+						PropertyDef->SetPropertyFlags(CPF_ContainsInstancedReference);
 					}
 				}
 			}
@@ -7664,17 +7588,16 @@ void FHeaderParser::ComputeFunctionParametersSize(FUnrealClassDefinitionInfo& Cl
 		ThisFunction->ParmsSize = 0;
 		for (FUnrealPropertyDefinitionInfo* PropertyDef : FunctionDef->GetProperties())
 		{
-			FProperty* Param = PropertyDef->GetProperty();
-
-			if (!(Param->PropertyFlags & CPF_ReturnParm) && (Param->PropertyFlags & CPF_OutParm))
+			if (PropertyDef->HasSpecificPropertyFlags(CPF_ReturnParm | CPF_OutParm, CPF_OutParm))
 			{
 				ThisFunction->FunctionFlags |= FUNC_HasOutParms;
 				FunctionDef->GetFunctionData().FunctionFlags |= FUNC_HasOutParms;
 			}
 				
-			if (FStructProperty* StructProp = CastField<FStructProperty>(Param))
+			if (PropertyDef->IsStructOrStructStaticArray())
 			{
-				if (StructProp->Struct->HasDefaults())
+				FUnrealScriptStructDefinitionInfo& StructDef = UHTCastChecked<FUnrealScriptStructDefinitionInfo>(PropertyDef->GetPropertyBase().ClassDef);
+				if (StructDef.HasDefaults())
 				{
 					ThisFunction->FunctionFlags |= FUNC_HasDefaults;
 					FunctionDef->GetFunctionData().FunctionFlags |= FUNC_HasDefaults;
@@ -9110,14 +9033,13 @@ void FHeaderParser::PostPopNestClass(FUnrealClassDefinitionInfo& CurrentClassDef
 						FUnrealPropertyDefinitionInfo* ClassParamDef = ClassFunctionDef->GetProperties()[Index];
 						if (!InterfaceParamDef->MatchesType(*ClassParamDef, true))
 						{
-							FProperty* Property = InterfaceParamDef->GetProperty();
-							if (Property->HasAnyPropertyFlags(CPF_ReturnParm))
+							if (InterfaceParamDef->HasAnyPropertyFlags(CPF_ReturnParm))
 							{
 								FError::Throwf(TEXT("Implementation of function '%s' conflicts only by return type with interface '%s'"), *InterfaceFunctionDef->GetName(), *InterfaceDef->GetName());
 							}
 							else
 							{
-								FError::Throwf(TEXT("Implementation of function '%s' conflicts with interface '%s' - parameter %i '%s'"), *InterfaceFunctionDef->GetName(), *InterfaceDef->GetName(), End, *Property->GetName());
+								FError::Throwf(TEXT("Implementation of function '%s' conflicts with interface '%s' - parameter %i '%s'"), *InterfaceFunctionDef->GetName(), *InterfaceDef->GetName(), End, *InterfaceParamDef->GetName());
 							}
 						}
 					}
@@ -9265,19 +9187,18 @@ void FHeaderParser::CheckDocumentationPolicyForStruct(FUnrealStructDefinitionInf
 		TMap<FString, FName> ToolTipToPropertyName;
 		for (FUnrealPropertyDefinitionInfo* PropertyDef : StructDef.GetProperties())
 		{
-			FProperty* Property = PropertyDef->GetProperty();
-			FString ToolTip = Property->GetToolTipText().ToString();
-			if (ToolTip.IsEmpty() || ToolTip.Equals(Property->GetDisplayNameText().ToString()))
+			FString ToolTip = PropertyDef->GetToolTipText().ToString();
+			if (ToolTip.IsEmpty() || ToolTip.Equals(PropertyDef->GetDisplayNameText().ToString()))
 			{
-				UE_LOG_ERROR_UHT(TEXT("Property '%s::%s' does not provide a tooltip / comment (DocumentationPolicy)."), *StructDef.GetName(), *Property->GetName());
+				UE_LOG_ERROR_UHT(TEXT("Property '%s::%s' does not provide a tooltip / comment (DocumentationPolicy)."), *StructDef.GetName(), *PropertyDef->GetName());
 				continue;
 			}
 			const FName* ExistingPropertyName = ToolTipToPropertyName.Find(ToolTip);
 			if (ExistingPropertyName != nullptr)
 			{
-				UE_LOG_ERROR_UHT(TEXT("Property '%s::%s' and '%s::%s' are using identical tooltips (DocumentationPolicy)."), *StructDef.GetName(), *ExistingPropertyName->ToString(), *StructDef.GetName(), *Property->GetName());
+				UE_LOG_ERROR_UHT(TEXT("Property '%s::%s' and '%s::%s' are using identical tooltips (DocumentationPolicy)."), *StructDef.GetName(), *ExistingPropertyName->ToString(), *StructDef.GetName(), *PropertyDef->GetName());
 			}
-			ToolTipToPropertyName.Add(MoveTemp(ToolTip), Property->GetFName());
+			ToolTipToPropertyName.Add(MoveTemp(ToolTip), PropertyDef->GetFName());
 		}
 	}
 
@@ -9285,15 +9206,14 @@ void FHeaderParser::CheckDocumentationPolicyForStruct(FUnrealStructDefinitionInf
 	{
 		for (FUnrealPropertyDefinitionInfo* PropertyDef : StructDef.GetProperties())
 		{
-			FProperty* Property = PropertyDef->GetProperty();
-			if(DoesCPPTypeRequireDocumentation(Property->GetCPPType()))
+			if(DoesCPPTypeRequireDocumentation(PropertyDef->GetCPPType()))
 			{
-				const FString& UIMin = Property->GetMetaData(NAME_UIMin);
-				const FString& UIMax = Property->GetMetaData(NAME_UIMax);
+				const FString& UIMin = PropertyDef->GetMetaData(NAME_UIMin);
+				const FString& UIMax = PropertyDef->GetMetaData(NAME_UIMax);
 
 				if(!CheckUIMinMaxRangeFromMetaData(UIMin, UIMax))
 				{
-					UE_LOG_ERROR_UHT(TEXT("Property '%s::%s' does not provide a valid UIMin / UIMax (DocumentationPolicy)."), *StructDef.GetName(), *Property->GetName());
+					UE_LOG_ERROR_UHT(TEXT("Property '%s::%s' does not provide a valid UIMin / UIMax (DocumentationPolicy)."), *StructDef.GetName(), *PropertyDef->GetName());
 				}
 			}
 		}
@@ -9372,8 +9292,7 @@ void FHeaderParser::CheckDocumentationPolicyForFunc(FUnrealClassDefinitionInfo&
 			TSet<FName> ExistingFields;
 			for (FUnrealPropertyDefinitionInfo* PropertyDef : FunctionDef.GetProperties())
 			{
-				FProperty* Property = PropertyDef->GetProperty();
-				FName ParamName = Property->GetFName();
+				FName ParamName = PropertyDef->GetFName();
 				if (ParamName == NAME_ReturnValue)
 				{
 					continue;
diff --git a/Engine/Source/Programs/UnrealHeaderTool/Private/HeaderParser.h b/Engine/Source/Programs/UnrealHeaderTool/Private/HeaderParser.h
--- a/Engine/Source/Programs/UnrealHeaderTool/Private/HeaderParser.h
+++ b/Engine/Source/Programs/UnrealHeaderTool/Private/HeaderParser.h
@@ -371,8 +371,8 @@ public:
 	*
 	* @return	a pointer to a UField with a name matching InIdentifier, or NULL if it wasn't found
 	*/
-	static UField* FindField(FUnrealStructDefinitionInfo& InScope, const TCHAR* InIdentifier, bool bIncludeParents=true, UClass* FieldClass=UField::StaticClass(), const TCHAR* Thing=nullptr );
-	static FField* FindProperty(FUnrealStructDefinitionInfo& InScope, const TCHAR* InIdentifier, bool bIncludeParents = true, FFieldClass* FieldClass = FField::StaticClass(), const TCHAR* Thing = nullptr);
+	static FUnrealFunctionDefinitionInfo* FindFunction(FUnrealStructDefinitionInfo& InScope, const TCHAR* InIdentifier, bool bIncludeParents = true, const TCHAR* Thing = nullptr);
+	static FUnrealPropertyDefinitionInfo* FindProperty(FUnrealStructDefinitionInfo& InScope, const TCHAR* InIdentifier, bool bIncludeParents = true, const TCHAR* Thing = nullptr);
 
 	// Checks ToValidate to make sure that its associated sparse class data struct, if one exists, is a valid structure to use for storing sparse class data.
 	static void CheckSparseClassData(const FUnrealStructDefinitionInfo& StructDef);
diff --git a/Engine/Source/Programs/UnrealHeaderTool/Private/NativeClassExporter.h b/Engine/Source/Programs/UnrealHeaderTool/Private/NativeClassExporter.h
--- a/Engine/Source/Programs/UnrealHeaderTool/Private/NativeClassExporter.h
+++ b/Engine/Source/Programs/UnrealHeaderTool/Private/NativeClassExporter.h
@@ -560,13 +560,13 @@ private:
 	 * @param	DeclOut			String output device to send the generated declarations to
 	 * @param	Out				String output device to send the generated code to
 	 * @param	Scope			The scope to prefix on all variable definitions
-	 * @param	Properties		Array of properties to export
+	 * @param	StructDef		The structure containing the properties to export
 	 * @param	Spaces			String of spaces to use as an indent for the declaration
 	 * @param	Spaces			String of spaces to use as an indent
 	 *
 	 * @return      A pair of strings which represents the pointer and a count of the emitted properties.
 	 */
-	TTuple<FString, FString> OutputProperties(FOutputDevice& DeclOut, FOutputDevice& Out, FReferenceGatherers& OutReferenceGatherers, const TCHAR* Scope, const TArray<FUnrealPropertyDefinitionInfo*>& PropertyDefs, const TCHAR* DeclSpaces, const TCHAR* Spaces) const;
+	TTuple<FString, FString> OutputProperties(FOutputDevice& DeclOut, FOutputDevice& Out, FReferenceGatherers& OutReferenceGatherers, const TCHAR* Scope, FUnrealStructDefinitionInfo& StructDef, const TCHAR* DeclSpaces, const TCHAR* Spaces) const;
 
 	/**
 	 * Function to output the C++ code necessary to set up a property
@@ -578,21 +578,7 @@ private:
 	 * @param	DeclSpaces		String of spaces to use as an indent for the declaration
 	 * @param	Spaces			String of spaces to use as an indent
 	**/
-	void OutputProperty(FOutputDevice& DeclOut, FOutputDevice& Out, FReferenceGatherers& OutReferenceGatherers, const TCHAR* Scope, TArray<FPropertyNamePointerPair>& PropertyNamesAndPointers, FUnrealPropertyDefinitionInfo& PropertyDef, const TCHAR* DeclSpaces, const TCHAR* Spaces) const;
-
-	/**
-	 * Function to output the C++ code necessary to set up a property, including an array property and its inner, array dimensions, etc.
-	 *
-	 * @param	DeclOut			The destination to write declarations to.
-	 * @param	Out				The destination to write to.
-	 * @param	Prop			Property to export
-	 * @param	OffsetStr		String specifying the property offset
-	 * @param	Name			Name for the generated variable
-	 * @param	DeclSpaces		String of spaces to use as an indent for the declaration
-	 * @param	Spaces			String of spaces to use as an indent
-	 * @param	SourceStruct	Structure that the property offset is relative to
-	**/
-	void PropertyNew(FOutputDevice& DeclOut, FOutputDevice& Out, FReferenceGatherers& OutReferenceGatherers, FUnrealPropertyDefinitionInfo& PropertyDef, const TCHAR* OffsetStr, const TCHAR* Name, const TCHAR* DeclSpaces, const TCHAR* Spaces, const TCHAR* SourceStruct = nullptr) const;
+	void OutputProperty(FOutputDevice& DeclOut, FOutputDevice& Out, FReferenceGatherers& OutReferenceGatherers, const TCHAR* Scope, TArray<FPropertyNamePointerPair>& PropertyNamesAndPointers, FUnrealPropertyDefinitionInfo& PropertyDef, const TCHAR* OffsetStr, FString&& Name, const TCHAR* DeclSpaces, const TCHAR* Spaces, const TCHAR* SourceStruct) const;
 
 	/**
 	 * Function to generate the property tag
diff --git a/Engine/Source/Programs/UnrealHeaderTool/Private/ParserHelper.cpp b/Engine/Source/Programs/UnrealHeaderTool/Private/ParserHelper.cpp
--- a/Engine/Source/Programs/UnrealHeaderTool/Private/ParserHelper.cpp
+++ b/Engine/Source/Programs/UnrealHeaderTool/Private/ParserHelper.cpp
@@ -63,53 +63,65 @@ FPropertyBase::FPropertyBase(FUnrealScriptStructDefinitionInfo& InStructDef)
 {
 }
 
-const TCHAR* FPropertyBase::GetPropertyTypeText( EPropertyType Type )
+FUnrealEnumDefinitionInfo* FPropertyBase::AsEnum() const
 {
-	switch ( Type )
+	return UHTCast<FUnrealEnumDefinitionInfo>(TypeDef);
+}
+
+bool FPropertyBase::IsEnum() const
+{
+	return AsEnum() != nullptr;
+}
+
+EUHTPropertyType FPropertyBase::GetUHTPropertyType() const
+{
+	if (ArrayType == EArrayType::Dynamic)
+	{
+		return EUHTPropertyType::DynamicArray;
+	}
+	else if (ArrayType == EArrayType::Set)
+	{
+		return EUHTPropertyType::Set;
+	}
+	else if (MapKeyProp.IsValid())
+	{
+		return EUHTPropertyType::Map;
+	}
+	else if (IsEnum())
+	{
+		return EUHTPropertyType::Enum;
+	}
+	else
 	{
-		CASE_TEXT(CPT_None);
-		CASE_TEXT(CPT_Byte);
-		CASE_TEXT(CPT_Int8);
-		CASE_TEXT(CPT_Int16);
-		CASE_TEXT(CPT_Int);
-		CASE_TEXT(CPT_Int64);
-		CASE_TEXT(CPT_UInt16);
-		CASE_TEXT(CPT_UInt32);
-		CASE_TEXT(CPT_UInt64);
-		CASE_TEXT(CPT_Bool);
-		CASE_TEXT(CPT_Bool8);
-		CASE_TEXT(CPT_Bool16);
-		CASE_TEXT(CPT_Bool32);
-		CASE_TEXT(CPT_Bool64);
-		CASE_TEXT(CPT_Float);
-		CASE_TEXT(CPT_Double);
-		CASE_TEXT(CPT_ObjectReference);
-		CASE_TEXT(CPT_Interface);
-		CASE_TEXT(CPT_Name);
-		CASE_TEXT(CPT_Delegate);
-		CASE_TEXT(CPT_Struct);
-		CASE_TEXT(CPT_String);
-		CASE_TEXT(CPT_Text);
-		CASE_TEXT(CPT_MulticastDelegate);
-		CASE_TEXT(CPT_SoftObjectReference);
-		CASE_TEXT(CPT_WeakObjectReference);
-		CASE_TEXT(CPT_LazyObjectReference);
-		CASE_TEXT(CPT_ObjectPtrReference);
-		CASE_TEXT(CPT_Map);
-		CASE_TEXT(CPT_Set);
-		CASE_TEXT(CPT_FieldPath);
-		CASE_TEXT(CPT_MAX);
+		return EUHTPropertyType(Type);
 	}
+}
 
-	return TEXT("");
+bool FPropertyBase::IsClassRefOrClassRefStaticArray() const
+{
+	return IsObjectRefOrObjectRefStaticArray() && ClassDef->GetClass()->IsChildOf(UClass::StaticClass());
+}
+
+bool FPropertyBase::IsByteEnumOrByteEnumStaticArray() const
+{
+	if (Type != CPT_Byte || !IsPrimitiveOrPrimitiveStaticArray())
+	{
+		return false;
+	}
+	FUnrealEnumDefinitionInfo* Enum = AsEnum();
+	if (Enum == nullptr)
+	{
+		return false;
+	}
+	return Enum->GetCppForm() != UEnum::ECppForm::EnumClass;
 }
 
-bool FPropertyBase::MatchesType(const FPropertyBase& Other, bool bDisallowGeneralization, bool bIgnoreImplementedInterfaces/* = false*/) const
+bool FPropertyBase::MatchesType(const FPropertyBase& Other, bool bDisallowGeneralization, bool bIgnoreImplementedInterfaces/* = false*/, bool bEmulateSameType/* = false*/) const
 {
 	check(Type != CPT_None || !bDisallowGeneralization);
 
-	bool bIsObjectType = IsObject();
-	bool bOtherIsObjectType = Other.IsObject();
+	bool bIsObjectType = IsObjectOrInterface();
+	bool bOtherIsObjectType = Other.IsObjectOrInterface();
 	bool bIsObjectComparison = bIsObjectType && bOtherIsObjectType;
 	bool bReverseClassChainCheck = true;
 
@@ -136,7 +148,7 @@ bool FPropertyBase::MatchesType(const FPropertyBase& Other, bool bDisallowGenera
 
 		// if Type == CPT_ObjectReference, out object function parm; allow derived classes to be passed in
 		// if Type == CPT_Interface, out interface function parm; allow derived classes to be passed in
-		else if ((PropertyFlags & CPF_ConstParm) == 0 || !IsObject())
+		else if ((PropertyFlags & CPF_ConstParm) == 0 || !IsObjectOrInterface())
 		{
 			// all other variable types must match exactly when passed as the value to an 'out' parameter
 			bDisallowGeneralization = true;
@@ -159,7 +171,7 @@ bool FPropertyBase::MatchesType(const FPropertyBase& Other, bool bDisallowGenera
 		// If Other has no type, accept anything.
 		return true;
 	}
-	else if (Type != Other.Type && !bIsObjectComparison)
+	else if (Type != Other.Type && !(bEmulateSameType && ::IsBool(Type) && ::IsBool(Other.Type)) && !bIsObjectComparison)
 	{
 		// Mismatched base types.
 		return false;
diff --git a/Engine/Source/Programs/UnrealHeaderTool/Private/ParserHelper.h b/Engine/Source/Programs/UnrealHeaderTool/Private/ParserHelper.h
--- a/Engine/Source/Programs/UnrealHeaderTool/Private/ParserHelper.h
+++ b/Engine/Source/Programs/UnrealHeaderTool/Private/ParserHelper.h
@@ -106,20 +106,61 @@ enum EAccessSpecifier
 	ACCESS_Num,
 };
 
-#ifndef CASE_TEXT
-#define CASE_TEXT(txt) case txt: return TEXT(#txt)
-#endif
+enum class EUHTPropertyType : uint8
+{
+	None = CPT_None,
+	Byte = CPT_Byte,
+	UInt16 = CPT_UInt16,
+	UInt32 = CPT_UInt32,
+	UInt64 = CPT_UInt64,
+	Int8 = CPT_Int8,
+	Int16 = CPT_Int16,
+	Int = CPT_Int,
+	Int64 = CPT_Int64,
+	Bool = CPT_Bool,
+	Bool8 = CPT_Bool8,
+	Bool16 = CPT_Bool16,
+	Bool32 = CPT_Bool32,
+	Bool64 = CPT_Bool64,
+	Float = CPT_Float,
+	ObjectReference = CPT_ObjectReference,
+	Name = CPT_Name,
+	Delegate = CPT_Delegate,
+	Interface = CPT_Interface,
+	Struct = CPT_Struct,
+	String = CPT_String,
+	Text = CPT_Text,
+	MulticastDelegate = CPT_MulticastDelegate,
+	WeakObjectReference = CPT_WeakObjectReference,
+	LazyObjectReference = CPT_LazyObjectReference,
+	ObjectPtrReference = CPT_ObjectPtrReference,
+	SoftObjectReference = CPT_SoftObjectReference,
+	Double = CPT_Double,
+	Map = CPT_Map,
+	Set = CPT_Set,
+	FieldPath = CPT_FieldPath,
+	LargeWorldCoordinatesReal = CPT_FLargeWorldCoordinatesReal,
+
+	Enum,
+	DynamicArray,
+
+	MAX,
+};
 
 inline bool IsBool(EPropertyType Type)
 {
 	return Type == CPT_Bool || Type == CPT_Bool8 || Type == CPT_Bool16 || Type == CPT_Bool32 || Type == CPT_Bool64;
 }
 
-inline bool IsObject(EPropertyType Type)
+inline bool IsNumeric(EPropertyType Type)
 {
-	return Type == CPT_ObjectReference || Type == CPT_Interface || Type == CPT_WeakObjectReference || Type == CPT_LazyObjectReference || Type == CPT_ObjectPtrReference || Type == CPT_SoftObjectReference;
+	return Type == CPT_Byte || Type == CPT_UInt16 || Type == CPT_UInt32 || Type == CPT_UInt64 || Type == CPT_Int8 || Type == CPT_Int16 || Type == CPT_Int || Type == CPT_Int64 || Type == CPT_Float || Type == CPT_Double;
 }
 
+inline bool IsObjectOrInterface(EPropertyType Type)
+{
+	return Type == CPT_ObjectReference || Type == CPT_Interface || Type == CPT_WeakObjectReference || Type == CPT_LazyObjectReference || Type == CPT_ObjectPtrReference || Type == CPT_SoftObjectReference;
+}
 
 /**
  * Basic information describing a type.
@@ -233,9 +274,9 @@ public:
 	/**
 	 * Returns whether this token represents an object reference
 	 */
-	bool IsObject() const
+	bool IsObjectOrInterface() const
 	{
-		return ::IsObject(Type);
+		return ::IsObjectOrInterface(Type);
 	}
 
 	bool IsBool() const
@@ -248,6 +289,44 @@ public:
 		return (ArrayType != EArrayType::None || MapKeyProp.IsValid());
 	}
 
+	bool IsPrimitiveOrPrimitiveStaticArray() const
+	{
+		return (ArrayType == EArrayType::None || ArrayType == EArrayType::Static) && !MapKeyProp.IsValid();
+	}
+
+	bool IsBooleanOrBooleanStaticArray() const
+	{
+		return IsBool() && IsPrimitiveOrPrimitiveStaticArray();
+	}
+
+	bool IsStructOrStructStaticArray() const
+	{
+		return Type == CPT_Struct && IsPrimitiveOrPrimitiveStaticArray();
+	}
+
+	bool IsObjectRefOrObjectRefStaticArray() const
+	{
+		return (Type == CPT_ObjectReference || Type == CPT_ObjectPtrReference) && IsPrimitiveOrPrimitiveStaticArray();
+	}
+
+	bool IsClassRefOrClassRefStaticArray() const;
+
+	bool IsInterfaceOrInterfaceStaticArray() const
+	{
+		return Type == CPT_Interface && IsPrimitiveOrPrimitiveStaticArray();
+	}
+
+	bool IsByteEnumOrByteEnumStaticArray() const;
+
+	bool IsNumericOrNumericStaticArray() const
+	{
+		return IsNumeric(Type) && IsPrimitiveOrPrimitiveStaticArray();
+	}
+
+	FUnrealEnumDefinitionInfo* AsEnum() const;
+
+	bool IsEnum() const;
+
 	/**
 	 * Determines whether this token's type is compatible with another token's type.
 	 *
@@ -265,8 +344,10 @@ public:
 	 *											of the other token's type (or vice versa, when dealing with structs
 	 * @param	bIgnoreImplementedInterfaces	controls whether two types can be considered a match if one type is an interface implemented
 	 *											by the other type.
+	 * @param   bEmulateSameType				If true, perform slightly different validation as per FProperty::SameType.  Implementation is not
+	 *											complete.
 	 */
-	bool MatchesType(const FPropertyBase& Other, bool bDisallowGeneralization, bool bIgnoreImplementedInterfaces = false) const;
+	bool MatchesType(const FPropertyBase& Other, bool bDisallowGeneralization, bool bIgnoreImplementedInterfaces = false, bool bEmulateSameType = false) const;
 	//@}
 
 	EIntType GetSizedIntTypeFromPropertyType(EPropertyType PropType)
@@ -288,7 +369,7 @@ public:
 		}
 	}
 
-	static const TCHAR* GetPropertyTypeText( EPropertyType Type );
+	EUHTPropertyType GetUHTPropertyType() const;
 
 	friend struct FPropertyBaseArchiveProxy;
 };
diff --git a/Engine/Source/Programs/UnrealHeaderTool/Private/PropertyTypes.cpp b/Engine/Source/Programs/UnrealHeaderTool/Private/PropertyTypes.cpp
--- a/Engine/Source/Programs/UnrealHeaderTool/Private/PropertyTypes.cpp
+++ b/Engine/Source/Programs/UnrealHeaderTool/Private/PropertyTypes.cpp
@@ -179,7 +179,7 @@ namespace
 
 		// Check if it's an enum class property
 		// NOTE: VarProperty.EnumDef is a union and might not be an enum
-		if (UHTCast<FUnrealEnumDefinitionInfo>(VarProperty.TypeDef) != nullptr)
+		if (VarProperty.IsEnum())
 		{
 			return FuncDispatch<FPropertyTypeTraitsEnum>()(std::forward<Args>(args)...);
 		}
@@ -202,15 +202,6 @@ namespace
 		}
 	};
 
-	template<typename TraitsType>
-	struct IsObjectDispatch
-	{
-		bool operator()()
-		{
-			return TraitsType::bIsObject;
-		}
-	};
-
 	template<typename TraitsType>
 	struct CreateEngineTypeDispatch
 	{
@@ -254,11 +245,6 @@ namespace
  */
 struct FPropertyTypeTraitsBase
 {
-	/**
-	 * If true, this property type is an object property
-	*/
-	static constexpr bool bIsObject = false;
-
 	/**
 	 * Transforms CPP-formated string containing default value, to inner formated string
 	 * If it cannot be transformed empty string is returned.
@@ -690,8 +676,6 @@ struct FPropertyTypeTraitsEnum : public FPropertyTypeTraitsBase
  */
 struct FPropertyTypeTraitsObjectBase : public FPropertyTypeTraitsBase
 {
-	static constexpr bool bIsObject = true;
-
 	static bool DefaultValueStringCppFormatToInnerFormat(const FUnrealPropertyDefinitionInfo& PropDef, const FString& CppForm, FString& OutForm)
 	{
 		const bool bIsNull = FDefaultValueHelper::Is(CppForm, TEXT("NULL")) || FDefaultValueHelper::Is(CppForm, TEXT("nullptr")) || FDefaultValueHelper::Is(CppForm, TEXT("0"));
@@ -1416,11 +1400,6 @@ bool FPropertyTraits::DefaultValueStringCppFormatToInnerFormat(const FUnrealProp
 	return PropertyTypeDispatch<DefaultValueStringCppFormatToInnerFormatDispatch, false, bool>(PropDef.GetPropertyBase(), std::ref(PropDef), std::ref(CppForm), std::ref(OutForm));
 }
 
-bool FPropertyTraits::IsObject(EPropertyType PropertyType)
-{
-	return PropertyTypeDispatch<IsObjectDispatch, bool>(PropertyType);
-}
-
 FUnrealPropertyDefinitionInfo& FPropertyTraits::CreateProperty(const FPropertyBase& VarProperty, FUnrealTypeDefinitionInfo& Outer, const FName& Name, EObjectFlags ObjectFlags, EVariableCategory::Type VariableCategory,
 	const TCHAR* Dimensions, FUnrealSourceFile& SourceFile, int LineNumber, int ParsePosition)
 {
diff --git a/Engine/Source/Programs/UnrealHeaderTool/Private/PropertyTypes.h b/Engine/Source/Programs/UnrealHeaderTool/Private/PropertyTypes.h
--- a/Engine/Source/Programs/UnrealHeaderTool/Private/PropertyTypes.h
+++ b/Engine/Source/Programs/UnrealHeaderTool/Private/PropertyTypes.h
@@ -13,11 +13,6 @@ class FUnrealSourceFile;
 
 struct FPropertyTraits
 {
-	/**
-	 * Return true if this property type is an object property
-	 */
-	static bool IsObject(EPropertyType PropertyType);
-
 	/**
 	 * Transforms CPP-formated string containing default value, to inner formated string
 	 * If it cannot be transformed empty string is returned.
diff --git a/Engine/Source/Programs/UnrealHeaderTool/Private/UnrealTypeDefinitionInfo.cpp b/Engine/Source/Programs/UnrealHeaderTool/Private/UnrealTypeDefinitionInfo.cpp
--- a/Engine/Source/Programs/UnrealHeaderTool/Private/UnrealTypeDefinitionInfo.cpp
+++ b/Engine/Source/Programs/UnrealHeaderTool/Private/UnrealTypeDefinitionInfo.cpp
@@ -282,6 +282,14 @@ bool FUnrealPropertyDefinitionInfo::IsOwnedByDynamicType() const
 	return false;
 }
 
+void FUnrealPropertyDefinitionInfo::SetDelegateFunctionSignature(FUnrealFunctionDefinitionInfo& DelegateFunctionDef)
+{
+	FDelegateProperty* DelegateProperty = CastFieldChecked<FDelegateProperty>(PropertyBase.ArrayType == EArrayType::None ? GetProperty() : GetValuePropDef().GetProperty());
+	DelegateProperty->SignatureFunction = DelegateFunctionDef.GetFunction();
+	PropertyBase.FunctionDef = &DelegateFunctionDef;
+}
+
+
 FUnrealPackageDefinitionInfo& FUnrealObjectDefinitionInfo::GetPackageDef() const
 {
 	if (HasSource())
@@ -420,6 +428,19 @@ void FUnrealStructDefinitionInfo::AddProperty(FUnrealPropertyDefinitionInfo& Pro
 {
 	Properties.Add(&PropertyDef);
 
+	// Add the property to the engine object
+	FProperty* Property = PropertyDef.GetProperty();
+
+	// Add to the end of the properties slist
+	FField** Prev = &GetStruct()->ChildProperties;
+	for (; *Prev != nullptr; Prev = &(*Prev)->Next)
+	{
+		// No body
+	}
+	check(*Prev == nullptr);
+	Property->Next = nullptr;
+	*Prev = Property;
+
 	// update the optimization flags
 	if (!bContainsDelegates)
 	{
diff --git a/Engine/Source/Programs/UnrealHeaderTool/Private/UnrealTypeDefinitionInfo.h b/Engine/Source/Programs/UnrealHeaderTool/Private/UnrealTypeDefinitionInfo.h
--- a/Engine/Source/Programs/UnrealHeaderTool/Private/UnrealTypeDefinitionInfo.h
+++ b/Engine/Source/Programs/UnrealHeaderTool/Private/UnrealTypeDefinitionInfo.h
@@ -10,6 +10,8 @@
 #include "BaseParser.h"
 #include "ParserHelper.h"
 
+#define UHT_ENABLE_ENGINE_TYPE_CHECKS 1
+
 // Forward declarations.
 class FScope;
 class FUnrealSourceFile;
@@ -249,6 +251,16 @@ public:
 		return NameCPP;
 	}
 
+	/**
+	 * Returns the name of the property
+	 */
+	virtual FString GetName() const = 0;
+
+	/** Returns the logical name of this object */
+	virtual FName GetFName() const = 0;
+
+	virtual FString GetFullName() const = 0;
+
 	/**
 	 * Return true if this type has source information
 	 */
@@ -343,6 +355,11 @@ public:
 		return Outer;
 	}
 
+	/**
+	 * Return the meta data in map form
+	 */
+	virtual TMap<FName, FString> GenerateMetadataMap() const = 0;
+
 protected:
 	explicit FUnrealTypeDefinitionInfo(FString&& InNameCPP)
 		: NameCPP(MoveTemp(InNameCPP))
@@ -412,6 +429,401 @@ public:
 		Property = InProperty;
 	}
 
+	void ExportCppDeclaration(FOutputDevice& Out, EExportedDeclaration::Type DeclarationType, const TCHAR* ArrayDimOverride = NULL, uint32 AdditionalExportCPPFlags = 0
+		, bool bSkipParameterName = false, const FString* ActualCppType = nullptr, const FString* ActualExtendedType = nullptr, const FString* ActualParameterName = nullptr) const
+	{
+		GetProperty()->ExportCppDeclaration(Out, DeclarationType, ArrayDimOverride, AdditionalExportCPPFlags, bSkipParameterName, ActualCppType, ActualExtendedType, ActualParameterName);
+	}
+
+	FString GetCPPMacroType(FString& ExtendedTypeText) const
+	{
+		return GetProperty()->GetCPPMacroType(ExtendedTypeText);
+	}
+
+	/**
+	 * Returns the name of the property
+	 */
+	virtual FString GetName() const override
+	{
+		return GetProperty()->GetName();
+	}
+
+	/** Returns the logical name of this object */
+	virtual FName GetFName() const override
+	{
+		return GetProperty()->GetFName();
+	}
+
+	virtual FString GetFullName() const override
+	{
+		return GetProperty()->GetFullName();
+	}
+
+	/**
+	 * Returns the fully qualified pathname for this object, in the format:
+	 * 'Outermost.[Outer:]Name'
+	 *
+	 * @param	StopOuter	if specified, indicates that the output string should be relative to this object.  if StopOuter
+	 *						does not exist in this object's Outer chain, the result would be the same as passing NULL.
+	 *
+	 * @note	safe to call on NULL object pointers!
+	 */
+	FString GetPathName(const UObject* StopOuter = nullptr)
+	{
+		return GetProperty()->GetPathName(StopOuter);
+	}
+
+	/**
+	 * Return the engine class name
+	 */
+	FString GetEngineClassName()
+	{
+		return GetProperty()->GetClass()->GetName();
+	}
+
+	/**
+	* Finds the localized display name or native display name as a fallback.
+	*
+	* @return The display name for this object.
+	*/
+	FText GetDisplayNameText() const
+	{
+		return GetProperty()->GetDisplayNameText();
+	}
+
+	/**
+	* Finds the localized tooltip or native tooltip as a fallback.
+	*
+	* @param bShortTooltip Look for a shorter version of the tooltip (falls back to the long tooltip if none was specified)
+	*
+	* @return The tooltip for this object.
+	*/
+	FText GetToolTipText(bool bShortTooltip = false) const
+	{
+		return GetProperty()->GetToolTipText(bShortTooltip);
+	}
+
+	/**
+	 * Returns the C++ name of the property, including the _DEPRECATED suffix if the
+	 * property is deprecated.
+	 *
+	 * @return C++ name of property
+	 */
+	FString GetPropertyNameCPP() const
+	{
+		return GetProperty()->GetNameCPP();
+	}
+
+	/**
+	 * Returns the text to use for exporting this property to header file.
+	 *
+	 * @param	ExtendedTypeText	for property types which use templates, will be filled in with the type
+	 * @param	CPPExportFlags		flags for modifying the behavior of the export
+	 */
+	FString GetCPPType(FString* ExtendedTypeText = NULL, uint32 CPPExportFlags = 0) const
+	{
+		return GetProperty()->GetCPPType(ExtendedTypeText, CPPExportFlags);
+	}
+
+	FString GetCPPTypeForwardDeclaration() const
+	{
+		return GetProperty()->GetCPPTypeForwardDeclaration();
+	}
+
+	/**
+	 * Returns this property's propertyflags
+	 */
+	EPropertyFlags GetPropertyFlags() const
+	{
+#if UHT_ENABLE_ENGINE_TYPE_CHECKS
+		check((GetProperty()->GetPropertyFlags() & ~CPF_ComputedFlags) == PropertyBase.PropertyFlags);
+#endif
+		return PropertyBase.PropertyFlags;
+	}
+
+	void SetPropertyFlags(EPropertyFlags NewFlags)
+	{
+#if UHT_ENABLE_ENGINE_TYPE_CHECKS
+		check((GetProperty()->GetPropertyFlags() & ~CPF_ComputedFlags) == PropertyBase.PropertyFlags);
+#endif
+		GetProperty()->PropertyFlags |= NewFlags;
+		PropertyBase.PropertyFlags |= NewFlags;
+	}
+
+	void ClearPropertyFlags(EPropertyFlags NewFlags)
+	{
+#if UHT_ENABLE_ENGINE_TYPE_CHECKS
+		check((GetProperty()->GetPropertyFlags() & ~CPF_ComputedFlags) == PropertyBase.PropertyFlags);
+#endif
+		GetProperty()->PropertyFlags &= ~NewFlags;
+		PropertyBase.PropertyFlags &= ~NewFlags;
+	}
+
+	/**
+	* 
+	 * Used to safely check whether any of the passed in flags are set. This is required
+	 * as PropertyFlags currently is a 64 bit data type and bool is a 32 bit data type so
+	 * simply using PropertyFlags&CPF_MyFlagBiggerThanMaxInt won't work correctly when
+	 * assigned directly to an bool.
+	 *
+	 * @param FlagsToCheck	Object flags to check for.
+	 *
+	 * @return				true if any of the passed in flags are set, false otherwise  (including no flags passed in).
+	 */
+	bool HasAnyPropertyFlags(uint64 FlagsToCheck) const
+	{
+#if UHT_ENABLE_ENGINE_TYPE_CHECKS
+		check((FlagsToCheck & CPF_ComputedFlags) == 0);
+		check((GetProperty()->GetPropertyFlags() & ~CPF_ComputedFlags) == PropertyBase.PropertyFlags);
+#endif
+		return (PropertyBase.PropertyFlags & FlagsToCheck) != 0;
+	}
+
+	/**
+	 * Used to safely check whether all of the passed in flags are set. This is required
+	 * as PropertyFlags currently is a 64 bit data type and bool is a 32 bit data type so
+	 * simply using PropertyFlags&CPF_MyFlagBiggerThanMaxInt won't work correctly when
+	 * assigned directly to an bool.
+	 *
+	 * @param FlagsToCheck	Object flags to check for
+	 *
+	 * @return true if all of the passed in flags are set (including no flags passed in), false otherwise
+	 */
+	bool HasAllPropertyFlags(uint64 FlagsToCheck) const
+	{
+#if UHT_ENABLE_ENGINE_TYPE_CHECKS
+		check((FlagsToCheck & CPF_ComputedFlags) == 0);
+		check((GetProperty()->GetPropertyFlags() & ~CPF_ComputedFlags) == PropertyBase.PropertyFlags);
+#endif
+		return (PropertyBase.PropertyFlags & FlagsToCheck) == FlagsToCheck;
+	}
+
+	bool HasSpecificPropertyFlags(uint64 FlagsToCheck, uint64 ExpectedFlags)
+	{
+#if UHT_ENABLE_ENGINE_TYPE_CHECKS
+		check((FlagsToCheck & CPF_ComputedFlags) == 0);
+		check((GetProperty()->GetPropertyFlags() & ~CPF_ComputedFlags) == PropertyBase.PropertyFlags);
+#endif
+		return (PropertyBase.PropertyFlags & FlagsToCheck) == ExpectedFlags;
+	}
+
+	/**
+	 * Get the cast flags
+	 */
+	uint64 GetCastFlags() const
+	{
+		return GetProperty()->GetCastFlags();
+	}
+
+	/**
+	* Determines if the property has any metadata associated with the key
+	*
+	* @param Key The key to lookup in the metadata
+	* @return true if there is a (possibly blank) value associated with this key
+	*/
+	bool HasMetaData(const TCHAR* Key) const { return FindMetaData(Key) != nullptr; }
+	bool HasMetaData(const FName& Key) const { return FindMetaData(Key) != nullptr; }
+
+	/**
+	* Find the metadata value associated with the key
+	*
+	* @param Key The key to lookup in the metadata
+	* @return The value associated with the key if it exists, null otherwise
+	*/
+	const FString* FindMetaData(const TCHAR* Key) const
+	{
+		return GetProperty()->FindMetaData(Key);
+	}
+	const FString* FindMetaData(const FName& Key) const
+	{
+		return GetProperty()->FindMetaData(Key);
+	}
+
+	/**
+	* Find the metadata value associated with the key
+	*
+	* @param Key The key to lookup in the metadata
+	* @return The value associated with the key
+	*/
+	const FString& GetMetaData(const TCHAR* Key) const
+	{
+		return GetProperty()->GetMetaData(Key);
+	}
+	const FString& GetMetaData(const FName& Key) const
+	{
+		return GetProperty()->GetMetaData(Key);
+	}
+
+	/**
+	 * Return the meta data in map form
+	 */
+	virtual TMap<FName, FString> GenerateMetadataMap() const override
+	{
+		TMap<FName, FString> MetaDataMap;
+		const TMap<FName, FString>* FieldMetaDataMap = GetProperty()->GetMetaDataMap();
+		if (FieldMetaDataMap)
+		{
+			MetaDataMap = *FieldMetaDataMap;
+		}
+		return MetaDataMap;
+	}
+
+	/**
+	 * Editor-only properties are those that only are used with the editor is present or cannot be removed from serialisation.
+	 * Editor-only properties include: EditorOnly properties
+	 * Properties that cannot be removed from serialisation are:
+	 *		Boolean properties (may affect GCC_BITFIELD_MAGIC computation)
+	 *		Native properties (native serialisation)
+	 */
+	bool IsEditorOnlyProperty() const
+	{
+		return GetProperty()->IsEditorOnlyProperty();
+	}
+
+	/**
+	 * Returns true if this property, or in the case of e.g. array or struct properties any sub- property, contains a
+	 * UObject reference that is marked CPF_NeedCtorLink (i.e. instanced keyword).
+	 *
+	 * @return true if property (or sub- properties) contain a FObjectProperty that is marked CPF_NeedCtorLink, false otherwise
+	 */
+	bool ContainsInstancedObjectProperty() const
+	{
+		//return (PropertyFlags & (CPF_ContainsInstancedReference | CPF_InstancedReference)) != 0;
+		return GetProperty()->ContainsInstancedObjectProperty();
+	}
+
+	/**
+	 * Test to see if two properties share the same types.  This method utilizes the underlying UProperty SameType implementation.
+	 */
+	bool SameType(const FUnrealPropertyDefinitionInfo& Other) const
+	{
+		return GetProperty()->SameType(Other.GetProperty());
+	}
+
+	/**
+	 * Determines whether this property's type is compatible with another property's type.
+	 *
+	 * @param	Other							the property to check against this one.
+	 *											Given the following example expressions, VarA is Other and VarB is 'this':
+	 *												VarA = VarB;
+	 *
+	 *												function func(type VarB) {}
+	 *												func(VarA);
+	 *
+	 *												static operator==(type VarB_1, type VarB_2) {}
+	 *												if ( VarA_1 == VarA_2 ) {}
+	 *
+	 * @param	bDisallowGeneralization			controls whether it should be considered a match if this property's type is a generalization
+	 *											of the other property's type (or vice versa, when dealing with structs
+	 * @param	bIgnoreImplementedInterfaces	controls whether two types can be considered a match if one type is an interface implemented
+	 *											by the other type.
+	 */
+	bool MatchesType(const FUnrealPropertyDefinitionInfo& Other, bool bDisallowGeneralization, bool bIgnoreImplementedInterfaces = false) const
+	{
+		return GetPropertyBase().MatchesType(Other.GetPropertyBase(), bDisallowGeneralization, bIgnoreImplementedInterfaces);
+	}
+
+	/**
+	 * Return true if the property is a primitive type or an primitive type array
+	 */
+	bool IsPrimitiveOrPrimitiveStaticArray() const
+	{
+		return PropertyBase.IsPrimitiveOrPrimitiveStaticArray();
+	}
+
+	/**
+	 * Return true if the property is a static array
+	 */
+	bool IsStaticArray() const
+	{
+		return PropertyBase.ArrayType == EArrayType::Static;
+	}
+
+	/**
+	 * Return true if the property is a dynamic array
+	 */
+	bool IsDynamicArray() const
+	{
+		return PropertyBase.ArrayType == EArrayType::Dynamic;
+	}
+
+	/**
+	 * Return true if the property is a map
+	 */
+	bool IsMap() const
+	{
+		return PropertyBase.ArrayType == EArrayType::None && PropertyBase.MapKeyProp.IsValid();
+	}
+
+	/**
+	 * Return true if the property is a set
+	 */
+	bool IsSet() const
+	{
+		return PropertyBase.ArrayType == EArrayType::Set;
+	}
+
+	/**
+	 * Return true if the property is a boolean or boolean static array
+	 */
+	bool IsBooleanOrBooleanStaticArray() const
+	{
+		return PropertyBase.IsBooleanOrBooleanStaticArray();
+	}
+
+	/**
+	 * Return true if the property is an object or object static array
+	 */
+	bool IsObjectRefOrObjectRefStaticArray() const
+	{
+		return PropertyBase.IsObjectRefOrObjectRefStaticArray();
+	}
+
+	/**
+	 * Return true if the property is a class or class static array 
+	 * NOTE: This is specific to Object and ObjectPtr where the class derrives from UClass
+	 */
+	bool IsClassRefOrClassRefStaticArray() const
+	{
+		return PropertyBase.IsClassRefOrClassRefStaticArray();
+	}
+
+	/**
+	 * Return true if this is a byte enumeration
+	 */
+	bool IsByteEnumOrByteEnumStaticArray() const
+	{
+		return PropertyBase.IsByteEnumOrByteEnumStaticArray();
+	}
+
+	/** 
+	 * Return true if this is a numeric or a numeric static array
+	 */
+	bool IsNumericOrNumericStaticArray() const
+	{
+		return PropertyBase.IsNumericOrNumericStaticArray();
+	}
+
+	/**
+	 * Return true if this is a struct or a struct static array
+	 */
+	bool IsStructOrStructStaticArray() const
+	{
+		return PropertyBase.IsStructOrStructStaticArray();
+	}
+
+	/**
+	 * Return true if this is a structure and it has no op constructor
+	 */
+	bool HasNoOpConstructor() const
+	{
+		if (FStructProperty* StructProperty = CastField<FStructProperty>(GetProperty()))
+		{
+			return StructProperty->HasNoOpConstructor();
+		}
+		return false;
+	}
+
 	/**
 	 * Set the string that represents the array dimensions
 	 */
@@ -478,6 +890,14 @@ public:
 	 */
 	virtual void AddMetaData(TMap<FName, FString>&& InMetaData) override;
 
+	/**
+	 * Return true if we have a key property def
+	 */
+	bool HasKeyPropDef() const
+	{
+		return KeyPropDef != nullptr;
+	}
+
 	/**
 	 * Get the associated key property definition (valid for maps)
 	 */
@@ -487,6 +907,14 @@ public:
 		return *KeyPropDef;
 	}
 
+	/**
+	 * Return true if we have a value property def
+	 */
+	bool HasValuePropDef() const
+	{
+		return ValuePropDef != nullptr;
+	}
+
 	/**
 	 * Sets the associated key property definition (valid for maps)
 	 */
@@ -522,29 +950,6 @@ public:
 		return ParsePosition;
 	}
 
-	/**
-	 * Determines whether this property's type is compatible with another property's type.
-	 *
-	 * @param	Other							the property to check against this one.
-	 *											Given the following example expressions, VarA is Other and VarB is 'this':
-	 *												VarA = VarB;
-	 *
-	 *												function func(type VarB) {}
-	 *												func(VarA);
-	 *
-	 *												static operator==(type VarB_1, type VarB_2) {}
-	 *												if ( VarA_1 == VarA_2 ) {}
-	 *
-	 * @param	bDisallowGeneralization			controls whether it should be considered a match if this property's type is a generalization
-	 *											of the other property's type (or vice versa, when dealing with structs
-	 * @param	bIgnoreImplementedInterfaces	controls whether two types can be considered a match if one type is an interface implemented
-	 *											by the other type.
-	 */
-	bool MatchesType(const FUnrealPropertyDefinitionInfo& Other, bool bDisallowGeneralization, bool bIgnoreImplementedInterfaces = false) const
-	{
-		return GetPropertyBase().MatchesType(Other.GetPropertyBase(), bDisallowGeneralization, bIgnoreImplementedInterfaces);
-	}
-
 	/**
 	 * Return the type package name
 	 */
@@ -563,6 +968,28 @@ public:
 	 */
 	virtual bool IsOwnedByDynamicType() const override;
 
+	/**
+	 * Return the replication notification function name
+	 */
+	FName GetRepNotifyFunc() const
+	{
+		return GetProperty()->RepNotifyFunc;
+	}
+
+	/**
+	 * Set the replication notification function name
+	 */
+	void SetRepNotifyFunc(FName InRepNotifyFunc)
+	{
+		GetProperty()->RepNotifyFunc = InRepNotifyFunc;
+		PropertyBase.RepNotifyName = InRepNotifyFunc;
+	}
+
+	/**
+	 * Set the delegate function signature
+	 */
+	void SetDelegateFunctionSignature(FUnrealFunctionDefinitionInfo& DelegateFunctionDef);
+
 private:
 	FPropertyBase PropertyBase;
 	FString ArrayDimensions;
@@ -610,17 +1037,22 @@ public:
 	/**
 	 * Returns the name of this object (with no path information)
 	 */
-	FString GetName() const
+	virtual FString GetName() const override
 	{
 		return GetObject()->GetName();
 	}
 
 	/** Returns the logical name of this object */
-	FName GetFName() const
+	virtual FName GetFName() const override
 	{
 		return GetObject()->GetFName();
 	}
 
+	virtual FString GetFullName() const override
+	{
+		return GetObject()->GetFullName();
+	}
+
 	/**
 	 * Returns the fully qualified pathname for this object, in the format:
 	 * 'Outermost.[Outer:]Name'
@@ -648,6 +1080,34 @@ public:
 		return static_cast<FUnrealObjectDefinitionInfo*>(FUnrealTypeDefinitionInfo::GetOuter());
 	}
 
+	/**
+	 * Return the meta data in map form
+	 */
+	virtual TMap<FName, FString> GenerateMetadataMap() const override
+	{
+		UObject* Obj = GetObject();
+		check(Obj);
+		UPackage* Package = Obj->GetPackage();
+		check(Package);
+		UMetaData* Metadata = Package->GetMetaData();
+		check(Metadata);
+
+		TMap<FName, FString>* PackageMap = Metadata->ObjectMetaDataMap.Find(Obj);
+		TMap<FName, FString> Map;
+		if (PackageMap)
+		{
+			for (const TPair<FName, FString>& MetaKeyValue : *PackageMap)
+			{
+				FString Key = MetaKeyValue.Key.ToString();
+				if (!Key.StartsWith(TEXT("/Script")))
+				{
+					Map.Add(MetaKeyValue.Key, MetaKeyValue.Value);
+				}
+			}
+		}
+		return Map;
+	}
+
 	/**
 	 * Helper methods to remove dependencies on engine types
 	 */
@@ -1489,7 +1949,7 @@ public:
 	 * @param FlagsToCheck	Function flags to check for
 	 * @return true if all of the passed in flags are set (including no flags passed in), false otherwise
 	 */
-	FORCEINLINE bool HasAllFunctionFlags(EStructFlags FlagsToCheck) const
+	FORCEINLINE bool HasAllStructFlags(EStructFlags FlagsToCheck) const
 	{
 		return EnumHasAllFlags(GetScriptStruct()->StructFlags, FlagsToCheck);
 	}
@@ -1501,7 +1961,7 @@ public:
 	 * @param ExpectedFlags The flags from the flags to check that should be set
 	 * @return true if specific of the passed in flags are set (including no flags passed in), false otherwise
 	 */
-	FORCEINLINE bool HasSpecificFunctionFlags(EStructFlags FlagsToCheck, EStructFlags ExpectedFlags) const
+	FORCEINLINE bool HasSpecificStructFlags(EStructFlags FlagsToCheck, EStructFlags ExpectedFlags) const
 	{
 		EStructFlags Flags = GetScriptStruct()->StructFlags;
 		return (Flags & FlagsToCheck) == ExpectedFlags;
@@ -1609,11 +2069,6 @@ public:
 		return (Flags & FlagsToCheck) == ExpectedFlags;
 	}
 
-	/**
-	 * Return the outer in Class form.  
-	 */
-	FUnrealClassDefinitionInfo* GetOuterClass() const;
-
 	/** @name getters */
 	//@{
 	FFuncInfo& GetFunctionData() { return FunctionData; }
@@ -1817,7 +2272,6 @@ public:
 	 */
 	bool HasAnyClassFlags(EClassFlags FlagsToCheck) const
 	{
-
 		return EnumHasAnyFlags(ParsedClassFlags, FlagsToCheck) || GetClass()->HasAnyClassFlags(FlagsToCheck);
 	}
 
@@ -2259,8 +2713,3 @@ struct TUHTFieldRange
 
 	TUHTFieldIterator<T> Begin;
 };
-
-inline FUnrealClassDefinitionInfo* FUnrealFunctionDefinitionInfo::GetOuterClass() const
-{
-	return UHTCast<FUnrealClassDefinitionInfo>(GetOuter());
-}
