Fixing Mac build where a block of diagnostic code slipped in.

#rnx

diff --git a/Engine/Source/Programs/UnrealHeaderTool/Private/CodeGenerator.cpp b/Engine/Source/Programs/UnrealHeaderTool/Private/CodeGenerator.cpp
--- a/Engine/Source/Programs/UnrealHeaderTool/Private/CodeGenerator.cpp
+++ b/Engine/Source/Programs/UnrealHeaderTool/Private/CodeGenerator.cpp
@@ -4920,11 +4920,6 @@ FString FNativeClassHeaderGenerator::GetFunctionParameterString(FUnrealFunctionD
 	FString ParameterList;
 	FUHTStringBuilder PropertyText;
 
-	if (FunctionDef.GetReturn() != nullptr)
-	{
-		ParameterList = ParameterList;
-	}
-
 	for (FUnrealPropertyDefinitionInfo* PropertyDef : FunctionDef.GetProperties())
 	{
 		OutReferenceGatherers.ForwardDeclarations.Add(PropertyDef->GetCPPTypeForwardDeclaration());
