Implement a common meta data container to be shared between all types.
Moving more property functionality into PropertyTypes
Assorted cleanup of other methods.

#rnx

diff --git a/Engine/Source/Programs/UnrealHeaderTool/Private/CodeGenerator.cpp b/Engine/Source/Programs/UnrealHeaderTool/Private/CodeGenerator.cpp
--- a/Engine/Source/Programs/UnrealHeaderTool/Private/CodeGenerator.cpp
+++ b/Engine/Source/Programs/UnrealHeaderTool/Private/CodeGenerator.cpp
@@ -797,7 +797,7 @@ void FNativeClassHeaderGenerator::GetPropertyTag(FUHTStringBuilder& Out, FUnreal
 void FNativeClassHeaderGenerator::OutputProperty(FOutputDevice& DeclOut, FOutputDevice& Out, FReferenceGatherers& OutReferenceGatherers, const TCHAR* Scope, TArray<FPropertyNamePointerPair>& PropertyNamesAndPointers, FUnrealPropertyDefinitionInfo& PropertyDef, const TCHAR* OffsetStr, FString&& Name, const TCHAR* DeclSpaces, const TCHAR* Spaces, const TCHAR* SourceStruct) const
 {
 	const FPropertyBase& PropertyBase   = PropertyDef.GetPropertyBase();
-	FString        PropName             = CreateUTF8LiteralString(FNativeClassHeaderGenerator::GetOverriddenName(&PropertyDef));
+	FString        PropName             = CreateUTF8LiteralString(FNativeClassHeaderGenerator::GetOverriddenName(PropertyDef));
 	FString        PropNameDep          = PropertyDef.HasAllPropertyFlags(CPF_Deprecated) ? PropertyDef.GetName() + TEXT("_DEPRECATED") : PropertyDef.GetName();
 	const TCHAR*   FPropertyObjectFlags = PropertyDef.IsOwnedByDynamicType() ? TEXT("RF_Public|RF_Transient") : TEXT("RF_Public|RF_Transient|RF_MarkAsNative");
 	EPropertyFlags PropFlags            = PropertyDef.GetPropertyFlags() & ~CPF_ComputedFlags;
@@ -1914,7 +1914,7 @@ void FNativeClassHeaderGenerator::ExportNativeGeneratedInitCode(FOutputDevice& O
 	bool bAllEditorOnlyFunctions = true;
 	for (FUnrealFunctionDefinitionInfo* LocalFuncDef : ClassDef.GetFunctions())
 	{
-		FName TrueName = FNativeClassHeaderGenerator::GetOverriddenFName(LocalFuncDef);
+		FName TrueName = FNativeClassHeaderGenerator::GetOverriddenFName(*LocalFuncDef);
 		bool bAlreadyIncluded = false;
 		AlreadyIncludedNames.Add(TrueName, &bAlreadyIncluded);
 		if (bAlreadyIncluded)
@@ -2157,7 +2157,7 @@ void FNativeClassHeaderGenerator::ExportNativeGeneratedInitCode(FOutputDevice& O
 		{
 			const FString& DynamicClassPackageName = ClassDef.GetTypePackageName();
 			GeneratedClassRegisterFunctionText.Logf(TEXT("\t\tUPackage* OuterPackage = FindOrConstructDynamicTypePackage(TEXT(\"%s\"));\r\n"), *DynamicClassPackageName);
-			GeneratedClassRegisterFunctionText.Logf(TEXT("\t\tUClass* OuterClass = Cast<UClass>(StaticFindObjectFast(UClass::StaticClass(), OuterPackage, TEXT(\"%s\")));\r\n"), *FNativeClassHeaderGenerator::GetOverriddenName(&ClassDef));
+			GeneratedClassRegisterFunctionText.Logf(TEXT("\t\tUClass* OuterClass = Cast<UClass>(StaticFindObjectFast(UClass::StaticClass(), OuterPackage, TEXT(\"%s\")));\r\n"), *FNativeClassHeaderGenerator::GetOverriddenName(ClassDef));
 			GeneratedClassRegisterFunctionText.Logf(TEXT("\t\tif (!OuterClass || !(OuterClass->ClassFlags & CLASS_Constructed))\r\n"));
 		}
 
@@ -2199,7 +2199,7 @@ void FNativeClassHeaderGenerator::ExportNativeGeneratedInitCode(FOutputDevice& O
 	FString SingletonName = ClassDef.GetSingletonName(true);
 	SingletonName.ReplaceInline(TEXT("()"), TEXT(""), ESearchCase::CaseSensitive); // function address
 
-	FString OverriddenClassName = *FNativeClassHeaderGenerator::GetOverriddenName(&ClassDef);
+	FString OverriddenClassName = *FNativeClassHeaderGenerator::GetOverriddenName(ClassDef);
 
 	const FString EmptyString;
 	const FString& InitSearchableValuesFunctionName = bIsDynamic ? ClassDef.GetMetaData(NAME_InitializeStaticSearchableValues) : EmptyString;
@@ -2227,7 +2227,7 @@ void FNativeClassHeaderGenerator::ExportNativeGeneratedInitCode(FOutputDevice& O
 			GeneratedClassRegisterFunctionText.Logf(TEXT("%s\r\n"), *SparseScriptStructDef->GetName());
 			for (FUnrealPropertyDefinitionInfo* ChildDef : TUHTFieldRange<FUnrealPropertyDefinitionInfo>(*SparseScriptStructDef))
 			{
-				GeneratedClassRegisterFunctionText.Logf(TEXT("%s %s\r\n"), *ChildDef->GetCPPType(), *ChildDef->GetPropertyNameCPP());
+				GeneratedClassRegisterFunctionText.Logf(TEXT("%s %s\r\n"), *ChildDef->GetCPPType(), *ChildDef->GetNameWithDeprecated());
 			}
 		}
 	}
@@ -2402,7 +2402,7 @@ void FNativeClassHeaderGenerator::ExportFunction(FOutputDevice& Out, FReferenceG
 		*StaticsStructName,
 		*OuterFunc,
 		*GetSingletonNameFuncAddr(SuperFunctionDef, OutReferenceGatherers.UniqueCrossModuleReferences),
-		*CreateUTF8LiteralString(FNativeClassHeaderGenerator::GetOverriddenName(&FunctionDef)),
+		*CreateUTF8LiteralString(FNativeClassHeaderGenerator::GetOverriddenName(FunctionDef)),
 		(SparseDelegateFunction ? *CreateUTF8LiteralString(SparseDelegateFunction->OwningClassName.ToString()) : TEXT("nullptr")),
 		(SparseDelegateFunction ? *CreateUTF8LiteralString(SparseDelegateFunction->DelegateName.ToString()) : TEXT("nullptr")),
 		*StructureSize,
@@ -2526,7 +2526,7 @@ void FNativeClassHeaderGenerator::ExportInterfaceCallFunctions(FOutputDevice& Ou
 		Out.Logf( TEXT(";") LINE_TERMINATOR );
 
 		FString FunctionNameName = FString::Printf(TEXT("NAME_%s_%s"), *UHTCastChecked<FUnrealStructDefinitionInfo>(FunctionDef->GetOuter()).GetAlternateNameCPP(), *FunctionName);
-		OutCpp.Logf(TEXT("\tstatic FName %s = FName(TEXT(\"%s\"));") LINE_TERMINATOR, *FunctionNameName, *GetOverriddenFName(FunctionDef).ToString());
+		OutCpp.Logf(TEXT("\tstatic FName %s = FName(TEXT(\"%s\"));") LINE_TERMINATOR, *FunctionNameName, *GetOverriddenFName(*FunctionDef).ToString());
 
 		ExportNativeFunctionHeader(OutCpp, OutReferenceGatherers.ForwardDeclarations, *FunctionDef, FunctionData, EExportFunctionType::Interface, EExportFunctionHeaderStyle::Definition, *ExtraParam, *APIString);
 		OutCpp.Logf( LINE_TERMINATOR TEXT("\t{") LINE_TERMINATOR );
@@ -3320,7 +3320,7 @@ void FNativeClassHeaderGenerator::ExportGeneratedStructBodyMacros(FOutputDevice&
 	ScriptStructDef.AddCrossModuleReference(OutReferenceGatherers.UniqueCrossModuleReferences, true);
 
 	const bool bIsDynamic = ScriptStructDef.IsDynamic();
-	const FString ActualStructName = FNativeClassHeaderGenerator::GetOverriddenName(&ScriptStructDef);
+	const FString ActualStructName = FNativeClassHeaderGenerator::GetOverriddenName(ScriptStructDef);
 	const FString& FriendApiString  = GetAPIString();
 
 	FUnrealStructDefinitionInfo* BaseStructDef = ScriptStructDef.GetSuperStruct();
@@ -3846,7 +3846,7 @@ void FNativeClassHeaderGenerator::ExportGeneratedEnumInitCode(FOutputDevice& Out
 	const bool bIsDynamic               = EnumDef.IsDynamic();
 	const FString& SingletonName        = EnumDef.GetSingletonNameChopped(true);
 	const FString EnumNameCpp           = EnumDef.GetName(); //UserDefinedEnum should already have a valid cpp name.
-	const FString OverriddenEnumNameCpp = FNativeClassHeaderGenerator::GetOverriddenName(&EnumDef);
+	const FString OverriddenEnumNameCpp = FNativeClassHeaderGenerator::GetOverriddenName(EnumDef);
 	const FString StaticsStructName     = SingletonName + TEXT("_Statics");
 	const bool bIsEditorOnlyDataType    = EnumDef.IsEditorOnly();
 
@@ -5270,7 +5270,7 @@ void FNativeClassHeaderGenerator::ExportCallbackFunctions(
 		if (ExportCallbackType != EExportCallbackType::Interface)
 		{
 			FunctionNameName = FString::Printf(TEXT("NAME_%s_%s"), *ClassName, *FunctionName);
-			OutCpp.Logf(TEXT("\tstatic FName %s = FName(TEXT(\"%s\"));") LINE_TERMINATOR, *FunctionNameName, *GetOverriddenFName(FunctionDef).ToString());
+			OutCpp.Logf(TEXT("\tstatic FName %s = FName(TEXT(\"%s\"));") LINE_TERMINATOR, *FunctionNameName, *GetOverriddenFName(*FunctionDef).ToString());
 		}
 
 		// Emit the thunk implementation
diff --git a/Engine/Source/Programs/UnrealHeaderTool/Private/HeaderParser.cpp b/Engine/Source/Programs/UnrealHeaderTool/Private/HeaderParser.cpp
--- a/Engine/Source/Programs/UnrealHeaderTool/Private/HeaderParser.cpp
+++ b/Engine/Source/Programs/UnrealHeaderTool/Private/HeaderParser.cpp
@@ -616,7 +616,7 @@ namespace
 		}
 	}
 
-	const TCHAR* GetHintText(EVariableCategory::Type VariableCategory)
+	const TCHAR* GetHintText(EVariableCategory VariableCategory)
 	{
 		switch (VariableCategory)
 		{
@@ -639,198 +639,6 @@ namespace
 		return nullptr;
 	}
 
-	/**
-	 * Ensures at script compile time that the metadata formatting is correct
-	 * @param	InKey			the metadata key being added
-	 * @param	InValue			the value string that will be associated with the InKey
-	 */
-	void ValidateMetaDataFormat(FUnrealTypeDefinitionInfo& TypeDef, const FName InKey, const FString& InValue)
-	{
-		switch (GetCheckedMetadataSpecifier(InKey))
-		{
-			default:
-			{
-				// Don't need to validate this specifier
-			}
-			break;
-
-			case ECheckedMetadataSpecifier::UIMin:
-			case ECheckedMetadataSpecifier::UIMax:
-			case ECheckedMetadataSpecifier::ClampMin:
-			case ECheckedMetadataSpecifier::ClampMax:
-			{
-				if (!InValue.IsNumeric())
-				{
-					FError::Throwf(TEXT("Metadata value for '%s' is non-numeric : '%s'"), *InKey.ToString(), *InValue);
-				}
-			}
-			break;
-
-			case ECheckedMetadataSpecifier::BlueprintProtected:
-			{
-				if (FUnrealFunctionDefinitionInfo* FuncDef = TypeDef.AsFunction())
-				{
-					UFunction* Function = FuncDef->GetFunction();
-					if (FuncDef->HasAnyFunctionFlags(FUNC_Static))
-					{
-						// Determine if it's a function library
-						FUnrealClassDefinitionInfo* ClassDef = FuncDef->GetOwnerClass();
-						for (; ClassDef != nullptr && ClassDef->GetSuperClass() != GUObjectDef; ClassDef = ClassDef->GetSuperClass())
-						{
-						}
-
-						if (ClassDef != nullptr && ClassDef->GetName() == TEXT("BlueprintFunctionLibrary"))
-						{
-							FError::Throwf(TEXT("%s doesn't make sense on static method '%s' in a blueprint function library"), *InKey.ToString(), *FuncDef->GetName());
-						}
-					}
-				}
-			}
-			break;
-
-			case ECheckedMetadataSpecifier::CommutativeAssociativeBinaryOperator:
-			{
-				if (FUnrealFunctionDefinitionInfo* FuncDef = TypeDef.AsFunction())
-				{
-					bool bGoodParams = (FuncDef->GetProperties().Num() == 3);
-					if (bGoodParams)
-					{
-						FUnrealPropertyDefinitionInfo* FirstParam = nullptr;
-						FUnrealPropertyDefinitionInfo* SecondParam = nullptr;
-						FUnrealPropertyDefinitionInfo* ReturnValue = nullptr;
-						for (FUnrealPropertyDefinitionInfo* PropertyDef : FuncDef->GetProperties())
-						{
-							if (PropertyDef->HasAnyPropertyFlags(CPF_ReturnParm))
-							{
-								ReturnValue = PropertyDef;
-							}
-							else
-							{
-								if (FirstParam == nullptr)
-								{
-									FirstParam = PropertyDef;
-								}
-								else if (SecondParam == nullptr)
-								{
-									SecondParam = PropertyDef;
-								}
-							}
-						}
-
-						if (ReturnValue == nullptr || SecondParam == nullptr || !SecondParam->SameType(*FirstParam))
-						{
-							bGoodParams = false;
-						}
-					}
-
-					if (!bGoodParams)
-					{
-						UE_LOG_ERROR_UHT(TEXT("Commutative associative binary operators must have exactly 2 parameters of the same type and a return value."));
-					}
-				}
-			}
-			break;
-
-			case ECheckedMetadataSpecifier::ExpandBoolAsExecs:
-			case ECheckedMetadataSpecifier::ExpandEnumAsExecs:
-			{
-				if (FUnrealFunctionDefinitionInfo* FuncDef = TypeDef.AsFunction())
-				{
-					// multiple entry parsing in the same format as eg SetParam.
-					TArray<FString> RawGroupings;
-					InValue.ParseIntoArray(RawGroupings, TEXT(","), false);
-
-					FUnrealPropertyDefinitionInfo* FirstInputDef = nullptr;
-					for (const FString& RawGroup : RawGroupings)
-					{
-						TArray<FString> IndividualEntries;
-						RawGroup.ParseIntoArray(IndividualEntries, TEXT("|"));
-
-						for (const FString& Entry : IndividualEntries)
-						{
-							if (Entry.IsEmpty())
-							{
-								continue;
-							}
-							
-							FUnrealPropertyDefinitionInfo* FoundFieldDef = FHeaderParser::FindProperty(*FuncDef, *Entry, false);
-							if (!FoundFieldDef)
-							{
-								UE_LOG_ERROR_UHT(TEXT("Function does not have a parameter named '%s'"), *Entry);
-							}
-							else
-							{
-								if (!FoundFieldDef->HasAnyPropertyFlags(CPF_ReturnParm) &&
-
-								    (!FoundFieldDef->HasAnyPropertyFlags(CPF_OutParm) ||
-										FoundFieldDef->HasAnyPropertyFlags(CPF_ReferenceParm)))
-								{
-									if (!FirstInputDef)
-									{
-										FirstInputDef = FoundFieldDef;
-									}
-									else
-									{
-										UE_LOG_ERROR_UHT(TEXT("Function already specified an ExpandEnumAsExec input (%s), but '%s' is also an input parameter. Only one is permitted."), *FirstInputDef->GetName(), *Entry);
-									}
-								}
-							}
-						}
-					}
-				}
-			}
-			break;
-
-			case ECheckedMetadataSpecifier::DevelopmentStatus:
-			{
-				const FString EarlyAccessValue(TEXT("EarlyAccess"));
-				const FString ExperimentalValue(TEXT("Experimental"));
-				if ((InValue != EarlyAccessValue) && (InValue != ExperimentalValue))
-				{
-					FError::Throwf(TEXT("'%s' metadata was '%s' but it must be %s or %s"), *InKey.ToString(), *InValue, *ExperimentalValue, *EarlyAccessValue);
-				}
-			}
-			break;
-
-			case ECheckedMetadataSpecifier::Units:
-			{
-				// Check for numeric property
-				if (FUnrealPropertyDefinitionInfo* PropDef = TypeDef.AsProperty())
-				{
-					if (!PropDef->IsNumericOrNumericStaticArray() && !PropDef->IsStructOrStructStaticArray())
-					{
-						FError::Throwf(TEXT("'Units' meta data can only be applied to numeric and struct properties"));
-					}
-				}
-
-				if (!FUnitConversion::UnitFromString(*InValue))
-				{
-					FError::Throwf(TEXT("Unrecognized units (%s) specified for property '%s'"), *InValue, *TypeDef.GetFullName());
-				}
-			}
-			break;
-
-			case ECheckedMetadataSpecifier::DocumentationPolicy:
-			{
-				const TCHAR* StrictValue = TEXT("Strict");
-				if (InValue != StrictValue)
-				{
-					FError::Throwf(TEXT("'%s' metadata was '%s' but it must be %s"), *InKey.ToString(), *InValue, *StrictValue);
-				}
-			}
-			break;
-		}
-	}
-
-	// Ensures at script compile time that the metadata formatting is correct
-	void ValidateMetaDataFormat(FUnrealTypeDefinitionInfo& TypeDef, const TMap<FName, FString>& MetaData)
-	{
-		for (const TPair<FName, FString>& Pair : MetaData)
-		{
-			ValidateMetaDataFormat(TypeDef, Pair.Key, Pair.Value);
-		}
-	}
-
 	void SkipAlignasIfNecessary(FBaseParser& Parser)
 	{
 		if (Parser.MatchIdentifier(TEXT("alignas"), ESearchCase::CaseSensitive))
@@ -920,27 +728,6 @@ void AddEditInlineMetaData(TMap<FName, FString>& MetaData)
 	MetaData.Add(NAME_EditInline, TEXT("true"));
 }
 
-// Validates the metadata, then adds it to the class data
-void AddMetaDataToClassData(FUnrealTypeDefinitionInfo& TypeDef, TMap<FName, FString>&& InMetaData)
-{
-	// Evaluate any key redirects on the passed in pairs
-	for (TPair<FName, FString>& Pair : InMetaData)
-	{
-		FName& CurrentKey = Pair.Key;
-		FName NewKey = UMetaData::GetRemappedKeyName(CurrentKey);
-
-		if (NewKey != NAME_None)
-		{
-			UE_LOG_WARNING_UHT(TEXT("Remapping old metadata key '%s' to new key '%s', please update the declaration."), *CurrentKey.ToString(), *NewKey.ToString());
-			CurrentKey = NewKey;
-		}
-	}
-
-	// Finish validating and associate the metadata with the field
-	ValidateMetaDataFormat(TypeDef, InMetaData);
-	TypeDef.AddMetaData(MoveTemp(InMetaData));
-}
-
 FUHTConfig::FUHTConfig()
 {
 	// Read Ini options, GConfig must exist by this point
@@ -1062,7 +849,7 @@ FUnrealClassDefinitionInfo* FHeaderParser::GetQualifiedClass(const TCHAR* Thing)
  */
 FUnrealFunctionDefinitionInfo* FHeaderParser::FindFunction
 (
-	FUnrealStructDefinitionInfo&  InScope,
+	const FUnrealStructDefinitionInfo&  InScope,
 	const TCHAR*	InIdentifier,
 	bool			bIncludeParents,
 	const TCHAR*	Thing
@@ -1072,7 +859,7 @@ FUnrealFunctionDefinitionInfo* FHeaderParser::FindFunction
 	FName InName(InIdentifier, FNAME_Find);
 	if (InName != NAME_None)
 	{
-		for (FUnrealStructDefinitionInfo* Scope = &InScope; Scope; Scope = UHTCast<FUnrealStructDefinitionInfo>(Scope->GetOuter()))
+		for (const FUnrealStructDefinitionInfo* Scope = &InScope; Scope; Scope = UHTCast<FUnrealStructDefinitionInfo>(Scope->GetOuter()))
 		{
 			for (FUnrealPropertyDefinitionInfo* PropertyDef : TUHTFieldRange<FUnrealPropertyDefinitionInfo>(*Scope))
 			{
@@ -1103,13 +890,13 @@ FUnrealFunctionDefinitionInfo* FHeaderParser::FindFunction
 	return nullptr;
 }
 
-FUnrealPropertyDefinitionInfo* FHeaderParser::FindProperty(FUnrealStructDefinitionInfo& InScope, const TCHAR* InIdentifier, bool bIncludeParents, const TCHAR* Thing)
+FUnrealPropertyDefinitionInfo* FHeaderParser::FindProperty(const FUnrealStructDefinitionInfo& InScope, const TCHAR* InIdentifier, bool bIncludeParents, const TCHAR* Thing)
 {
 	check(InIdentifier);
 	FName InName(InIdentifier, FNAME_Find);
 	if (InName != NAME_None)
 	{
-		for (FUnrealStructDefinitionInfo* Scope = &InScope; Scope; Scope = UHTCast<FUnrealStructDefinitionInfo>(Scope->GetOuter()))
+		for (const FUnrealStructDefinitionInfo* Scope = &InScope; Scope; Scope = UHTCast<FUnrealStructDefinitionInfo>(Scope->GetOuter()))
 		{
 			for (FUnrealFunctionDefinitionInfo* FunctionDef : TUHTFieldRange<FUnrealFunctionDefinitionInfo>(*Scope))
 			{
@@ -1273,7 +1060,7 @@ FUnrealEnumDefinitionInfo& FHeaderParser::CompileEnum()
 	}
 
 	// Validate the metadata for the enum
-	ValidateMetaDataFormat(EnumDef, EnumMetaData);
+	EnumDef.ValidateMetaDataFormat(EnumMetaData);
 
 	// Read base for enum class
 	EUnderlyingEnumType UnderlyingType = EUnderlyingEnumType::uint8;
@@ -2147,7 +1934,7 @@ FUnrealScriptStructDefinitionInfo& FHeaderParser::CompileStructDeclaration()
 	AddFormattedPrevCommentAsTooltipMetaData(MetaData);
 
 	// Register the metadata
-	AddMetaDataToClassData(StructDef, MoveTemp(MetaData));
+	FUHTMetaData::RemapAndAddMetaData(StructDef, MoveTemp(MetaData));
 
 	// Get opening brace.
 	RequireSymbol( TEXT('{'), TEXT("'struct'") );
@@ -2899,7 +2686,7 @@ void FHeaderParser::VerifyPropertyMarkups(FUnrealClassDefinitionInfo& TargetClas
 				// Since the function map is not valid yet, we have to iterate over the fields to look for the function
 				for (FUnrealFunctionDefinitionInfo* FunctionDef : SearchClassDef->GetFunctions())
 				{
-					if (FNativeClassHeaderGenerator::GetOverriddenFName(FunctionDef) == FuncName)
+					if (FNativeClassHeaderGenerator::GetOverriddenFName(*FunctionDef) == FuncName)
 					{
 						return FunctionDef;
 					}
@@ -3108,7 +2895,7 @@ void FHeaderParser::GetVarType(
 	const FToken*                   OuterPropertyType,
 	const EPropertyFlags*			OuterPropertyFlags,
 	EPropertyDeclarationStyle::Type PropertyDeclarationStyle,
-	EVariableCategory::Type         VariableCategory,
+	EVariableCategory               VariableCategory,
 	FIndexRange*                    ParsedVarIndexRange,
 	ELayoutMacroType*               OutLayoutMacroType
 )
@@ -4684,7 +4471,7 @@ FUnrealPropertyDefinitionInfo& FHeaderParser::GetVarNameAndDim
 (
 	FUnrealStructDefinitionInfo& ParentStruct,
 	FPropertyBase& VarProperty,
-	EVariableCategory::Type VariableCategory,
+	EVariableCategory VariableCategory,
 	ELayoutMacroType        LayoutMacroType
 )
 {
@@ -4954,15 +4741,11 @@ FUnrealPropertyDefinitionInfo& FHeaderParser::GetVarNameAndDim
 	}
 
 	// Create the property
-	// NOTE: Only FProperyBase part of VarProperty is used.  
 	FName PropertyName(Identifier);
 	FUnrealPropertyDefinitionInfo& PropDef = FPropertyTraits::CreateProperty(VarProperty, ParentStruct, PropertyName, ObjectFlags, VariableCategory, Dimensions.String, SourceFile, InputLine, InputPos);
 
 	// Add the property to the parent
 	ParentStruct.AddProperty(PropDef);
-
-	// if we had any metadata, add it to the class
-	AddMetaDataToClassData(PropDef, TMap<FName,FString>(PropDef.GetPropertyBase().MetaData));
 	return PropDef;
 }
 
@@ -5686,7 +5469,7 @@ FUnrealClassDefinitionInfo& FHeaderParser::CompileClassDeclaration()
 	AddModuleRelativePathToMetadata(ClassDef, MetaData);
 
 	// Register the metadata
-	AddMetaDataToClassData(ClassDef, MoveTemp(MetaData));
+	FUHTMetaData::RemapAndAddMetaData(ClassDef, MoveTemp(MetaData));
 
 	// Handle the start of the rest of the class
 	RequireSymbol( TEXT('{'), TEXT("'Class'") );
@@ -5912,7 +5695,7 @@ void FHeaderParser::CompileInterfaceDeclaration()
 
 	// Register the metadata
 	AddModuleRelativePathToMetadata(*InterfaceClassDef, MetaData);
-	AddMetaDataToClassData(*InterfaceClassDef, MoveTemp(MetaData));
+	FUHTMetaData::RemapAndAddMetaData(*InterfaceClassDef, MoveTemp(MetaData));
 
 	// Handle the start of the rest of the interface
 	RequireSymbol( TEXT('{'), TEXT("'Class'") );
@@ -6168,7 +5951,7 @@ void FHeaderParser::ParseParameterList(FUnrealFunctionDefinitionInfo& FunctionDe
 	{
 		// Get parameter type.
 		FPropertyBase Property(CPT_None);
-		EVariableCategory::Type VariableCategory = FunctionDef.HasAnyFunctionFlags(FUNC_Net) ? EVariableCategory::ReplicatedParameter : EVariableCategory::RegularParameter;
+		EVariableCategory VariableCategory = FunctionDef.HasAnyFunctionFlags(FUNC_Net) ? EVariableCategory::ReplicatedParameter : EVariableCategory::RegularParameter;
 		GetVarType(GetCurrentScope(), Property, ~(CPF_ParmFlags | CPF_AutoWeak | CPF_RepSkip | CPF_UObjectWrapper | CPF_NativeAccessSpecifiers), nullptr, nullptr, EPropertyDeclarationStyle::None, VariableCategory);
 		Property.PropertyFlags |= CPF_Parm;
 
@@ -6224,7 +6007,7 @@ void FHeaderParser::ParseParameterList(FUnrealFunctionDefinitionInfo& FunctionDe
 					UE_LOG_ERROR_UHT(TEXT("Only service request functions cannot contain NoReplication parameters"));
 				}
 
-				if ((PropDef.GetCastFlags() & CASTCLASS_FDelegateProperty) != 0)
+				if (PropDef.GetPropertyBase().IsDelegateOrDelegateStaticArray())
 				{
 					UE_LOG_ERROR_UHT(TEXT("Replicated functions cannot contain delegate parameters (this would be insecure)"));
 				}
@@ -6246,7 +6029,7 @@ void FHeaderParser::ParseParameterList(FUnrealFunctionDefinitionInfo& FunctionDe
 					UE_LOG_ERROR_UHT(TEXT("Service request functions cannot contain out parameters, unless marked NotReplicated"));
 				}
 
-				if (!(Property.PropertyFlags & CPF_RepSkip) && (PropDef.GetCastFlags() & CASTCLASS_FDelegateProperty) != 0)
+				if (!(Property.PropertyFlags & CPF_RepSkip) && PropDef.GetPropertyBase().IsDelegateOrDelegateStaticArray())
 				{
 					UE_LOG_ERROR_UHT(TEXT("Service request functions cannot contain delegate parameters, unless marked NotReplicated"));
 				}
@@ -6533,7 +6316,7 @@ FUnrealFunctionDefinitionInfo& FHeaderParser::CompileDelegateDeclaration(const T
 
 	AddFormattedPrevCommentAsTooltipMetaData(MetaData);
 
-	AddMetaDataToClassData(DelegateSignatureFunctionDef, MoveTemp(MetaData));
+	FUHTMetaData::RemapAndAddMetaData(DelegateSignatureFunctionDef, MoveTemp(MetaData));
 
 	// Optionally consume a semicolon, it's not required for the delegate macro since it contains one internally
 	MatchSemi();
@@ -6950,7 +6733,7 @@ void FHeaderParser::CompileFunctionDeclaration()
 
 	AddFormattedPrevCommentAsTooltipMetaData(MetaData);
 
-	AddMetaDataToClassData(FuncDef, MoveTemp(MetaData));
+	FUHTMetaData::RemapAndAddMetaData(FuncDef, MoveTemp(MetaData));
 
 	// 'final' and 'override' can appear in any order before an optional '= 0' pure virtual specifier
 	bool bFoundFinal    = MatchIdentifier(TEXT("final"), ESearchCase::CaseSensitive);
@@ -7296,20 +7079,20 @@ void FHeaderParser::CompileVariableDeclaration(FUnrealStructDefinitionInfo& Stru
 	EPropertyFlags EdFlags       = CPF_None;
 
 	// Get variable type.
-	FPropertyBase OriginalProperty(CPT_None);
+	FPropertyBase OriginalPropertyBase(CPT_None);
 	FIndexRange TypeRange;
 	ELayoutMacroType LayoutMacroType = ELayoutMacroType::None;
-	GetVarType(&*StructDef.GetScope(), OriginalProperty, DisallowFlags, /*OuterPropertyType=*/ nullptr, nullptr, EPropertyDeclarationStyle::UPROPERTY, EVariableCategory::Member, &TypeRange, &LayoutMacroType);
-	OriginalProperty.PropertyFlags |= EdFlags;
+	GetVarType(&*StructDef.GetScope(), OriginalPropertyBase, DisallowFlags, /*OuterPropertyType=*/ nullptr, nullptr, EPropertyDeclarationStyle::UPROPERTY, EVariableCategory::Member, &TypeRange, &LayoutMacroType);
+	OriginalPropertyBase.PropertyFlags |= EdFlags;
 
-	FString* Category = OriginalProperty.MetaData.Find(NAME_Category);
+	FString* Category = OriginalPropertyBase.MetaData.Find(NAME_Category);
 
 	// First check if the category was specified at all and if the property was exposed to the editor.
-	if (!Category && (OriginalProperty.PropertyFlags & (CPF_Edit|CPF_BlueprintVisible)))
+	if (!Category && (OriginalPropertyBase.PropertyFlags & (CPF_Edit|CPF_BlueprintVisible)))
 	{
 		if ((&StructDef.GetPackageDef() != nullptr) && !bIsCurrentModulePartOfEngine)
 		{
-			Category = &OriginalProperty.MetaData.Add(NAME_Category, StructDef.GetName());
+			Category = &OriginalPropertyBase.MetaData.Add(NAME_Category, StructDef.GetName());
 		}
 		else
 		{
@@ -7318,7 +7101,7 @@ void FHeaderParser::CompileVariableDeclaration(FUnrealStructDefinitionInfo& Stru
 	}
 
 	// Validate that pointer properties are not interfaces (which are not GC'd and so will cause runtime errors)
-	if (OriginalProperty.PointerType == EPointerType::Native && OriginalProperty.ClassDef->AsClass() != nullptr && OriginalProperty.ClassDef->IsInterface())
+	if (OriginalPropertyBase.PointerType == EPointerType::Native && OriginalPropertyBase.ClassDef->AsClass() != nullptr && OriginalPropertyBase.ClassDef->IsInterface())
 	{
 		// Get the name of the type, removing the asterisk representing the pointer
 		FString TypeName = FString(TypeRange.Count, Input + TypeRange.StartIndex).TrimStartAndEnd().LeftChop(1).TrimEnd();
@@ -7326,34 +7109,34 @@ void FHeaderParser::CompileVariableDeclaration(FUnrealStructDefinitionInfo& Stru
 	}
 
 	// If the category was specified explicitly, it wins
-	if (Category && !(OriginalProperty.PropertyFlags & (CPF_Edit|CPF_BlueprintVisible|CPF_BlueprintAssignable|CPF_BlueprintCallable)))
+	if (Category && !(OriginalPropertyBase.PropertyFlags & (CPF_Edit|CPF_BlueprintVisible|CPF_BlueprintAssignable|CPF_BlueprintCallable)))
 	{
 		UE_LOG_WARNING_UHT(TEXT("Property has a Category set but is not exposed to the editor or Blueprints with EditAnywhere, BlueprintReadWrite, VisibleAnywhere, BlueprintReadOnly, BlueprintAssignable, BlueprintCallable keywords.\r\n"));
 	}
 
 	// Make sure that editblueprint variables are editable
-	if(!(OriginalProperty.PropertyFlags & CPF_Edit))
+	if(!(OriginalPropertyBase.PropertyFlags & CPF_Edit))
 	{
-		if (OriginalProperty.PropertyFlags & CPF_DisableEditOnInstance)
+		if (OriginalPropertyBase.PropertyFlags & CPF_DisableEditOnInstance)
 		{
 			UE_LOG_ERROR_UHT(TEXT("Property cannot have 'DisableEditOnInstance' without being editable"));
 		}
 
-		if (OriginalProperty.PropertyFlags & CPF_DisableEditOnTemplate)
+		if (OriginalPropertyBase.PropertyFlags & CPF_DisableEditOnTemplate)
 		{
 			UE_LOG_ERROR_UHT(TEXT("Property cannot have 'DisableEditOnTemplate' without being editable"));
 		}
 	}
 
 	// Validate.
-	if (OriginalProperty.PropertyFlags & CPF_ParmFlags)
+	if (OriginalPropertyBase.PropertyFlags & CPF_ParmFlags)
 	{
 		FError::Throwf(TEXT("Illegal type modifiers in member variable declaration") );
 	}
 
-	if (FString* ExposeOnSpawnValue = OriginalProperty.MetaData.Find(NAME_ExposeOnSpawn))
+	if (FString* ExposeOnSpawnValue = OriginalPropertyBase.MetaData.Find(NAME_ExposeOnSpawn))
 	{
-		if ((*ExposeOnSpawnValue == TEXT("true")) && !FExposeOnSpawnValidator::IsSupported(OriginalProperty))
+		if ((*ExposeOnSpawnValue == TEXT("true")) && !FExposeOnSpawnValidator::IsSupported(OriginalPropertyBase))
 		{
 			UE_LOG_ERROR_UHT(TEXT("ExposeOnSpawn - Property cannot be exposed"));
 		}
@@ -7365,20 +7148,20 @@ void FHeaderParser::CompileVariableDeclaration(FUnrealStructDefinitionInfo& Stru
 	}
 
 	// If Property is a Replicated Struct check to make sure there are no Properties that are not allowed to be Replicated in the Struct 
-	if (OriginalProperty.Type == CPT_Struct && OriginalProperty.PropertyFlags & CPF_Net && OriginalProperty.ScriptStructDef)
+	if (OriginalPropertyBase.Type == CPT_Struct && OriginalPropertyBase.PropertyFlags & CPF_Net && OriginalPropertyBase.ScriptStructDef)
 	{
-		ValidateScriptStructOkForNet(OriginalProperty.ScriptStructDef->GetName(), *OriginalProperty.ScriptStructDef);
+		ValidateScriptStructOkForNet(OriginalPropertyBase.ScriptStructDef->GetName(), *OriginalPropertyBase.ScriptStructDef);
 	}
 
 	// Process all variables of this type.
 	TArray<FUnrealPropertyDefinitionInfo*> NewProperties;
 	for (;;)
 	{
-		FPropertyBase Property = OriginalProperty;
-		FUnrealPropertyDefinitionInfo& NewPropDef = GetVarNameAndDim(StructDef, Property, EVariableCategory::Member, LayoutMacroType);
+		FPropertyBase PropertyBase = OriginalPropertyBase;
+		FUnrealPropertyDefinitionInfo& NewPropDef = GetVarNameAndDim(StructDef, PropertyBase, EVariableCategory::Member, LayoutMacroType);
 
 		// Optionally consume the :1 at the end of a bitfield boolean declaration
-		if (Property.IsBool())
+		if (PropertyBase.IsBool())
 		{
 			if (LayoutMacroType == ELayoutMacroType::Bitfield || LayoutMacroType == ELayoutMacroType::BitfieldEditorOnly || MatchSymbol(TEXT(':')))
 			{
@@ -7391,7 +7174,7 @@ void FHeaderParser::CompileVariableDeclaration(FUnrealStructDefinitionInfo& Stru
 		}
 
 		// Deprecation validation
-		ValidatePropertyIsDeprecatedIfNecessary(Property, NULL);
+		ValidatePropertyIsDeprecatedIfNecessary(PropertyBase, NULL);
 
 		if (TopNest->NestType != ENestType::FunctionDeclaration)
 		{
@@ -7402,13 +7185,9 @@ void FHeaderParser::CompileVariableDeclaration(FUnrealStructDefinitionInfo& Stru
 		}
 
 		NewProperties.Add(&NewPropDef);
-		// we'll need any metadata tags we parsed later on when we call ConvertEOLCommentToTooltip() so the tags aren't clobbered
-		OriginalProperty.MetaData = Property.MetaData;
 
-		if (NewPropDef.HasAnyPropertyFlags(CPF_RepNotify))
-		{
-			NewPropDef.SetRepNotifyFunc(OriginalProperty.RepNotifyName);
-		}
+		// we'll need any metadata tags we parsed later on when we call ConvertEOLCommentToTooltip() so the tags aren't clobbered
+		OriginalPropertyBase.MetaData = PropertyBase.MetaData;
 
 		if (UScriptStruct* StructBeingBuilt = Cast<UScriptStruct>(Struct))
 		{
diff --git a/Engine/Source/Programs/UnrealHeaderTool/Private/HeaderParser.h b/Engine/Source/Programs/UnrealHeaderTool/Private/HeaderParser.h
--- a/Engine/Source/Programs/UnrealHeaderTool/Private/HeaderParser.h
+++ b/Engine/Source/Programs/UnrealHeaderTool/Private/HeaderParser.h
@@ -61,18 +61,6 @@ namespace EDelegateSpecifierAction
 	};
 }
 
-/** The category of variable declaration being parsed */
-namespace EVariableCategory
-{
-	enum Type
-	{
-		RegularParameter,
-		ReplicatedParameter,
-		Return,
-		Member
-	};
-}
-
 enum class ELayoutMacroType
 {
 	None = -1,
@@ -371,8 +359,8 @@ public:
 	*
 	* @return	a pointer to a UField with a name matching InIdentifier, or NULL if it wasn't found
 	*/
-	static FUnrealFunctionDefinitionInfo* FindFunction(FUnrealStructDefinitionInfo& InScope, const TCHAR* InIdentifier, bool bIncludeParents = true, const TCHAR* Thing = nullptr);
-	static FUnrealPropertyDefinitionInfo* FindProperty(FUnrealStructDefinitionInfo& InScope, const TCHAR* InIdentifier, bool bIncludeParents = true, const TCHAR* Thing = nullptr);
+	static FUnrealFunctionDefinitionInfo* FindFunction(const FUnrealStructDefinitionInfo& InScope, const TCHAR* InIdentifier, bool bIncludeParents = true, const TCHAR* Thing = nullptr);
+	static FUnrealPropertyDefinitionInfo* FindProperty(const FUnrealStructDefinitionInfo& InScope, const TCHAR* InIdentifier, bool bIncludeParents = true, const TCHAR* Thing = nullptr);
 
 	// Checks ToValidate to make sure that its associated sparse class data struct, if one exists, is a valid structure to use for storing sparse class data.
 	static void CheckSparseClassData(const FUnrealStructDefinitionInfo& StructDef);
@@ -490,7 +478,7 @@ protected:
 		const FToken*                   OuterPropertyType,
 		const EPropertyFlags*			OuterPropertyFlags,
 		EPropertyDeclarationStyle::Type PropertyDeclarationStyle,
-		EVariableCategory::Type         VariableCategory,
+		EVariableCategory               VariableCategory,
 		FIndexRange*                    ParsedVarIndexRange = nullptr,
 		ELayoutMacroType*               OutLayoutMacroType = nullptr);
 
@@ -506,7 +494,7 @@ protected:
 	FUnrealPropertyDefinitionInfo& GetVarNameAndDim(
 		FUnrealStructDefinitionInfo& ParentStruct,
 		FPropertyBase& VarProperty,
-		EVariableCategory::Type VariableCategory,
+		EVariableCategory VariableCategory,
 		ELayoutMacroType LayoutMacroType = ELayoutMacroType::None);
 	
 	/**
diff --git a/Engine/Source/Programs/UnrealHeaderTool/Private/NativeClassExporter.h b/Engine/Source/Programs/UnrealHeaderTool/Private/NativeClassExporter.h
--- a/Engine/Source/Programs/UnrealHeaderTool/Private/NativeClassExporter.h
+++ b/Engine/Source/Programs/UnrealHeaderTool/Private/NativeClassExporter.h
@@ -691,24 +691,31 @@ public:
 	}
 
 	template <typename T>
-	static FName GetOverriddenFName(const T* Item)
+	static FString GetOverriddenName(const T& Item)
 	{
-		FString OverriddenName = Item->GetMetaData(TEXT("OverrideNativeName"));
+		const FString& OverriddenName = Item.GetMetaData(TEXT("OverrideNativeName"));
 		if (!OverriddenName.IsEmpty())
 		{
-			return FName(*OverriddenName);
+			return OverriddenName.ReplaceCharWithEscapedChar();
 		}
-		return Item->GetFName();
+		return Item.GetName();
 	}
 
-	static FString GetOverriddenPathName(const FUnrealPropertyDefinitionInfo& ItemDef)
+	template <typename T>
+	static FName GetOverriddenFName(const T& Item)
 	{
-		return FString::Printf(TEXT("%s.%s"), *ItemDef.GetTypePackageName(), *GetOverriddenName(ItemDef.GetProperty()));
+		FString OverriddenName = Item.GetMetaData(TEXT("OverrideNativeName"));
+		if (!OverriddenName.IsEmpty())
+		{
+			return FName(*OverriddenName);
+		}
+		return Item.GetFName();
 	}
 
-	static FString GetOverriddenPathName(const FUnrealFieldDefinitionInfo& ItemDef)
+	template <typename T>
+	static FString GetOverriddenPathName(const T& Def)
 	{
-		return FString::Printf(TEXT("%s.%s"), *ItemDef.GetTypePackageName(), *GetOverriddenName(ItemDef.GetField()));
+		return FString::Printf(TEXT("%s.%s"), *Def.GetTypePackageName(), *GetOverriddenName(Def));
 	}
 
 	/**
diff --git a/Engine/Source/Programs/UnrealHeaderTool/Private/ParserHelper.h b/Engine/Source/Programs/UnrealHeaderTool/Private/ParserHelper.h
--- a/Engine/Source/Programs/UnrealHeaderTool/Private/ParserHelper.h
+++ b/Engine/Source/Programs/UnrealHeaderTool/Private/ParserHelper.h
@@ -323,6 +323,16 @@ public:
 		return IsNumeric(Type) && IsPrimitiveOrPrimitiveStaticArray();
 	}
 
+	bool IsDelegateOrDelegateStaticArray() const
+	{
+		return Type == CPT_Delegate && IsPrimitiveOrPrimitiveStaticArray();
+	}
+
+	bool IsMulticastDelegateOrMulticastDelegateStaticArray() const
+	{
+		return Type == CPT_Delegate && IsPrimitiveOrPrimitiveStaticArray();
+	}
+
 	FUnrealEnumDefinitionInfo* AsEnum() const;
 
 	bool IsEnum() const;
diff --git a/Engine/Source/Programs/UnrealHeaderTool/Private/PropertyTypes.cpp b/Engine/Source/Programs/UnrealHeaderTool/Private/PropertyTypes.cpp
--- a/Engine/Source/Programs/UnrealHeaderTool/Private/PropertyTypes.cpp
+++ b/Engine/Source/Programs/UnrealHeaderTool/Private/PropertyTypes.cpp
@@ -10,7 +10,6 @@
 #include "UObject/ObjectMacros.h"
 
 void AddEditInlineMetaData(TMap<FName, FString>& MetaData);
-void AddMetaDataToClassData(FUnrealTypeDefinitionInfo& TypeDef, TMap<FName, FString>&& InMetaData);
 
 // Following is the relationship between the property types
 //
@@ -89,17 +88,15 @@ namespace
 {
 	void PropagateFlagsFromInnerAndHandlePersistentInstanceMetadata(EPropertyFlags& DestFlags, const TMap<FName, FString>& InMetaData, FUnrealPropertyDefinitionInfo& InnerDef)
 	{
-		FProperty* Inner = InnerDef.GetProperty();
-
 		// Copy some of the property flags to the container property.
-		if (Inner->PropertyFlags & (CPF_ContainsInstancedReference | CPF_InstancedReference))
+		if (InnerDef.HasAnyPropertyFlags(CPF_ContainsInstancedReference | CPF_InstancedReference))
 		{
 			DestFlags |= CPF_ContainsInstancedReference;
 			DestFlags &= ~(CPF_InstancedReference | CPF_PersistentInstance); //this was propagated to the inner
 
-			if (Inner->PropertyFlags & CPF_PersistentInstance)
+			if (InnerDef.HasAnyPropertyFlags(CPF_PersistentInstance))
 			{
-				AddMetaDataToClassData(InnerDef, TMap<FName, FString>(InMetaData));
+				FUHTMetaData::RemapAndAddMetaData(InnerDef, TMap<FName, FString>(InMetaData));
 			}
 		}
 	};
@@ -202,12 +199,21 @@ namespace
 		}
 	};
 
+	template<typename TraitsType>
+	struct CreatePropertyDispatch
+	{
+		FProperty* operator()(FUnrealPropertyDefinitionInfo& PropDef, FFieldVariant Scope, const FName& Name, EObjectFlags ObjectFlags)
+		{
+			return TraitsType::CreateProperty(PropDef, Scope, Name, ObjectFlags);
+		}
+	};
+
 	template<typename TraitsType>
 	struct CreateEngineTypeDispatch
 	{
-		FProperty* operator()(FUnrealPropertyDefinitionInfo& PropDef, FFieldVariant Scope, const FName& Name, EObjectFlags ObjectFlags, EVariableCategory::Type VariableCategory, const TCHAR* Dimensions)
+		FProperty* operator()(FUnrealPropertyDefinitionInfo& PropDef, FFieldVariant Scope, const FName& Name, EObjectFlags ObjectFlags)
 		{
-			return TraitsType::CreateEngineType(PropDef, Scope, Name, ObjectFlags, VariableCategory, Dimensions);
+			return TraitsType::CreateProperty(PropDef, Scope, Name, ObjectFlags);
 		}
 	};
 
@@ -220,6 +226,33 @@ namespace
 		}
 	};
 
+	template<typename TraitsType>
+	struct GetEngineClassNameDispatch
+	{
+		FString operator()(const FUnrealPropertyDefinitionInfo& PropDef)
+		{
+			return TraitsType::GetEngineClassName(PropDef);
+		}
+	};
+
+	template<typename TraitsType>
+	struct GetCPPTypeDispatch
+	{
+		FString operator()(const FUnrealPropertyDefinitionInfo& PropDef, FString* ExtendedTypeText, uint32 CPPExportFlags)
+		{
+			return TraitsType::GetCPPType(PropDef, ExtendedTypeText, CPPExportFlags);
+		}
+	};
+
+	template<typename TraitsType>
+	struct GetCPPTypeForwardDeclarationDispatch
+	{
+		FString operator()(const FUnrealPropertyDefinitionInfo& PropDef)
+		{
+			return TraitsType::GetCPPTypeForwardDeclaration(PropDef);
+		}
+	};
+
 	//----------------------------------------------------------------------------------------------------------------------------------------------------------
 	//----------------------------------------------------------------------------------------------------------------------------------------------------------
 	// Helper Methods
@@ -227,11 +260,9 @@ namespace
 	//----------------------------------------------------------------------------------------------------------------------------------------------------------
 
 	template <bool bHandleContainers>
-	FProperty* CreatePropertyHelper(FUnrealPropertyDefinitionInfo& PropDef, FFieldVariant Scope, const FName& Name, EObjectFlags ObjectFlags, EVariableCategory::Type VariableCategory, const TCHAR* Dimensions)
+	FProperty* CreatePropertyHelper(FUnrealPropertyDefinitionInfo& PropDef, FFieldVariant Scope, const FName& Name, EObjectFlags ObjectFlags)
 	{
-		FProperty* Property = PropertyTypeDispatch<CreateEngineTypeDispatch, bHandleContainers, FProperty*>(PropDef.GetPropertyBase(), std::ref(PropDef), Scope, std::ref(Name), ObjectFlags, VariableCategory, Dimensions);
-		Property->PropertyFlags = PropDef.GetPropertyBase().PropertyFlags;
-		return Property;
+		return PropertyTypeDispatch<CreatePropertyDispatch, bHandleContainers, FProperty*>(PropDef.GetPropertyBase(), std::ref(PropDef), Scope, std::ref(Name), ObjectFlags);
 	}
 
 	bool IsSupportedByBlueprintSansContainers(const FUnrealPropertyDefinitionInfo& PropDef, bool bMemberVariable)
@@ -268,11 +299,22 @@ struct FPropertyTypeTraitsBase
 	 * @param Scope The parent object owning the property
 	 * @param Name The name of the property
 	 * @param ObjectFlags The flags associated with the property
-	 * @param VariableCategory The parsing context of the property
-	 * @param Dimensions When this is a static array, this represents the dimensions value
 	 * @return The pointer to the newly created property.  It will be attached to the definition by the caller
 	 */
-	 //static FProperty* CreateEngineType(FUnrealPropertyDefinitionInfo& PropDef, FFieldVariant Scope, const FName& Name, EObjectFlags ObjectFlags, EVariableCategory::Type VariableCategory, const TCHAR* Dimensions)
+	 //static FProperty* CreateProperty(FUnrealPropertyDefinitionInfo& PropDef, FFieldVariant Scope, const FName& Name, EObjectFlags ObjectFlags)
+
+	/**
+	 * Given a property definition create the underlying engine type
+	 *
+	 * NOTE: This method MUST be implemented by each of the property types.
+	 *
+	 * @param PropDef The definition of the property
+	 * @param Scope The parent object owning the property
+	 * @param Name The name of the property
+	 * @param ObjectFlags The flags associated with the property
+	 * @return The pointer to the newly created property.  It will be attached to the definition by the caller
+	 */
+	 //static FProperty* CreateProperty(FUnrealPropertyDefinitionInfo& PropDef, FFieldVariant Scope, const FName& Name, EObjectFlags ObjectFlags)
 
 	/**
 	 * Returns true if this property is supported by blueprints
@@ -281,6 +323,24 @@ struct FPropertyTypeTraitsBase
 	{
 		return false;
 	}
+
+	/**
+	 * Return the engine class name for the given property information.
+	 * @param PropDef The property in question
+	 * @return The name of the engine property that will represent this definition.
+	 */
+	 //static FString GetEngineClassName(const FUnrealPropertyDefinitionInfo& PropDef);
+
+	/**
+	 * Returns the text to use for exporting this property to header file.
+	 *
+	 * @param PropDef The definition of the property
+	 * @param ExtendedTypeText for property types which use templates, will be filled in with the type
+	 * @param CPPExportFlags flags for modifying the behavior of the export
+	 */
+	//static FString GetCPPType(const FUnrealPropertyDefinitionInfo& PropDef, FString* ExtendedTypeText, uint32 CPPExportFlags);
+
+	//static FString GetCPPTypeForwardDeclaration(const FUnrealPropertyDefinitionInfo& PropDef);
 };
 
 //----------------------------------------------------------------------------------------------------------------------------------------------------------
@@ -291,6 +351,10 @@ struct FPropertyTypeTraitsBase
 
 struct FPropertyTypeTraitsNumericBase : public FPropertyTypeTraitsBase
 {
+	static FString GetCPPTypeForwardDeclaration(const FUnrealPropertyDefinitionInfo& PropDef)
+	{
+		return FString();
+	}
 };
 
 struct FPropertyTypeTraitsByte : public FPropertyTypeTraitsNumericBase
@@ -306,7 +370,7 @@ struct FPropertyTypeTraitsByte : public FPropertyTypeTraitsNumericBase
 		return false;
 	}
 
-	static FProperty* CreateEngineType(FUnrealPropertyDefinitionInfo& PropDef, FFieldVariant Scope, const FName& Name, EObjectFlags ObjectFlags, EVariableCategory::Type VariableCategory, const TCHAR* Dimensions)
+	static FProperty* CreateProperty(FUnrealPropertyDefinitionInfo& PropDef, FFieldVariant Scope, const FName& Name, EObjectFlags ObjectFlags)
 	{
 		const FPropertyBase& VarProperty = PropDef.GetPropertyBase();
 		UEnum* Enum = VarProperty.EnumDef ? VarProperty.EnumDef->GetEnum() : nullptr;
@@ -323,28 +387,97 @@ struct FPropertyTypeTraitsByte : public FPropertyTypeTraitsNumericBase
 	{
 		return true;
 	}
+
+	static FString GetEngineClassName(const FUnrealPropertyDefinitionInfo& PropDef)
+	{
+		return FByteProperty::StaticClass()->GetName();
+	}
+
+	static FString GetCPPType(const FUnrealPropertyDefinitionInfo& PropDef, FString* ExtendedTypeText, uint32 CPPExportFlags)
+	{
+		if (FUnrealEnumDefinitionInfo* EnumDef = PropDef.GetPropertyBase().AsEnum())
+		{
+			UEnum* Enum = EnumDef->GetEnum();
+			const bool bEnumClassForm = EnumDef->GetCppForm() == UEnum::ECppForm::EnumClass;
+			const bool bNonNativeEnum = Enum->GetClass() != UEnum::StaticClass(); // cannot use RF_Native flag, because in UHT the flag is not set
+			const bool bRawParam = (CPPExportFlags & CPPF_ArgumentOrReturnValue)
+				&& ((PropDef.HasAnyPropertyFlags(CPF_ReturnParm) || !PropDef.HasAnyPropertyFlags(CPF_OutParm))
+					|| bNonNativeEnum);
+			const bool bConvertedCode = (CPPExportFlags & CPPF_BlueprintCppBackend) && bNonNativeEnum;
+
+			FString FullyQualifiedEnumName;
+			if (!Enum->CppType.IsEmpty())
+			{
+				FullyQualifiedEnumName = Enum->CppType;
+			}
+			else
+			{
+				// This would give the wrong result if it's a namespaced type and the CppType hasn't
+				// been set, but we do this here in case existing code relies on it... somehow.
+				if ((CPPExportFlags & CPPF_BlueprintCppBackend) && bNonNativeEnum)
+				{
+					ensure(Enum->CppType.IsEmpty());
+					FullyQualifiedEnumName = ::UnicodeToCPPIdentifier(Enum->GetName(), false, TEXT("E__"));
+				}
+				else
+				{
+					FullyQualifiedEnumName = Enum->GetName();
+				}
+			}
+
+			if (bEnumClassForm || bRawParam || bConvertedCode)
+			{
+				return FullyQualifiedEnumName;
+			}
+			else
+			{
+				return FString::Printf(TEXT("TEnumAsByte<%s>"), *FullyQualifiedEnumName);
+			}
+		}
+		return TEXT("uint8");
+	}
 };
 
 struct FPropertyTypeTraitsInt8 : public FPropertyTypeTraitsNumericBase
 {
-	static FProperty* CreateEngineType(FUnrealPropertyDefinitionInfo& PropDef, FFieldVariant Scope, const FName& Name, EObjectFlags ObjectFlags, EVariableCategory::Type VariableCategory, const TCHAR* Dimensions)
+	static FProperty* CreateProperty(FUnrealPropertyDefinitionInfo& PropDef, FFieldVariant Scope, const FName& Name, EObjectFlags ObjectFlags)
 	{
 		const FPropertyBase& VarProperty = PropDef.GetPropertyBase();
 		FInt8Property* Result = new FInt8Property(Scope, Name, ObjectFlags);
 		check(VarProperty.IntType == EIntType::Sized);
 		return Result;
 	}
+
+	static FString GetEngineClassName(const FUnrealPropertyDefinitionInfo& PropDef)
+	{
+		return FInt8Property::StaticClass()->GetName();
+	}
+
+	static FString GetCPPType(const FUnrealPropertyDefinitionInfo& PropDef, FString* ExtendedTypeText, uint32 CPPExportFlags)
+	{
+		return TEXT("int8");
+	}
 };
 
 struct FPropertyTypeTraitsInt16 : public FPropertyTypeTraitsNumericBase
 {
-	static FProperty* CreateEngineType(FUnrealPropertyDefinitionInfo& PropDef, FFieldVariant Scope, const FName& Name, EObjectFlags ObjectFlags, EVariableCategory::Type VariableCategory, const TCHAR* Dimensions)
+	static FProperty* CreateProperty(FUnrealPropertyDefinitionInfo& PropDef, FFieldVariant Scope, const FName& Name, EObjectFlags ObjectFlags)
 	{
 		const FPropertyBase& VarProperty = PropDef.GetPropertyBase();
 		FInt16Property* Result = new FInt16Property(Scope, Name, ObjectFlags);
 		check(VarProperty.IntType == EIntType::Sized);
 		return Result;
 	}
+
+	static FString GetEngineClassName(const FUnrealPropertyDefinitionInfo& PropDef)
+	{
+		return FInt16Property::StaticClass()->GetName();
+	}
+
+	static FString GetCPPType(const FUnrealPropertyDefinitionInfo& PropDef, FString* ExtendedTypeText, uint32 CPPExportFlags)
+	{
+		return TEXT("int16");
+	}
 };
 
 struct FPropertyTypeTraitsInt : public FPropertyTypeTraitsNumericBase
@@ -360,7 +493,7 @@ struct FPropertyTypeTraitsInt : public FPropertyTypeTraitsNumericBase
 		return false;
 	}
 
-	static FProperty* CreateEngineType(FUnrealPropertyDefinitionInfo& PropDef, FFieldVariant Scope, const FName& Name, EObjectFlags ObjectFlags, EVariableCategory::Type VariableCategory, const TCHAR* Dimensions)
+	static FProperty* CreateProperty(FUnrealPropertyDefinitionInfo& PropDef, FFieldVariant Scope, const FName& Name, EObjectFlags ObjectFlags)
 	{
 		const FPropertyBase& VarProperty = PropDef.GetPropertyBase();
 		FIntProperty* Result = new FIntProperty(Scope, Name, ObjectFlags);
@@ -372,6 +505,16 @@ struct FPropertyTypeTraitsInt : public FPropertyTypeTraitsNumericBase
 	{
 		return true;
 	}
+
+	static FString GetEngineClassName(const FUnrealPropertyDefinitionInfo& PropDef)
+	{
+		return FIntProperty::StaticClass()->GetName();
+	}
+
+	static FString GetCPPType(const FUnrealPropertyDefinitionInfo& PropDef, FString* ExtendedTypeText, uint32 CPPExportFlags)
+	{
+		return TEXT("int32");
+	}
 };
 
 struct FPropertyTypeTraitsInt64 : public FPropertyTypeTraitsNumericBase
@@ -387,7 +530,7 @@ struct FPropertyTypeTraitsInt64 : public FPropertyTypeTraitsNumericBase
 		return false;
 	}
 
-	static FProperty* CreateEngineType(FUnrealPropertyDefinitionInfo& PropDef, FFieldVariant Scope, const FName& Name, EObjectFlags ObjectFlags, EVariableCategory::Type VariableCategory, const TCHAR* Dimensions)
+	static FProperty* CreateProperty(FUnrealPropertyDefinitionInfo& PropDef, FFieldVariant Scope, const FName& Name, EObjectFlags ObjectFlags)
 	{
 		const FPropertyBase& VarProperty = PropDef.GetPropertyBase();
 		FInt64Property* Result = new FInt64Property(Scope, Name, ObjectFlags);
@@ -399,39 +542,79 @@ struct FPropertyTypeTraitsInt64 : public FPropertyTypeTraitsNumericBase
 	{
 		return true;
 	}
+
+	static FString GetEngineClassName(const FUnrealPropertyDefinitionInfo& PropDef)
+	{
+		return FInt64Property::StaticClass()->GetName();
+	}
+
+	static FString GetCPPType(const FUnrealPropertyDefinitionInfo& PropDef, FString* ExtendedTypeText, uint32 CPPExportFlags)
+	{
+		return TEXT("int64");
+	}
 };
 
 struct FPropertyTypeTraitsUInt16 : public FPropertyTypeTraitsNumericBase
 {
-	static FProperty* CreateEngineType(FUnrealPropertyDefinitionInfo& PropDef, FFieldVariant Scope, const FName& Name, EObjectFlags ObjectFlags, EVariableCategory::Type VariableCategory, const TCHAR* Dimensions)
+	static FProperty* CreateProperty(FUnrealPropertyDefinitionInfo& PropDef, FFieldVariant Scope, const FName& Name, EObjectFlags ObjectFlags)
 	{
 		const FPropertyBase& VarProperty = PropDef.GetPropertyBase();
 		FUInt16Property* Result = new FUInt16Property(Scope, Name, ObjectFlags);
 		check(VarProperty.IntType == EIntType::Sized);
 		return Result;
 	}
+
+	static FString GetEngineClassName(const FUnrealPropertyDefinitionInfo& PropDef)
+	{
+		return FUInt16Property::StaticClass()->GetName();
+	}
+
+	static FString GetCPPType(const FUnrealPropertyDefinitionInfo& PropDef, FString* ExtendedTypeText, uint32 CPPExportFlags)
+	{
+		return TEXT("uint16");
+	}
 };
 
 struct FPropertyTypeTraitsUInt32 : public FPropertyTypeTraitsNumericBase
 {
-	static FProperty* CreateEngineType(FUnrealPropertyDefinitionInfo& PropDef, FFieldVariant Scope, const FName& Name, EObjectFlags ObjectFlags, EVariableCategory::Type VariableCategory, const TCHAR* Dimensions)
+	static FProperty* CreateProperty(FUnrealPropertyDefinitionInfo& PropDef, FFieldVariant Scope, const FName& Name, EObjectFlags ObjectFlags)
 	{
 		const FPropertyBase& VarProperty = PropDef.GetPropertyBase();
 		FUInt32Property* Result = new FUInt32Property(Scope, Name, ObjectFlags);
 		PropDef.SetUnsized(VarProperty.IntType == EIntType::Unsized);
 		return Result;
 	}
+
+	static FString GetEngineClassName(const FUnrealPropertyDefinitionInfo& PropDef)
+	{
+		return FUInt32Property::StaticClass()->GetName();
+	}
+
+	static FString GetCPPType(const FUnrealPropertyDefinitionInfo& PropDef, FString* ExtendedTypeText, uint32 CPPExportFlags)
+	{
+		return TEXT("uint32");
+	}
 };
 
 struct FPropertyTypeTraitsUInt64 : public FPropertyTypeTraitsNumericBase
 {
-	static FProperty* CreateEngineType(FUnrealPropertyDefinitionInfo& PropDef, FFieldVariant Scope, const FName& Name, EObjectFlags ObjectFlags, EVariableCategory::Type VariableCategory, const TCHAR* Dimensions)
+	static FProperty* CreateProperty(FUnrealPropertyDefinitionInfo& PropDef, FFieldVariant Scope, const FName& Name, EObjectFlags ObjectFlags)
 	{
 		const FPropertyBase& VarProperty = PropDef.GetPropertyBase();
 		FUInt64Property* Result = new FUInt64Property(Scope, Name, ObjectFlags);
 		check(VarProperty.IntType == EIntType::Sized);
 		return Result;
 	}
+
+	static FString GetEngineClassName(const FUnrealPropertyDefinitionInfo& PropDef)
+	{
+		return FUInt64Property::StaticClass()->GetName();
+	}
+
+	static FString GetCPPType(const FUnrealPropertyDefinitionInfo& PropDef, FString* ExtendedTypeText, uint32 CPPExportFlags)
+	{
+		return TEXT("uint64");
+	}
 };
 
 struct FPropertyTypeTraitsFloat : public FPropertyTypeTraitsNumericBase
@@ -447,7 +630,7 @@ struct FPropertyTypeTraitsFloat : public FPropertyTypeTraitsNumericBase
 		return false;
 	}
 
-	static FProperty* CreateEngineType(FUnrealPropertyDefinitionInfo& PropDef, FFieldVariant Scope, const FName& Name, EObjectFlags ObjectFlags, EVariableCategory::Type VariableCategory, const TCHAR* Dimensions)
+	static FProperty* CreateProperty(FUnrealPropertyDefinitionInfo& PropDef, FFieldVariant Scope, const FName& Name, EObjectFlags ObjectFlags)
 	{
 		FFloatProperty* Result = new FFloatProperty(Scope, Name, ObjectFlags);
 		return Result;
@@ -457,6 +640,16 @@ struct FPropertyTypeTraitsFloat : public FPropertyTypeTraitsNumericBase
 	{
 		return true;
 	}
+
+	static FString GetEngineClassName(const FUnrealPropertyDefinitionInfo& PropDef)
+	{
+		return FFloatProperty::StaticClass()->GetName();
+	}
+
+	static FString GetCPPType(const FUnrealPropertyDefinitionInfo& PropDef, FString* ExtendedTypeText, uint32 CPPExportFlags)
+	{
+		return TEXT("float");
+	}
 };
 
 struct FPropertyTypeTraitsDouble : public FPropertyTypeTraitsNumericBase
@@ -472,7 +665,7 @@ struct FPropertyTypeTraitsDouble : public FPropertyTypeTraitsNumericBase
 		return false;
 	}
 
-	static FProperty* CreateEngineType(FUnrealPropertyDefinitionInfo& PropDef, FFieldVariant Scope, const FName& Name, EObjectFlags ObjectFlags, EVariableCategory::Type VariableCategory, const TCHAR* Dimensions)
+	static FProperty* CreateProperty(FUnrealPropertyDefinitionInfo& PropDef, FFieldVariant Scope, const FName& Name, EObjectFlags ObjectFlags)
 	{
 		FDoubleProperty* Result = new FDoubleProperty(Scope, Name, ObjectFlags);
 		return Result;
@@ -482,6 +675,16 @@ struct FPropertyTypeTraitsDouble : public FPropertyTypeTraitsNumericBase
 	{
 		return true;
 	}
+
+	static FString GetEngineClassName(const FUnrealPropertyDefinitionInfo& PropDef)
+	{
+		return FDoubleProperty::StaticClass()->GetName();
+	}
+
+	static FString GetCPPType(const FUnrealPropertyDefinitionInfo& PropDef, FString* ExtendedTypeText, uint32 CPPExportFlags)
+	{
+		return TEXT("double");
+	}
 };
 
 struct FPropertyTypeTraitsLargeWorldCoordinatesReal : public FPropertyTypeTraitsNumericBase
@@ -497,7 +700,7 @@ struct FPropertyTypeTraitsLargeWorldCoordinatesReal : public FPropertyTypeTraits
 		return false;
 	}
 
-	static FProperty* CreateEngineType(FUnrealPropertyDefinitionInfo& PropDef, FFieldVariant Scope, const FName& Name, EObjectFlags ObjectFlags, EVariableCategory::Type VariableCategory, const TCHAR* Dimensions)
+	static FProperty* CreateProperty(FUnrealPropertyDefinitionInfo& PropDef, FFieldVariant Scope, const FName& Name, EObjectFlags ObjectFlags)
 	{
 		FLargeWorldCoordinatesRealProperty* Result = new FLargeWorldCoordinatesRealProperty(Scope, Name, ObjectFlags);
 		return Result;
@@ -507,6 +710,16 @@ struct FPropertyTypeTraitsLargeWorldCoordinatesReal : public FPropertyTypeTraits
 	{
 		return true;
 	}
+
+	static FString GetEngineClassName(const FUnrealPropertyDefinitionInfo& PropDef)
+	{
+		return FLargeWorldCoordinatesRealProperty::StaticClass()->GetName();
+	}
+
+	static FString GetCPPType(const FUnrealPropertyDefinitionInfo& PropDef, FString* ExtendedTypeText, uint32 CPPExportFlags)
+	{
+		return TEXT("double");
+	}
 };
 
 //----------------------------------------------------------------------------------------------------------------------------------------------------------
@@ -520,6 +733,31 @@ struct FPropertyTypeTraitsLargeWorldCoordinatesReal : public FPropertyTypeTraits
  */
 struct FPropertyTypeTraitsBooleanBase : public FPropertyTypeTraitsBase
 {
+	template <typename SizeType, bool bIsNativeBool>
+	static FProperty* CreatePropertyHelper(FUnrealPropertyDefinitionInfo& PropDef, FFieldVariant Scope, const FName& Name, EObjectFlags ObjectFlags)
+	{
+		FBoolProperty* Result = new FBoolProperty(Scope, Name, ObjectFlags);
+		bool bActsLikeNativeBool = bIsNativeBool || PropDef.GetVariableCategory() == EVariableCategory::Return;
+		Result->SetBoolSize(bActsLikeNativeBool ? sizeof(bool) : sizeof(SizeType), bActsLikeNativeBool);
+		return Result;
+	}
+
+	static FString GetCPPTypeHelper(const FUnrealPropertyDefinitionInfo& PropDef, FString* ExtendedTypeText, uint32 CPPExportFlags, bool bIsNativeBool, const TCHAR* SizeText)
+	{
+		if (bIsNativeBool
+			|| ((CPPExportFlags & (CPPF_Implementation | CPPF_ArgumentOrReturnValue)) == (CPPF_Implementation | CPPF_ArgumentOrReturnValue))
+			|| ((CPPExportFlags & CPPF_BlueprintCppBackend) != 0))
+		{
+			// Export as bool if this is actually a bool or it's being exported as a return value of C++ function definition.
+			return TEXT("bool");
+		}
+		else
+		{
+			return SizeText;
+		}
+		return TEXT("uint32");
+	}
+
 	static bool DefaultValueStringCppFormatToInnerFormat(const FUnrealPropertyDefinitionInfo& PropDef, const FString& CppForm, FString& OutForm)
 	{
 		if (FDefaultValueHelper::Is(CppForm, TEXT("true")) ||
@@ -531,58 +769,85 @@ struct FPropertyTypeTraitsBooleanBase : public FPropertyTypeTraitsBase
 		return false;
 	}
 
-	template <typename SizeType, bool bIsNativeBool>
-	static FProperty* CreateEngineTypeHelper(FUnrealPropertyDefinitionInfo& PropDef, FFieldVariant Scope, const FName& Name, EObjectFlags ObjectFlags, EVariableCategory::Type VariableCategory, const TCHAR* Dimensions)
-	{
-		FBoolProperty* Result = new FBoolProperty(Scope, Name, ObjectFlags);
-		bool bActsLikeNativeBool = bIsNativeBool || VariableCategory == EVariableCategory::Return;
-		Result->SetBoolSize(bActsLikeNativeBool ? sizeof(bool) : sizeof(SizeType), bActsLikeNativeBool);
-		return Result;
-	}
 
 	static bool IsSupportedByBlueprint(const FUnrealPropertyDefinitionInfo& PropDef, bool bMemberVariable)
 	{
 		return true;
 	}
+
+	static FString GetEngineClassName(const FUnrealPropertyDefinitionInfo& PropDef)
+	{
+		return FBoolProperty::StaticClass()->GetName();
+	}
+
+	static FString GetCPPTypeForwardDeclaration(const FUnrealPropertyDefinitionInfo& PropDef)
+	{
+		return FString();
+	}
 };
 
 struct FPropertyTypeTraitsBool : public FPropertyTypeTraitsBooleanBase
 {
-	static FProperty* CreateEngineType(FUnrealPropertyDefinitionInfo& PropDef, FFieldVariant Scope, const FName& Name, EObjectFlags ObjectFlags, EVariableCategory::Type VariableCategory, const TCHAR* Dimensions)
+	static FProperty* CreateProperty(FUnrealPropertyDefinitionInfo& PropDef, FFieldVariant Scope, const FName& Name, EObjectFlags ObjectFlags)
+	{
+		return CreatePropertyHelper<bool, true>(PropDef, Scope, Name, ObjectFlags);
+	}
+
+	static FString GetCPPType(const FUnrealPropertyDefinitionInfo& PropDef, FString* ExtendedTypeText, uint32 CPPExportFlags)
 	{
-		return CreateEngineTypeHelper<bool, true>(PropDef, Scope, Name, ObjectFlags, VariableCategory, Dimensions);
+		return GetCPPTypeHelper(PropDef, ExtendedTypeText, CPPExportFlags, true, TEXT("bool"));
 	}
 };
 
 struct FPropertyTypeTraitsBool8 : public FPropertyTypeTraitsBooleanBase
 {
-	static FProperty* CreateEngineType(FUnrealPropertyDefinitionInfo& PropDef, FFieldVariant Scope, const FName& Name, EObjectFlags ObjectFlags, EVariableCategory::Type VariableCategory, const TCHAR* Dimensions)
+	static FProperty* CreateProperty(FUnrealPropertyDefinitionInfo& PropDef, FFieldVariant Scope, const FName& Name, EObjectFlags ObjectFlags)
+	{
+		return CreatePropertyHelper<uint8, false>(PropDef, Scope, Name, ObjectFlags);
+	}
+
+	static FString GetCPPType(const FUnrealPropertyDefinitionInfo& PropDef, FString* ExtendedTypeText, uint32 CPPExportFlags)
 	{
-		return CreateEngineTypeHelper<uint8, false>(PropDef, Scope, Name, ObjectFlags, VariableCategory, Dimensions);
+		return GetCPPTypeHelper(PropDef, ExtendedTypeText, CPPExportFlags, false, TEXT("uint8"));
 	}
 };
 
 struct FPropertyTypeTraitsBool16 : public FPropertyTypeTraitsBooleanBase
 {
-	static FProperty* CreateEngineType(FUnrealPropertyDefinitionInfo& PropDef, FFieldVariant Scope, const FName& Name, EObjectFlags ObjectFlags, EVariableCategory::Type VariableCategory, const TCHAR* Dimensions)
+	static FProperty* CreateProperty(FUnrealPropertyDefinitionInfo& PropDef, FFieldVariant Scope, const FName& Name, EObjectFlags ObjectFlags)
 	{
-		return CreateEngineTypeHelper<uint16, false>(PropDef, Scope, Name, ObjectFlags, VariableCategory, Dimensions);
+		return CreatePropertyHelper<uint16, false>(PropDef, Scope, Name, ObjectFlags);
+	}
+
+	static FString GetCPPType(const FUnrealPropertyDefinitionInfo& PropDef, FString* ExtendedTypeText, uint32 CPPExportFlags)
+	{
+		return GetCPPTypeHelper(PropDef, ExtendedTypeText, CPPExportFlags, false, TEXT("uint16"));
 	}
 };
 
 struct FPropertyTypeTraitsBool32 : public FPropertyTypeTraitsBooleanBase
 {
-	static FProperty* CreateEngineType(FUnrealPropertyDefinitionInfo& PropDef, FFieldVariant Scope, const FName& Name, EObjectFlags ObjectFlags, EVariableCategory::Type VariableCategory, const TCHAR* Dimensions)
+	static FProperty* CreateProperty(FUnrealPropertyDefinitionInfo& PropDef, FFieldVariant Scope, const FName& Name, EObjectFlags ObjectFlags)
+	{
+		return CreatePropertyHelper<uint32, false>(PropDef, Scope, Name, ObjectFlags);
+	}
+
+	static FString GetCPPType(const FUnrealPropertyDefinitionInfo& PropDef, FString* ExtendedTypeText, uint32 CPPExportFlags)
 	{
-		return CreateEngineTypeHelper<uint32, false>(PropDef, Scope, Name, ObjectFlags, VariableCategory, Dimensions);
+		return GetCPPTypeHelper(PropDef, ExtendedTypeText, CPPExportFlags, false, TEXT("uint32"));
 	}
 };
 
 struct FPropertyTypeTraitsBool64 : public FPropertyTypeTraitsBooleanBase
 {
-	static FProperty* CreateEngineType(FUnrealPropertyDefinitionInfo& PropDef, FFieldVariant Scope, const FName& Name, EObjectFlags ObjectFlags, EVariableCategory::Type VariableCategory, const TCHAR* Dimensions)
+	static FProperty* CreateProperty(FUnrealPropertyDefinitionInfo& PropDef, FFieldVariant Scope, const FName& Name, EObjectFlags ObjectFlags)
+	{
+		return CreatePropertyHelper<uint64, false>(PropDef, Scope, Name, ObjectFlags);
+	}
+
+	static FString GetCPPType(const FUnrealPropertyDefinitionInfo& PropDef, FString* ExtendedTypeText, uint32 CPPExportFlags)
 	{
-		return CreateEngineTypeHelper<uint64, false>(PropDef, Scope, Name, ObjectFlags, VariableCategory, Dimensions);
+		return GetCPPTypeHelper(PropDef, ExtendedTypeText, CPPExportFlags, false, TEXT("uint64"));
 	}
 };
 
@@ -606,12 +871,12 @@ struct FPropertyTypeTraitsEnum : public FPropertyTypeTraitsBase
 		}
 		if (Enum->HasMetaData(TEXT("Hidden"), EnumEntryIndex))
 		{
-			FError::Throwf(TEXT("Hidden enum entries cannot be used as default values: %s \"%s\" "), *PropDef.GetProperty()->GetName(), *CppForm);
+			FError::Throwf(TEXT("Hidden enum entries cannot be used as default values: %s \"%s\" "), *PropDef.GetName(), *CppForm);
 		}
 		return true;
 	}
 
-	static FProperty* CreateEngineType(FUnrealPropertyDefinitionInfo& PropDef, FFieldVariant Scope, const FName& Name, EObjectFlags ObjectFlags, EVariableCategory::Type VariableCategory, const TCHAR* Dimensions)
+	static FProperty* CreateProperty(FUnrealPropertyDefinitionInfo& PropDef, FFieldVariant Scope, const FName& Name, EObjectFlags ObjectFlags)
 	{
 		const FPropertyBase& VarProperty = PropDef.GetPropertyBase();
 
@@ -620,7 +885,7 @@ struct FPropertyTypeTraitsEnum : public FPropertyTypeTraitsBase
 		if (Enum->GetCppForm() != UEnum::ECppForm::EnumClass)
 		{
 			check(VarProperty.Type == EPropertyType::CPT_Byte);
-			return FPropertyTypeTraitsByte::CreateEngineType(PropDef, Scope, Name, ObjectFlags, VariableCategory, Dimensions);
+			return FPropertyTypeTraitsByte::CreateProperty(PropDef, Scope, Name, ObjectFlags);
 		}
 
 #if UHT_ENABLE_VALUE_PROPERTY_TAG
@@ -631,6 +896,8 @@ struct FPropertyTypeTraitsEnum : public FPropertyTypeTraitsBase
 		UnderlyingProperty.EnumDef = nullptr;
 		UnderlyingProperty.PropertyFlags = EPropertyFlags::CPF_None;
 		UnderlyingProperty.ArrayType = EArrayType::None;
+		UnderlyingProperty.RepNotifyName = NAME_None;
+		UnderlyingProperty.MetaData.Empty();
 		switch (VarProperty.EnumDef->GetUnderlyingType())
 		{
 		case EUnderlyingEnumType::int8:        UnderlyingProperty.Type = CPT_Int8;   break;
@@ -652,7 +919,8 @@ struct FPropertyTypeTraitsEnum : public FPropertyTypeTraitsBase
 
 		FEnumProperty* Result = new FEnumProperty(Scope, Name, ObjectFlags);
 		PropDef.SetProperty(Result);
-		FUnrealPropertyDefinitionInfo& SubProp = FPropertyTraits::CreateProperty(UnderlyingProperty, PropDef, TEXT("UnderlyingType"), ObjectFlags, VariableCategory, Dimensions, PropDef.GetUnrealSourceFile(), PropDef.GetLineNumber(), PropDef.GetParsePosition());
+		FUnrealPropertyDefinitionInfo& SubProp = FPropertyTraits::CreateProperty(UnderlyingProperty, PropDef, TEXT("UnderlyingType"), ObjectFlags,
+			PropDef.GetVariableCategory(), PropDef.GetArrayDimensions(), PropDef.GetUnrealSourceFile(), PropDef.GetLineNumber(), PropDef.GetParsePosition());
 		PropDef.SetValuePropDef(SubProp);
 		Result->UnderlyingProp = CastFieldChecked<FNumericProperty>(SubProp.GetProperty());
 		Result->Enum = Enum;
@@ -663,6 +931,64 @@ struct FPropertyTypeTraitsEnum : public FPropertyTypeTraitsBase
 	{
 		return true;
 	}
+
+	static FString GetEngineClassName(const FUnrealPropertyDefinitionInfo& PropDef)
+	{
+		const FPropertyBase& VarProperty = PropDef.GetPropertyBase();
+		UEnum* Enum = VarProperty.EnumDef->GetEnum();
+		if (Enum->GetCppForm() != UEnum::ECppForm::EnumClass)
+		{
+			return FPropertyTypeTraitsByte::GetEngineClassName(PropDef);
+		}
+		else
+		{
+			return FEnumProperty::StaticClass()->GetName();
+		}
+	}
+
+	static FString GetCPPType(const FUnrealPropertyDefinitionInfo& PropDef, FString* ExtendedTypeText, uint32 CPPExportFlags)
+	{
+		const FPropertyBase& VarProperty = PropDef.GetPropertyBase();
+		UEnum* Enum = VarProperty.EnumDef->GetEnum();
+		if (Enum->GetCppForm() != UEnum::ECppForm::EnumClass)
+		{
+			return FPropertyTypeTraitsByte::GetCPPType(PropDef, ExtendedTypeText, CPPExportFlags);
+		}
+		else
+		{
+			const bool bNonNativeEnum = Enum->GetClass() != UEnum::StaticClass(); // cannot use RF_Native flag, because in UHT the flag is not set
+
+			if (!Enum->CppType.IsEmpty())
+			{
+				return Enum->CppType;
+			}
+
+			FString EnumName = Enum->GetName();
+
+			// This would give the wrong result if it's a namespaced type and the CppType hasn't
+			// been set, but we do this here in case existing code relies on it... somehow.
+			if ((CPPExportFlags & CPPF_BlueprintCppBackend) && bNonNativeEnum)
+			{
+				ensure(Enum->CppType.IsEmpty());
+				FString Result = ::UnicodeToCPPIdentifier(EnumName, false, TEXT("E__"));
+				return Result;
+			}
+
+			return EnumName;
+		}
+	}
+
+	static FString GetCPPTypeForwardDeclaration(const FUnrealPropertyDefinitionInfo& PropDef)
+	{
+		const FPropertyBase& VarProperty = PropDef.GetPropertyBase();
+		UEnum* Enum = VarProperty.EnumDef->GetEnum();
+		if (Enum->GetCppForm() != UEnum::ECppForm::EnumClass)
+		{
+			return FPropertyTypeTraitsByte::GetCPPTypeForwardDeclaration(PropDef);
+		}
+
+		return FString::Printf(TEXT("enum class %s : %s;"), *Enum->GetName(), *PropDef.GetValuePropDef().GetCPPType());
+	}
 };
 
 //----------------------------------------------------------------------------------------------------------------------------------------------------------
@@ -689,7 +1015,7 @@ struct FPropertyTypeTraitsObjectBase : public FPropertyTypeTraitsBase
 
 struct FPropertyTypeTraitsObjectReference : public FPropertyTypeTraitsObjectBase
 {
-	static FProperty* CreateEngineType(FUnrealPropertyDefinitionInfo& PropDef, FFieldVariant Scope, const FName& Name, EObjectFlags ObjectFlags, EVariableCategory::Type VariableCategory, const TCHAR* Dimensions)
+	static FProperty* CreateProperty(FUnrealPropertyDefinitionInfo& PropDef, FFieldVariant Scope, const FName& Name, EObjectFlags ObjectFlags)
 	{
 		FPropertyBase& VarProperty = PropDef.GetPropertyBase();
 		check(VarProperty.ClassDef);
@@ -724,11 +1050,65 @@ struct FPropertyTypeTraitsObjectReference : public FPropertyTypeTraitsObjectBase
 	{
 		return true;
 	}
+
+	static FString GetEngineClassName(const FUnrealPropertyDefinitionInfo& PropDef)
+	{
+		const FPropertyBase& VarProperty = PropDef.GetPropertyBase();
+		check(VarProperty.ClassDef);
+		UClass* PropertyClass = VarProperty.ClassDef->GetClass();
+		if (PropertyClass->IsChildOf(UClass::StaticClass()))
+		{
+			return FClassProperty::StaticClass()->GetName();
+		}
+		else
+		{
+			return FObjectProperty::StaticClass()->GetName();
+		}
+	}
+
+	static FString GetCPPType(const FUnrealPropertyDefinitionInfo& PropDef, FString* ExtendedTypeText, uint32 CPPExportFlags)
+	{
+		const FPropertyBase& VarProperty = PropDef.GetPropertyBase();
+		check(VarProperty.ClassDef);
+		UClass* PropertyClass = VarProperty.ClassDef->GetClass();
+		if (PropertyClass->IsChildOf(UClass::StaticClass()))
+		{
+			if (PropDef.HasAnyPropertyFlags(CPF_UObjectWrapper))
+			{
+				UClass* MetaClass = VarProperty.MetaClassDef->GetClass();
+				return FString::Printf(TEXT("TSubclassOf<%s%s> "), MetaClass->GetPrefixCPP(), *MetaClass->GetName());
+			}
+			else
+			{
+				return TEXT("UClass*");
+			}
+		}
+		else
+		{
+			return FString::Printf(TEXT("%s%s*"), PropertyClass->GetPrefixCPP(), *PropertyClass->GetName());
+		}
+	}
+
+	static FString GetCPPTypeForwardDeclaration(const FUnrealPropertyDefinitionInfo& PropDef)
+	{
+		const FPropertyBase& VarProperty = PropDef.GetPropertyBase();
+		check(VarProperty.ClassDef);
+		UClass* PropertyClass = VarProperty.ClassDef->GetClass();
+		if (PropertyClass->IsChildOf(UClass::StaticClass()))
+		{
+			UClass* MetaClass = VarProperty.MetaClassDef->GetClass();
+			return FString::Printf(TEXT("class %s%s;"), MetaClass->GetPrefixCPP(), *MetaClass->GetName());
+		}
+		else
+		{
+			return FString::Printf(TEXT("class %s%s;"), PropertyClass->GetPrefixCPP(), *PropertyClass->GetName());
+		}
+	}
 };
 
 struct FPropertyTypeTraitsWeakObjectReference : public FPropertyTypeTraitsObjectBase
 {
-	static FProperty* CreateEngineType(FUnrealPropertyDefinitionInfo& PropDef, FFieldVariant Scope, const FName& Name, EObjectFlags ObjectFlags, EVariableCategory::Type VariableCategory, const TCHAR* Dimensions)
+	static FProperty* CreateProperty(FUnrealPropertyDefinitionInfo& PropDef, FFieldVariant Scope, const FName& Name, EObjectFlags ObjectFlags)
 	{
 		const FPropertyBase& VarProperty = PropDef.GetPropertyBase();
 		check(VarProperty.ClassDef);
@@ -747,11 +1127,37 @@ struct FPropertyTypeTraitsWeakObjectReference : public FPropertyTypeTraitsObject
 	{
 		return bMemberVariable;
 	}
-};
 
-struct FPropertyTypeTraitsLazyObjectReference : public FPropertyTypeTraitsObjectBase
-{
-	static FProperty* CreateEngineType(FUnrealPropertyDefinitionInfo& PropDef, FFieldVariant Scope, const FName& Name, EObjectFlags ObjectFlags, EVariableCategory::Type VariableCategory, const TCHAR* Dimensions)
+	static FString GetEngineClassName(const FUnrealPropertyDefinitionInfo& PropDef)
+	{
+		return FWeakObjectProperty::StaticClass()->GetName();
+	}
+
+	static FString GetCPPType(const FUnrealPropertyDefinitionInfo& PropDef, FString* ExtendedTypeText, uint32 CPPExportFlags)
+	{
+		const FPropertyBase& VarProperty = PropDef.GetPropertyBase();
+		check(VarProperty.ClassDef);
+		UClass* PropertyClass = VarProperty.ClassDef->GetClass();
+		FString InnerNativeTypeName = FString::Printf(TEXT("%s%s"), PropertyClass->GetPrefixCPP(), *PropertyClass->GetName());
+		if (PropDef.HasAnyPropertyFlags(CPF_AutoWeak))
+		{
+			return FString::Printf(TEXT("TAutoWeakObjectPtr<%s>"), *InnerNativeTypeName);
+		}
+		return FString::Printf(TEXT("TWeakObjectPtr<%s>"), *InnerNativeTypeName);
+	}
+
+	static FString GetCPPTypeForwardDeclaration(const FUnrealPropertyDefinitionInfo& PropDef)
+	{
+		const FPropertyBase& VarProperty = PropDef.GetPropertyBase();
+		check(VarProperty.ClassDef);
+		UClass* PropertyClass = VarProperty.ClassDef->GetClass();
+		return FString::Printf(TEXT("class %s%s;"), PropertyClass->GetPrefixCPP(), *PropertyClass->GetName());
+	}
+};
+
+struct FPropertyTypeTraitsLazyObjectReference : public FPropertyTypeTraitsObjectBase
+{
+	static FProperty* CreateProperty(FUnrealPropertyDefinitionInfo& PropDef, FFieldVariant Scope, const FName& Name, EObjectFlags ObjectFlags)
 	{
 		const FPropertyBase& VarProperty = PropDef.GetPropertyBase();
 		check(VarProperty.ClassDef);
@@ -765,11 +1171,32 @@ struct FPropertyTypeTraitsLazyObjectReference : public FPropertyTypeTraitsObject
 		Result->PropertyClass = PropertyClass;
 		return Result;
 	}
+
+	static FString GetEngineClassName(const FUnrealPropertyDefinitionInfo& PropDef)
+	{
+		return FLazyObjectProperty::StaticClass()->GetName();
+	}
+
+	static FString GetCPPType(const FUnrealPropertyDefinitionInfo& PropDef, FString* ExtendedTypeText, uint32 CPPExportFlags)
+	{
+		const FPropertyBase& VarProperty = PropDef.GetPropertyBase();
+		check(VarProperty.ClassDef);
+		UClass* PropertyClass = VarProperty.ClassDef->GetClass();
+		return FString::Printf(TEXT("TLazyObjectPtr<%s%s>"), PropertyClass->GetPrefixCPP(), *PropertyClass->GetName());
+	}
+
+	static FString GetCPPTypeForwardDeclaration(const FUnrealPropertyDefinitionInfo& PropDef)
+	{
+		const FPropertyBase& VarProperty = PropDef.GetPropertyBase();
+		check(VarProperty.ClassDef);
+		UClass* PropertyClass = VarProperty.ClassDef->GetClass();
+		return FString::Printf(TEXT("class %s%s;"), PropertyClass->GetPrefixCPP(), *PropertyClass->GetName());
+	}
 };
 
 struct FPropertyTypeTraitsObjectPtrReference : public FPropertyTypeTraitsObjectBase
 {
-	static FProperty* CreateEngineType(FUnrealPropertyDefinitionInfo& PropDef, FFieldVariant Scope, const FName& Name, EObjectFlags ObjectFlags, EVariableCategory::Type VariableCategory, const TCHAR* Dimensions)
+	static FProperty* CreateProperty(FUnrealPropertyDefinitionInfo& PropDef, FFieldVariant Scope, const FName& Name, EObjectFlags ObjectFlags)
 	{
 		FPropertyBase& VarProperty = PropDef.GetPropertyBase();
 		check(VarProperty.ClassDef);
@@ -804,11 +1231,57 @@ struct FPropertyTypeTraitsObjectPtrReference : public FPropertyTypeTraitsObjectB
 	{
 		return true;
 	}
+
+	static FString GetEngineClassName(const FUnrealPropertyDefinitionInfo& PropDef)
+	{
+		const FPropertyBase& VarProperty = PropDef.GetPropertyBase();
+		check(VarProperty.ClassDef);
+		UClass* PropertyClass = VarProperty.ClassDef->GetClass();
+		if (PropertyClass->IsChildOf(UClass::StaticClass()))
+		{
+			return FClassPtrProperty::StaticClass()->GetName();
+		}
+		else
+		{
+			return FObjectPtrProperty::StaticClass()->GetName();
+		}
+	}
+
+	static FString GetCPPType(const FUnrealPropertyDefinitionInfo& PropDef, FString* ExtendedTypeText, uint32 CPPExportFlags)
+	{
+		const FPropertyBase& VarProperty = PropDef.GetPropertyBase();
+		check(VarProperty.ClassDef);
+		UClass* PropertyClass = VarProperty.ClassDef->GetClass();
+		if (PropertyClass->IsChildOf(UClass::StaticClass()))
+		{
+			return FString::Printf(TEXT("TObjectPtr<%s%s>"), PropertyClass->GetPrefixCPP(), *PropertyClass->GetName());
+		}
+		else
+		{
+			return FString::Printf(TEXT("TObjectPtr<%s%s>"), PropertyClass->GetPrefixCPP(), *PropertyClass->GetName());
+		}
+	}
+
+	static FString GetCPPTypeForwardDeclaration(const FUnrealPropertyDefinitionInfo& PropDef)
+	{
+		const FPropertyBase& VarProperty = PropDef.GetPropertyBase();
+		check(VarProperty.ClassDef);
+		UClass* PropertyClass = VarProperty.ClassDef->GetClass();
+		if (PropertyClass->IsChildOf(UClass::StaticClass()))
+		{
+			UClass* MetaClass = VarProperty.MetaClassDef->GetClass();
+			return FString::Printf(TEXT("class %s%s;"), MetaClass->GetPrefixCPP(), *MetaClass->GetName());
+		}
+		else
+		{
+			return FString::Printf(TEXT("class %s%s;"), PropertyClass->GetPrefixCPP(), *PropertyClass->GetName());
+		}
+	}
 };
 
 struct FPropertyTypeTraitsSoftObjectReference : public FPropertyTypeTraitsObjectBase
 {
-	static FProperty* CreateEngineType(FUnrealPropertyDefinitionInfo& PropDef, FFieldVariant Scope, const FName& Name, EObjectFlags ObjectFlags, EVariableCategory::Type VariableCategory, const TCHAR* Dimensions)
+	static FProperty* CreateProperty(FUnrealPropertyDefinitionInfo& PropDef, FFieldVariant Scope, const FName& Name, EObjectFlags ObjectFlags)
 	{
 		FPropertyBase& VarProperty = PropDef.GetPropertyBase();
 		check(VarProperty.ClassDef);
@@ -835,13 +1308,60 @@ struct FPropertyTypeTraitsSoftObjectReference : public FPropertyTypeTraitsObject
 
 	static bool IsSupportedByBlueprint(const FUnrealPropertyDefinitionInfo& PropDef, bool bMemberVariable)
 	{
-		return PropDef.GetProperty()->IsA<FSoftObjectProperty>(); // Not SoftClass???
+		return true;
+	}
+
+	static FString GetEngineClassName(const FUnrealPropertyDefinitionInfo& PropDef)
+	{
+		const FPropertyBase& VarProperty = PropDef.GetPropertyBase();
+		check(VarProperty.ClassDef);
+		UClass* PropertyClass = VarProperty.ClassDef->GetClass();
+		if (PropertyClass->IsChildOf(UClass::StaticClass()))
+		{
+			return FSoftClassProperty::StaticClass()->GetName();
+		}
+		else
+		{
+			return FSoftObjectProperty::StaticClass()->GetName();
+		}
+	}
+
+	static FString GetCPPType(const FUnrealPropertyDefinitionInfo& PropDef, FString* ExtendedTypeText, uint32 CPPExportFlags)
+	{
+		const FPropertyBase& VarProperty = PropDef.GetPropertyBase();
+		check(VarProperty.ClassDef);
+		UClass* PropertyClass = VarProperty.ClassDef->GetClass();
+		if (PropertyClass->IsChildOf(UClass::StaticClass()))
+		{
+			UClass* MetaClass = VarProperty.MetaClassDef->GetClass();
+			return FString::Printf(TEXT("TSoftClassPtr<%s%s> "), MetaClass->GetPrefixCPP(), *MetaClass->GetName()); //@TODO - This has an extra space in it
+		}
+		else
+		{
+			return FString::Printf(TEXT("TSoftObjectPtr<%s%s>"), PropertyClass->GetPrefixCPP(), *PropertyClass->GetName());
+		}
+	}
+
+	static FString GetCPPTypeForwardDeclaration(const FUnrealPropertyDefinitionInfo& PropDef)
+	{
+		const FPropertyBase& VarProperty = PropDef.GetPropertyBase();
+		check(VarProperty.ClassDef);
+		UClass* PropertyClass = VarProperty.ClassDef->GetClass();
+		if (PropertyClass->IsChildOf(UClass::StaticClass()))
+		{
+			UClass* MetaClass = VarProperty.MetaClassDef->GetClass();
+			return FString::Printf(TEXT("class %s%s;"), MetaClass->GetPrefixCPP(), *MetaClass->GetName());
+		}
+		else
+		{
+			return FString::Printf(TEXT("class %s%s;"), PropertyClass->GetPrefixCPP(), *PropertyClass->GetName());
+		}
 	}
 };
 
 struct FPropertyTypeTraitsInterface : public FPropertyTypeTraitsBase
 {
-	static FProperty* CreateEngineType(FUnrealPropertyDefinitionInfo& PropDef, FFieldVariant Scope, const FName& Name, EObjectFlags ObjectFlags, EVariableCategory::Type VariableCategory, const TCHAR* Dimensions)
+	static FProperty* CreateProperty(FUnrealPropertyDefinitionInfo& PropDef, FFieldVariant Scope, const FName& Name, EObjectFlags ObjectFlags)
 	{
 		FPropertyBase& VarProperty = PropDef.GetPropertyBase();
 		check(VarProperty.ClassDef);
@@ -861,6 +1381,44 @@ struct FPropertyTypeTraitsInterface : public FPropertyTypeTraitsBase
 	{
 		return true;
 	}
+
+	static FString GetEngineClassName(const FUnrealPropertyDefinitionInfo& PropDef)
+	{
+		return FInterfaceProperty::StaticClass()->GetName();
+	}
+
+	static FString GetCPPType(const FUnrealPropertyDefinitionInfo& PropDef, FString* ExtendedTypeText, uint32 CPPExportFlags)
+	{
+		if (ExtendedTypeText != NULL)
+		{
+			UClass* ExportClass = PropDef.GetPropertyBase().ClassDef->GetClass();
+			if (0 == (CPPF_BlueprintCppBackend & CPPExportFlags))
+			{
+				while (ExportClass && !ExportClass->HasAnyClassFlags(CLASS_Native))
+				{
+					ExportClass = ExportClass->GetSuperClass();
+				}
+			}
+			check(ExportClass);
+			check(ExportClass->HasAnyClassFlags(CLASS_Interface) || 0 != (CPPF_BlueprintCppBackend & CPPExportFlags));
+
+			*ExtendedTypeText = FString::Printf(TEXT("<I%s>"), *ExportClass->GetName());
+		}
+		return TEXT("TScriptInterface");
+	}
+
+	static FString GetCPPTypeForwardDeclaration(const FUnrealPropertyDefinitionInfo& PropDef)
+	{
+		FUnrealClassDefinitionInfo* ExportClass = PropDef.GetPropertyBase().ClassDef;
+		while (ExportClass && !ExportClass->HasAnyClassFlags(CLASS_Native))
+		{
+			ExportClass = ExportClass->GetSuperClass();
+		}
+		check(ExportClass);
+		check(ExportClass->HasAnyClassFlags(CLASS_Interface));
+
+		return FString::Printf(TEXT("class I%s;"), *ExportClass->GetName());
+	}
 };
 
 //----------------------------------------------------------------------------------------------------------------------------------------------------------
@@ -881,7 +1439,7 @@ struct FPropertyTypeTraitsName : public FPropertyTypeTraitsBase
 		return FDefaultValueHelper::StringFromCppString(CppForm, TEXT("FName"), OutForm);
 	}
 
-	static FProperty* CreateEngineType(FUnrealPropertyDefinitionInfo& PropDef, FFieldVariant Scope, const FName& Name, EObjectFlags ObjectFlags, EVariableCategory::Type VariableCategory, const TCHAR* Dimensions)
+	static FProperty* CreateProperty(FUnrealPropertyDefinitionInfo& PropDef, FFieldVariant Scope, const FName& Name, EObjectFlags ObjectFlags)
 	{
 		FNameProperty* Result = new FNameProperty(Scope, Name, ObjectFlags);
 		return Result;
@@ -891,6 +1449,21 @@ struct FPropertyTypeTraitsName : public FPropertyTypeTraitsBase
 	{
 		return true;
 	}
+
+	static FString GetEngineClassName(const FUnrealPropertyDefinitionInfo& PropDef)
+	{
+		return FNameProperty::StaticClass()->GetName();
+	}
+
+	static FString GetCPPType(const FUnrealPropertyDefinitionInfo& PropDef, FString* ExtendedTypeText, uint32 CPPExportFlags)
+	{
+		return FString(TEXT("FName"));
+	}
+
+	static FString GetCPPTypeForwardDeclaration(const FUnrealPropertyDefinitionInfo& PropDef)
+	{
+		return FString();
+	}
 };
 
 struct FPropertyTypeTraitsString : public FPropertyTypeTraitsBase
@@ -900,7 +1473,7 @@ struct FPropertyTypeTraitsString : public FPropertyTypeTraitsBase
 		return FDefaultValueHelper::StringFromCppString(CppForm, TEXT("FString"), OutForm);
 	}
 
-	static FProperty* CreateEngineType(FUnrealPropertyDefinitionInfo& PropDef, FFieldVariant Scope, const FName& Name, EObjectFlags ObjectFlags, EVariableCategory::Type VariableCategory, const TCHAR* Dimensions)
+	static FProperty* CreateProperty(FUnrealPropertyDefinitionInfo& PropDef, FFieldVariant Scope, const FName& Name, EObjectFlags ObjectFlags)
 	{
 		FStrProperty* Result = new FStrProperty(Scope, Name, ObjectFlags);
 		return Result;
@@ -910,6 +1483,21 @@ struct FPropertyTypeTraitsString : public FPropertyTypeTraitsBase
 	{
 		return true;
 	}
+
+	static FString GetEngineClassName(const FUnrealPropertyDefinitionInfo& PropDef)
+	{
+		return FStrProperty::StaticClass()->GetName();
+	}
+
+	static FString GetCPPType(const FUnrealPropertyDefinitionInfo& PropDef, FString* ExtendedTypeText, uint32 CPPExportFlags)
+	{
+		return FString(TEXT("FString"));
+	}
+
+	static FString GetCPPTypeForwardDeclaration(const FUnrealPropertyDefinitionInfo& PropDef)
+	{
+		return FString();
+	}
 };
 
 struct FPropertyTypeTraitsText : public FPropertyTypeTraitsBase
@@ -946,7 +1534,7 @@ struct FPropertyTypeTraitsText : public FPropertyTypeTraitsBase
 			{
 				if (ParsedTextNamespace.GetValue().Equals(UHTDummyNamespace))
 				{
-					FError::Throwf(TEXT("LOCTEXT default parameter values are not supported; use NSLOCTEXT instead: %s \"%s\" "), *PropDef.GetProperty()->GetName(), *CppForm);
+					FError::Throwf(TEXT("LOCTEXT default parameter values are not supported; use NSLOCTEXT instead: %s \"%s\" "), *PropDef.GetName(), *CppForm);
 				}
 			}
 		}
@@ -956,7 +1544,7 @@ struct FPropertyTypeTraitsText : public FPropertyTypeTraitsBase
 		return true;
 	}
 
-	static FProperty* CreateEngineType(FUnrealPropertyDefinitionInfo& PropDef, FFieldVariant Scope, const FName& Name, EObjectFlags ObjectFlags, EVariableCategory::Type VariableCategory, const TCHAR* Dimensions)
+	static FProperty* CreateProperty(FUnrealPropertyDefinitionInfo& PropDef, FFieldVariant Scope, const FName& Name, EObjectFlags ObjectFlags)
 	{
 		FTextProperty* Result = new FTextProperty(Scope, Name, ObjectFlags);
 		return Result;
@@ -966,6 +1554,21 @@ struct FPropertyTypeTraitsText : public FPropertyTypeTraitsBase
 	{
 		return true;
 	}
+
+	static FString GetEngineClassName(const FUnrealPropertyDefinitionInfo& PropDef)
+	{
+		return FTextProperty::StaticClass()->GetName();
+	}
+
+	static FString GetCPPType(const FUnrealPropertyDefinitionInfo& PropDef, FString* ExtendedTypeText, uint32 CPPExportFlags)
+	{
+		return FString(TEXT("FText"));
+	}
+
+	static FString GetCPPTypeForwardDeclaration(const FUnrealPropertyDefinitionInfo& PropDef)
+	{
+		return FString();
+	}
 };
 
 struct FPropertyTypeTraitsStruct : public FPropertyTypeTraitsBase
@@ -1174,7 +1777,7 @@ struct FPropertyTypeTraitsStruct : public FPropertyTypeTraitsBase
 		return !OutForm.IsEmpty();
 	}
 
-	static FProperty* CreateEngineType(FUnrealPropertyDefinitionInfo& PropDef, FFieldVariant Scope, const FName& Name, EObjectFlags ObjectFlags, EVariableCategory::Type VariableCategory, const TCHAR* Dimensions)
+	static FProperty* CreateProperty(FUnrealPropertyDefinitionInfo& PropDef, FFieldVariant Scope, const FName& Name, EObjectFlags ObjectFlags)
 	{
 		FPropertyBase& VarProperty = PropDef.GetPropertyBase();
 		UScriptStruct* Struct = PropDef.GetPropertyBase().ScriptStructDef->GetScriptStruct();
@@ -1195,13 +1798,37 @@ struct FPropertyTypeTraitsStruct : public FPropertyTypeTraitsBase
 
 	static bool IsSupportedByBlueprint(const FUnrealPropertyDefinitionInfo& PropDef, bool bMemberVariable)
 	{
-		return CastField<const FStructProperty>(PropDef.GetProperty())->Struct->GetBoolMetaDataHierarchical(FHeaderParserNames::NAME_BlueprintType);
+		const FPropertyBase& VarProperty = PropDef.GetPropertyBase();
+		UScriptStruct* Struct = PropDef.GetPropertyBase().ScriptStructDef->GetScriptStruct();
+		return Struct->GetBoolMetaDataHierarchical(FHeaderParserNames::NAME_BlueprintType);
+	}
+
+	static FString GetEngineClassName(const FUnrealPropertyDefinitionInfo& PropDef)
+	{
+		return FStructProperty::StaticClass()->GetName();
+	}
+
+	static FString GetCPPType(const FUnrealPropertyDefinitionInfo& PropDef, FString* ExtendedTypeText, uint32 CPPExportFlags)
+	{
+		return FString::Printf(TEXT("F%s"), *PropDef.GetPropertyBase().ClassDef->GetName());
+	}
+
+	static FString GetCPPTypeForwardDeclaration(const FUnrealPropertyDefinitionInfo& PropDef)
+	{
+		// Core type structs don't need to forward declare in UHT as every generated.h indirectly includes CoreMinimal.h
+		const UScriptStruct* Struct = PropDef.GetPropertyBase().ScriptStructDef->GetScriptStruct();
+		if (Struct->GetCppStructOps() && Struct->GetCppStructOps()->IsUECoreType())
+		{
+			return FString();
+		}
+			
+		return FString::Printf(TEXT("struct F%s;"), *PropDef.GetPropertyBase().ClassDef->GetName());
 	}
 };
 
 struct FPropertyTypeTraitsDelegate : public FPropertyTypeTraitsBase
 {
-	static FProperty* CreateEngineType(FUnrealPropertyDefinitionInfo& PropDef, FFieldVariant Scope, const FName& Name, EObjectFlags ObjectFlags, EVariableCategory::Type VariableCategory, const TCHAR* Dimensions)
+	static FProperty* CreateProperty(FUnrealPropertyDefinitionInfo& PropDef, FFieldVariant Scope, const FName& Name, EObjectFlags ObjectFlags)
 	{
 		FDelegateProperty* Result = new FDelegateProperty(Scope, Name, ObjectFlags);
 		return Result;
@@ -1211,11 +1838,62 @@ struct FPropertyTypeTraitsDelegate : public FPropertyTypeTraitsBase
 	{
 		return true;
 	}
+
+	static FString GetEngineClassName(const FUnrealPropertyDefinitionInfo& PropDef)
+	{
+		return FDelegateProperty::StaticClass()->GetName();
+	}
+
+	static FString GetCPPType(const FUnrealPropertyDefinitionInfo& PropDef, FString* ExtendedTypeText, uint32 CPPExportFlags)
+	{
+		const FPropertyBase& VarProperty = PropDef.GetPropertyBase();
+		check(VarProperty.FunctionDef);
+		UFunction* SignatureFunction = VarProperty.FunctionDef->GetFunction();
+		check(SignatureFunction);
+
+		FString UnmangledFunctionName = SignatureFunction->GetName().LeftChop(FString(HEADER_GENERATED_DELEGATE_SIGNATURE_SUFFIX).Len());
+		const bool bBlueprintCppBackend = (0 != (CPPExportFlags & EPropertyExportCPPFlags::CPPF_BlueprintCppBackend));
+		const bool bNative = SignatureFunction->IsNative();
+		if (bBlueprintCppBackend && bNative)
+		{
+			UStruct* StructOwner = Cast<UStruct>(SignatureFunction->GetOuter());
+			if (StructOwner)
+			{
+				return FString::Printf(TEXT("%s%s::F%s"), StructOwner->GetPrefixCPP(), *StructOwner->GetName(), *UnmangledFunctionName);
+			}
+		}
+		else
+		{
+			const bool NonNativeClassOwner = (SignatureFunction->GetOwnerClass() && !SignatureFunction->GetOwnerClass()->HasAnyClassFlags(CLASS_Native));
+			if (bBlueprintCppBackend && NonNativeClassOwner)
+			{
+				// The name must be valid, this removes spaces, ?, etc from the user's function name. It could
+				// be slightly shorter because the postfix ("__pf") is not needed here because we further post-
+				// pend to the string. Normally the postfix is needed to make sure we don't mangle to a valid
+				// identifier and collide:
+				UnmangledFunctionName = UnicodeToCPPIdentifier(UnmangledFunctionName, false, TEXT(""));
+				// the name must be unique
+				const FString OwnerName = UnicodeToCPPIdentifier(SignatureFunction->GetOwnerClass()->GetName(), false, TEXT(""));
+				const FString NewUnmangledFunctionName = FString::Printf(TEXT("%s__%s"), *UnmangledFunctionName, *OwnerName);
+				UnmangledFunctionName = NewUnmangledFunctionName;
+			}
+			if (0 != (CPPExportFlags & EPropertyExportCPPFlags::CPPF_CustomTypeName))
+			{
+				UnmangledFunctionName += TEXT("__SinglecastDelegate");
+			}
+		}
+		return FString(TEXT("F")) + UnmangledFunctionName;
+	}
+
+	static FString GetCPPTypeForwardDeclaration(const FUnrealPropertyDefinitionInfo& PropDef)
+	{
+		return FString();
+	}
 };
 
 struct FPropertyTypeTraitsMulticastDelegate : public FPropertyTypeTraitsBase
 {
-	static FProperty* CreateEngineType(FUnrealPropertyDefinitionInfo& PropDef, FFieldVariant Scope, const FName& Name, EObjectFlags ObjectFlags, EVariableCategory::Type VariableCategory, const TCHAR* Dimensions)
+	static FProperty* CreateProperty(FUnrealPropertyDefinitionInfo& PropDef, FFieldVariant Scope, const FName& Name, EObjectFlags ObjectFlags)
 	{
 		const FPropertyBase& VarProperty = PropDef.GetPropertyBase();
 		UFunction* Function = PropDef.GetPropertyBase().FunctionDef->GetFunction();
@@ -1236,11 +1914,78 @@ struct FPropertyTypeTraitsMulticastDelegate : public FPropertyTypeTraitsBase
 	{
 		return bMemberVariable;
 	}
+
+	static FString GetEngineClassName(const FUnrealPropertyDefinitionInfo& PropDef)
+	{
+		const FPropertyBase& VarProperty = PropDef.GetPropertyBase();
+		UFunction* Function = PropDef.GetPropertyBase().FunctionDef->GetFunction();
+		if (Function->IsA<USparseDelegateFunction>())
+		{
+			return FMulticastSparseDelegateProperty::StaticClass()->GetName();
+		}
+		else
+		{
+			return FMulticastInlineDelegateProperty::StaticClass()->GetName();
+		}
+	}
+
+	static FString GetCPPType(const FUnrealPropertyDefinitionInfo& PropDef, FString* ExtendedTypeText, uint32 CPPExportFlags)
+	{
+		const FPropertyBase& VarProperty = PropDef.GetPropertyBase();
+		check(VarProperty.FunctionDef);
+		UFunction* SignatureFunction = VarProperty.FunctionDef->GetFunction();
+
+		// We have this test because sometimes the delegate hasn't been set up by FixupDelegateProperties at the time
+		// we need the type for an error message.  We deliberately format it so that it's unambiguously not CPP code, but is still human-readable.
+		if (!SignatureFunction)
+		{
+			return FString(TEXT("{multicast delegate type}"));
+		}
+
+		FString UnmangledFunctionName = SignatureFunction->GetName().LeftChop(FString(HEADER_GENERATED_DELEGATE_SIGNATURE_SUFFIX).Len());
+		const UClass* OwnerClass = SignatureFunction->GetOwnerClass();
+
+		const bool bBlueprintCppBackend = (0 != (CPPExportFlags & EPropertyExportCPPFlags::CPPF_BlueprintCppBackend));
+		const bool bNative = SignatureFunction->IsNative();
+		if (bBlueprintCppBackend && bNative)
+		{
+			UStruct* StructOwner = Cast<UStruct>(SignatureFunction->GetOuter());
+			if (StructOwner)
+			{
+				return FString::Printf(TEXT("%s%s::F%s"), StructOwner->GetPrefixCPP(), *StructOwner->GetName(), *UnmangledFunctionName);
+			}
+		}
+		else
+		{
+			if ((0 != (CPPExportFlags & EPropertyExportCPPFlags::CPPF_BlueprintCppBackend)) && OwnerClass && !OwnerClass->HasAnyClassFlags(CLASS_Native))
+			{
+				// The name must be valid, this removes spaces, ?, etc from the user's function name. It could
+				// be slightly shorter because the postfix ("__pf") is not needed here because we further post-
+				// pend to the string. Normally the postfix is needed to make sure we don't mangle to a valid
+				// identifier and collide:
+				UnmangledFunctionName = UnicodeToCPPIdentifier(UnmangledFunctionName, false, TEXT(""));
+				// the name must be unique
+				const FString OwnerName = UnicodeToCPPIdentifier(OwnerClass->GetName(), false, TEXT(""));
+				const FString NewUnmangledFunctionName = FString::Printf(TEXT("%s__%s"), *UnmangledFunctionName, *OwnerName);
+				UnmangledFunctionName = NewUnmangledFunctionName;
+			}
+			if (0 != (CPPExportFlags & EPropertyExportCPPFlags::CPPF_CustomTypeName))
+			{
+				UnmangledFunctionName += TEXT("__MulticastDelegate");
+			}
+		}
+		return FString(TEXT("F")) + UnmangledFunctionName;
+	}
+
+	static FString GetCPPTypeForwardDeclaration(const FUnrealPropertyDefinitionInfo& PropDef)
+	{
+		return FString();
+	}
 };
 
 struct FPropertyTypeTraitsFieldPath : public FPropertyTypeTraitsBase
 {
-	static FProperty* CreateEngineType(FUnrealPropertyDefinitionInfo& PropDef, FFieldVariant Scope, const FName& Name, EObjectFlags ObjectFlags, EVariableCategory::Type VariableCategory, const TCHAR* Dimensions)
+	static FProperty* CreateProperty(FUnrealPropertyDefinitionInfo& PropDef, FFieldVariant Scope, const FName& Name, EObjectFlags ObjectFlags)
 	{
 		const FPropertyBase& VarProperty = PropDef.GetPropertyBase();
 		FFieldPathProperty* Result = new FFieldPathProperty(Scope, Name, ObjectFlags);
@@ -1252,6 +1997,29 @@ struct FPropertyTypeTraitsFieldPath : public FPropertyTypeTraitsBase
 	{
 		return true;
 	}
+
+	static FString GetEngineClassName(const FUnrealPropertyDefinitionInfo& PropDef)
+	{
+		return FFieldPathProperty::StaticClass()->GetName();
+	}
+
+	static FString GetCPPType(const FUnrealPropertyDefinitionInfo& PropDef, FString* ExtendedTypeText, uint32 CPPExportFlags)
+	{
+		if (ExtendedTypeText != nullptr)
+		{
+			const FPropertyBase& VarProperty = PropDef.GetPropertyBase();
+			FString& InnerTypeText = *ExtendedTypeText;
+			InnerTypeText = TEXT("<F");
+			InnerTypeText += VarProperty.PropertyPathClass->GetName();
+			InnerTypeText += TEXT(">");
+		}
+		return TEXT("TFieldPath");
+	}
+
+	static FString GetCPPTypeForwardDeclaration(const FUnrealPropertyDefinitionInfo& PropDef)
+	{
+		return FString::Printf(TEXT("class F%s;"), *PropDef.GetPropertyBase().PropertyPathClass->GetName());
+	}
 };
 
 //----------------------------------------------------------------------------------------------------------------------------------------------------------
@@ -1262,11 +2030,10 @@ struct FPropertyTypeTraitsFieldPath : public FPropertyTypeTraitsBase
 
 struct FPropertyTypeTraitsStaticArray : public FPropertyTypeTraitsBase
 {
-	static FProperty* CreateEngineType(FUnrealPropertyDefinitionInfo& PropDef, FFieldVariant Scope, const FName& Name, EObjectFlags ObjectFlags, EVariableCategory::Type VariableCategory, const TCHAR* Dimensions)
+	static FProperty* CreateProperty(FUnrealPropertyDefinitionInfo& PropDef, FFieldVariant Scope, const FName& Name, EObjectFlags ObjectFlags)
 	{
-		FProperty* Property = CreatePropertyHelper<false>(PropDef, Scope, Name, ObjectFlags, VariableCategory, Dimensions);
+		FProperty* Property = CreatePropertyHelper<false>(PropDef, Scope, Name, ObjectFlags);
 		Property->ArrayDim = 2;
-		PropDef.SetArrayDimensions(Dimensions);
 		return Property;
 	}
 
@@ -1274,11 +2041,26 @@ struct FPropertyTypeTraitsStaticArray : public FPropertyTypeTraitsBase
 	{
 		return IsSupportedByBlueprintSansContainers(PropDef, bMemberVariable);
 	}
+
+	static FString GetEngineClassName(const FUnrealPropertyDefinitionInfo& PropDef)
+	{
+		return PropertyTypeDispatch<GetEngineClassNameDispatch, false, FString>(PropDef.GetPropertyBase(), std::ref(PropDef));
+	}
+
+	static FString GetCPPType(const FUnrealPropertyDefinitionInfo& PropDef, FString* ExtendedTypeText, uint32 CPPExportFlags)
+	{
+		return PropertyTypeDispatch<GetCPPTypeDispatch, false, FString>(PropDef.GetPropertyBase(), std::ref(PropDef), ExtendedTypeText, CPPExportFlags);
+	}
+
+	static FString GetCPPTypeForwardDeclaration(const FUnrealPropertyDefinitionInfo& PropDef)
+	{
+		return PropertyTypeDispatch<GetCPPTypeForwardDeclarationDispatch, false, FString>(PropDef.GetPropertyBase(), std::ref(PropDef));
+	}
 };
 
 struct FPropertyTypeTraitsDynamicArray : public FPropertyTypeTraitsBase
 {
-	static FProperty* CreateEngineType(FUnrealPropertyDefinitionInfo& PropDef, FFieldVariant Scope, const FName& Name, EObjectFlags ObjectFlags, EVariableCategory::Type VariableCategory, const TCHAR* Dimensions)
+	static FProperty* CreateProperty(FUnrealPropertyDefinitionInfo& PropDef, FFieldVariant Scope, const FName& Name, EObjectFlags ObjectFlags)
 	{
 		FPropertyBase& VarProperty = PropDef.GetPropertyBase();
 
@@ -1287,12 +2069,15 @@ struct FPropertyTypeTraitsDynamicArray : public FPropertyTypeTraitsBase
 
 		FPropertyBase InnerVarProperty = PropDef.GetPropertyBase();
 		InnerVarProperty.ArrayType = EArrayType::None;
-		FUnrealPropertyDefinitionInfo& InnerPropDef = FPropertyTraits::CreateProperty(InnerVarProperty, PropDef, Name, RF_Public, VariableCategory, Dimensions, PropDef.GetUnrealSourceFile(), PropDef.GetLineNumber(), PropDef.GetParsePosition());
+		InnerVarProperty.RepNotifyName = NAME_None;
+		InnerVarProperty.MetaData.Empty();
+		FUnrealPropertyDefinitionInfo& InnerPropDef = FPropertyTraits::CreateProperty(InnerVarProperty, PropDef, Name, RF_Public, 
+			PropDef.GetVariableCategory(), PropDef.GetArrayDimensions(), PropDef.GetUnrealSourceFile(), PropDef.GetLineNumber(), PropDef.GetParsePosition());
 		FProperty* InnerProp = InnerPropDef.GetProperty();
 
 		Array->Inner = InnerProp;
 		VarProperty.PropertyFlags = Array->Inner->PropertyFlags;
-		VarProperty.MetaData = MoveTemp(InnerPropDef.GetPropertyBase().MetaData);
+		VarProperty.MetaData.Append(MoveTemp(InnerPropDef.GetPropertyBase().MetaData)); // Move any added meta data items to the container
 		PropDef.SetAllocatorType(VarProperty.AllocatorType);
 		PropDef.SetValuePropDef(InnerPropDef);
 
@@ -1307,11 +2092,42 @@ struct FPropertyTypeTraitsDynamicArray : public FPropertyTypeTraitsBase
 	{
 		return IsSupportedByBlueprintSansContainers(PropDef.GetValuePropDef(), false);
 	}
+
+	static FString GetEngineClassName(const FUnrealPropertyDefinitionInfo& PropDef)
+	{
+		return FArrayProperty::StaticClass()->GetName();
+	}
+
+	static FString GetCPPType(const FUnrealPropertyDefinitionInfo& PropDef, FString* ExtendedTypeText, uint32 CPPExportFlags)
+	{
+		if (ExtendedTypeText != NULL)
+		{
+			FString InnerExtendedTypeText;
+			FString InnerTypeText = PropDef.GetValuePropDef().GetCPPType(&InnerExtendedTypeText, CPPExportFlags & ~CPPF_ArgumentOrReturnValue); // we won't consider array inners to be "arguments or return values"
+			if (InnerExtendedTypeText.Len() && InnerExtendedTypeText.Right(1) == TEXT(">"))
+			{
+				// if our internal property type is a template class, add a space between the closing brackets b/c VS.NET cannot parse this correctly
+				InnerExtendedTypeText += TEXT(" ");
+			}
+			else if (!InnerExtendedTypeText.Len() && InnerTypeText.Len() && InnerTypeText.Right(1) == TEXT(">"))
+			{
+				// if our internal property type is a template class, add a space between the closing brackets b/c VS.NET cannot parse this correctly
+				InnerExtendedTypeText += TEXT(" ");
+			}
+			*ExtendedTypeText = FString::Printf(TEXT("<%s%s>"), *InnerTypeText, *InnerExtendedTypeText);
+		}
+		return TEXT("TArray");
+	}
+
+	static FString GetCPPTypeForwardDeclaration(const FUnrealPropertyDefinitionInfo& PropDef)
+	{
+		return PropDef.GetValuePropDef().GetCPPTypeForwardDeclaration();
+	}
 };
 
 struct FPropertyTypeTraitsSet : public FPropertyTypeTraitsBase
 {
-	static FProperty* CreateEngineType(FUnrealPropertyDefinitionInfo& PropDef, FFieldVariant Scope, const FName& Name, EObjectFlags ObjectFlags, EVariableCategory::Type VariableCategory, const TCHAR* Dimensions)
+	static FProperty* CreateProperty(FUnrealPropertyDefinitionInfo& PropDef, FFieldVariant Scope, const FName& Name, EObjectFlags ObjectFlags)
 	{
 		FPropertyBase& VarProperty = PropDef.GetPropertyBase();
 
@@ -1320,12 +2136,15 @@ struct FPropertyTypeTraitsSet : public FPropertyTypeTraitsBase
 
 		FPropertyBase InnerVarProperty = PropDef.GetPropertyBase();
 		InnerVarProperty.ArrayType = EArrayType::None;
-		FUnrealPropertyDefinitionInfo& InnerPropDef = FPropertyTraits::CreateProperty(InnerVarProperty, PropDef, Name, RF_Public, VariableCategory, Dimensions, PropDef.GetUnrealSourceFile(), PropDef.GetLineNumber(), PropDef.GetParsePosition());
+		InnerVarProperty.RepNotifyName = NAME_None;
+		InnerVarProperty.MetaData.Empty();
+		FUnrealPropertyDefinitionInfo& InnerPropDef = FPropertyTraits::CreateProperty(InnerVarProperty, PropDef, Name, RF_Public, 
+			PropDef.GetVariableCategory(), PropDef.GetArrayDimensions(), PropDef.GetUnrealSourceFile(), PropDef.GetLineNumber(), PropDef.GetParsePosition());
 		FProperty* InnerProp = InnerPropDef.GetProperty();
 
 		Set->ElementProp = InnerProp;
 		VarProperty.PropertyFlags = InnerProp->PropertyFlags;
-		VarProperty.MetaData = MoveTemp(InnerPropDef.GetPropertyBase().MetaData);
+		VarProperty.MetaData.Append(MoveTemp(InnerPropDef.GetPropertyBase().MetaData)); // Move any added meta data items to the container
 		PropDef.SetValuePropDef(InnerPropDef);
 
 		// Propagate flags
@@ -1339,29 +2158,63 @@ struct FPropertyTypeTraitsSet : public FPropertyTypeTraitsBase
 	{
 		return IsSupportedByBlueprintSansContainers(PropDef.GetValuePropDef(), false);
 	}
+
+	static FString GetEngineClassName(const FUnrealPropertyDefinitionInfo& PropDef)
+	{
+		return FSetProperty::StaticClass()->GetName();
+	}
+
+	static FString GetCPPType(const FUnrealPropertyDefinitionInfo& PropDef, FString* ExtendedTypeText, uint32 CPPExportFlags)
+	{
+		if (ExtendedTypeText != NULL)
+		{
+			FString ElementExtendedTypeText;
+			FString ElementTypeText = PropDef.GetValuePropDef().GetCPPType(&ElementExtendedTypeText, CPPExportFlags & ~CPPF_ArgumentOrReturnValue); // we won't consider set values to be "arguments or return values"
+
+			// if property type is a template class, add a space between the closing brackets
+			if ((ElementExtendedTypeText.Len() && ElementExtendedTypeText.Right(1) == TEXT(">"))
+				|| (!ElementExtendedTypeText.Len() && ElementTypeText.Len() && ElementTypeText.Right(1) == TEXT(">")))
+			{
+				ElementExtendedTypeText += TEXT(" ");
+			}
+
+			*ExtendedTypeText = FString::Printf(TEXT("<%s%s>"), *ElementTypeText, *ElementExtendedTypeText);
+		}
+		return TEXT("TSet");
+	}
+
+	static FString GetCPPTypeForwardDeclaration(const FUnrealPropertyDefinitionInfo& PropDef)
+	{
+		return PropDef.GetValuePropDef().GetCPPTypeForwardDeclaration();
+	}
 };
 
 struct FPropertyTypeTraitsMap : public FPropertyTypeTraitsBase
 {
-	static FProperty* CreateEngineType(FUnrealPropertyDefinitionInfo& PropDef, FFieldVariant Scope, const FName& Name, EObjectFlags ObjectFlags, EVariableCategory::Type VariableCategory, const TCHAR* Dimensions)
+	static FProperty* CreateProperty(FUnrealPropertyDefinitionInfo& PropDef, FFieldVariant Scope, const FName& Name, EObjectFlags ObjectFlags)
 	{
 		FPropertyBase& VarProperty = PropDef.GetPropertyBase();
 
 		FMapProperty* Map = new FMapProperty(Scope, Name, ObjectFlags);
 		PropDef.SetProperty(Map);
 
-		FUnrealPropertyDefinitionInfo& KeyPropDef = FPropertyTraits::CreateProperty(*PropDef.GetPropertyBase().MapKeyProp, PropDef, *(Name.ToString() + TEXT("_Key")), RF_Public, VariableCategory, Dimensions, PropDef.GetUnrealSourceFile(), PropDef.GetLineNumber(), PropDef.GetParsePosition());
+		FUnrealPropertyDefinitionInfo& KeyPropDef = FPropertyTraits::CreateProperty(*PropDef.GetPropertyBase().MapKeyProp, PropDef, *(Name.ToString() + TEXT("_Key")), RF_Public,
+			PropDef.GetVariableCategory(), PropDef.GetArrayDimensions(), PropDef.GetUnrealSourceFile(), PropDef.GetLineNumber(), PropDef.GetParsePosition());
 		FProperty* KeyProp = KeyPropDef.GetProperty();
+
 		FPropertyBase ValueVarProperty = PropDef.GetPropertyBase();
 		ValueVarProperty.ArrayType = EArrayType::None;
 		ValueVarProperty.MapKeyProp = nullptr;
-		FUnrealPropertyDefinitionInfo& ValuePropDef = FPropertyTraits::CreateProperty(ValueVarProperty, PropDef, Name, RF_Public, VariableCategory, Dimensions, PropDef.GetUnrealSourceFile(), PropDef.GetLineNumber(), PropDef.GetParsePosition());
+		ValueVarProperty.RepNotifyName = NAME_None;
+		ValueVarProperty.MetaData.Empty();
+		FUnrealPropertyDefinitionInfo& ValuePropDef = FPropertyTraits::CreateProperty(ValueVarProperty, PropDef, Name, RF_Public,
+			PropDef.GetVariableCategory(), PropDef.GetArrayDimensions(), PropDef.GetUnrealSourceFile(), PropDef.GetLineNumber(), PropDef.GetParsePosition());
 		FProperty* ValueProp = ValuePropDef.GetProperty();
 
 		Map->KeyProp = KeyProp;
 		Map->ValueProp = ValueProp;
 		VarProperty.PropertyFlags = ValueProp->PropertyFlags;
-		VarProperty.MetaData = MoveTemp(ValuePropDef.GetPropertyBase().MetaData);
+		VarProperty.MetaData.Append(MoveTemp(ValuePropDef.GetPropertyBase().MetaData)); // Move any added meta data items to the container
 		PropDef.SetAllocatorType(VarProperty.AllocatorType);
 		PropDef.SetKeyPropDef(KeyPropDef);
 		PropDef.SetValuePropDef(ValuePropDef);
@@ -1381,6 +2234,45 @@ struct FPropertyTypeTraitsMap : public FPropertyTypeTraitsBase
 			IsSupportedByBlueprintSansContainers(PropDef.GetValuePropDef(), false) &&
 			IsSupportedByBlueprintSansContainers(PropDef.GetKeyPropDef(), false);
 	}
+
+	static FString GetEngineClassName(const FUnrealPropertyDefinitionInfo& PropDef)
+	{
+		return FMapProperty::StaticClass()->GetName();
+	}
+
+	static FString GetCPPType(const FUnrealPropertyDefinitionInfo& PropDef, FString* ExtendedTypeText, uint32 CPPExportFlags)
+	{
+		if (ExtendedTypeText != NULL)
+		{
+			FString KeyExtendedTypeText;
+			FString ValueExtendedTypeText;
+
+			FString KeyTypeText = PropDef.GetKeyPropDef().GetCPPType(&KeyExtendedTypeText, CPPExportFlags & ~CPPF_ArgumentOrReturnValue); // we won't consider map keys to be "arguments or return values"
+			FString ValueTypeText = PropDef.GetValuePropDef().GetCPPType(&ValueExtendedTypeText, CPPExportFlags & ~CPPF_ArgumentOrReturnValue); // we won't consider map values to be "arguments or return values"
+
+			// if property type is a template class, add a space between the closing brackets
+			if ((KeyExtendedTypeText.Len() && KeyExtendedTypeText.Right(1) == TEXT(">"))
+				|| (!KeyExtendedTypeText.Len() && KeyTypeText.Len() && KeyTypeText.Right(1) == TEXT(">")))
+			{
+				KeyExtendedTypeText += TEXT(" ");
+			}
+
+			// if property type is a template class, add a space between the closing brackets
+			if ((ValueExtendedTypeText.Len() && ValueExtendedTypeText.Right(1) == TEXT(">"))
+				|| (!ValueExtendedTypeText.Len() && ValueTypeText.Len() && ValueTypeText.Right(1) == TEXT(">")))
+			{
+				ValueExtendedTypeText += TEXT(" ");
+			}
+
+			*ExtendedTypeText = FString::Printf(TEXT("<%s%s,%s%s>"), *KeyTypeText, *KeyExtendedTypeText, *ValueTypeText, *ValueExtendedTypeText);
+		}
+		return TEXT("TMap");
+	}
+
+	static FString GetCPPTypeForwardDeclaration(const FUnrealPropertyDefinitionInfo& PropDef)
+	{
+		return FString::Printf(TEXT("%s %s"), *PropDef.GetKeyPropDef().GetCPPTypeForwardDeclaration(), *PropDef.GetValuePropDef().GetCPPTypeForwardDeclaration());
+	}
 };
 
 //----------------------------------------------------------------------------------------------------------------------------------------------------------
@@ -1400,16 +2292,45 @@ bool FPropertyTraits::DefaultValueStringCppFormatToInnerFormat(const FUnrealProp
 	return PropertyTypeDispatch<DefaultValueStringCppFormatToInnerFormatDispatch, false, bool>(PropDef.GetPropertyBase(), std::ref(PropDef), std::ref(CppForm), std::ref(OutForm));
 }
 
-FUnrealPropertyDefinitionInfo& FPropertyTraits::CreateProperty(const FPropertyBase& VarProperty, FUnrealTypeDefinitionInfo& Outer, const FName& Name, EObjectFlags ObjectFlags, EVariableCategory::Type VariableCategory,
+FUnrealPropertyDefinitionInfo& FPropertyTraits::CreateProperty(const FPropertyBase& VarProperty, FUnrealTypeDefinitionInfo& Outer, const FName& Name, EObjectFlags ObjectFlags, EVariableCategory VariableCategory,
 	const TCHAR* Dimensions, FUnrealSourceFile& SourceFile, int LineNumber, int ParsePosition)
 {
-	TSharedRef<FUnrealPropertyDefinitionInfo> PropDefRef = MakeShared<FUnrealPropertyDefinitionInfo>(SourceFile, LineNumber, ParsePosition, VarProperty, Name.ToString(), Outer);
+	TSharedRef<FUnrealPropertyDefinitionInfo> PropDefRef = MakeShared<FUnrealPropertyDefinitionInfo>(SourceFile, LineNumber, ParsePosition, VarProperty, VariableCategory, Dimensions, Name, Outer);
 	FUnrealPropertyDefinitionInfo& PropDef = *PropDefRef;
 
+	// Create the property and attach to the definition
 	FFieldVariant Scope = Outer.AsProperty() ? FFieldVariant(Outer.AsProperty()->GetProperty()) : FFieldVariant(Outer.AsObject()->GetObject());
-	FProperty* Property = CreatePropertyHelper<true>(PropDef, Scope, Name, ObjectFlags, VariableCategory, Dimensions);
-
+	FProperty* Property = CreatePropertyHelper<true>(PropDef, Scope, Name, ObjectFlags);
 	PropDef.SetProperty(Property);
+
+	// Perform some final initialization from the property base data
+	Property->PropertyFlags = PropDef.GetPropertyBase().PropertyFlags;
+
+	// Special initialization for member variables
+	if (VariableCategory == EVariableCategory::Member)
+	{
+		if (PropDef.HasAnyPropertyFlags(CPF_RepNotify))
+		{
+			Property->RepNotifyFunc = PropDef.GetPropertyBase().RepNotifyName;
+		}
+	}
+
+	// If we aren't a sub-property, then add any meta data to the property
+	// We do this because of the edit inline metadata which might be added.
+	// Technically, it should be added to the outermost property.
+	if (UHTCast<FUnrealPropertyDefinitionInfo>(Outer) == nullptr)
+	{
+		FUHTMetaData::RemapAndAddMetaData(PropDef, PropDef.GetPropertyBase().MetaData);
+	}
+
+	//FString a = GetCPPTypeForwardDeclaration(PropDef);
+	//FString b = PropDef.GetProperty()->GetCPPTypeForwardDeclaration();
+	//if (a != b)
+	//{
+	//	GetCPPType(PropDef, nullptr, 0);
+	//}
+
+	// Add the property to the lookup map
 	GTypeDefinitionInfoMap.Add(PropDef.GetProperty(), MoveTemp(PropDefRef));
 	return PropDef;
 }
@@ -1418,3 +2339,18 @@ bool FPropertyTraits::IsSupportedByBlueprint(const FUnrealPropertyDefinitionInfo
 {
 	return PropertyTypeDispatch<IsSupportedByBlueprintDispatch, true, bool>(PropDef.GetPropertyBase(), std::ref(PropDef), bMemberVariable);
 }
+
+FString FPropertyTraits::GetEngineClassName(const FUnrealPropertyDefinitionInfo& PropDef)
+{
+	return PropertyTypeDispatch<GetEngineClassNameDispatch, true, FString>(PropDef.GetPropertyBase(), std::ref(PropDef));
+}
+
+FString FPropertyTraits::GetCPPType(const FUnrealPropertyDefinitionInfo& PropDef, FString* ExtendedTypeText/* = nullptr */, uint32 CPPExportFlags/* = 0 */)
+{
+	return PropertyTypeDispatch<GetCPPTypeDispatch, true, FString>(PropDef.GetPropertyBase(), std::ref(PropDef), ExtendedTypeText, CPPExportFlags);
+}
+
+FString FPropertyTraits::GetCPPTypeForwardDeclaration(const FUnrealPropertyDefinitionInfo& PropDef)
+{
+	return PropertyTypeDispatch<GetCPPTypeForwardDeclarationDispatch, true, FString>(PropDef.GetPropertyBase(), std::ref(PropDef));
+}
diff --git a/Engine/Source/Programs/UnrealHeaderTool/Private/PropertyTypes.h b/Engine/Source/Programs/UnrealHeaderTool/Private/PropertyTypes.h
--- a/Engine/Source/Programs/UnrealHeaderTool/Private/PropertyTypes.h
+++ b/Engine/Source/Programs/UnrealHeaderTool/Private/PropertyTypes.h
@@ -38,7 +38,7 @@ struct FPropertyTraits
 	 * @param ParsePosition Character position of the property in the header
 	 * @return The pointer to the newly created property.  It will be attached to the definition by the caller
 	 */
-	static FUnrealPropertyDefinitionInfo& CreateProperty(const FPropertyBase& VarProperty, FUnrealTypeDefinitionInfo& Outer, const FName& Name, EObjectFlags ObjectFlags, EVariableCategory::Type VariableCategory, const TCHAR* Dimensions, FUnrealSourceFile& SourceFile, int LineNumber, int ParsePosition);
+	static FUnrealPropertyDefinitionInfo& CreateProperty(const FPropertyBase& VarProperty, FUnrealTypeDefinitionInfo& Outer, const FName& Name, EObjectFlags ObjectFlags, EVariableCategory VariableCategory, const TCHAR* Dimensions, FUnrealSourceFile& SourceFile, int LineNumber, int ParsePosition);
 
 	/**
 	 * Test to see if the property can be used in a blueprint
@@ -48,4 +48,21 @@ struct FPropertyTraits
 	 * @return Return true if the property is supported in blueprints
 	 */
 	static bool IsSupportedByBlueprint(const FUnrealPropertyDefinitionInfo& PropDef, bool bMemberVariable);
+
+	/**
+	 * Return the engine class name for the given property information.
+	 * @param PropDef The property in question
+	 * @return The name of the engine property that will represent this definition.
+	 */
+	static FString GetEngineClassName(const FUnrealPropertyDefinitionInfo& PropDef);
+
+	/**
+	 * Returns the text to use for exporting this property to header file.
+	 *
+	 * @param	ExtendedTypeText	for property types which use templates, will be filled in with the type
+	 * @param	CPPExportFlags		flags for modifying the behavior of the export
+	 */
+	static FString GetCPPType(const FUnrealPropertyDefinitionInfo& PropDef, FString* ExtendedTypeText = nullptr, uint32 CPPExportFlags = 0);
+
+	static FString GetCPPTypeForwardDeclaration(const FUnrealPropertyDefinitionInfo& PropDef);
 };
diff --git a/Engine/Source/Programs/UnrealHeaderTool/Private/UnrealTypeDefinitionInfo.cpp b/Engine/Source/Programs/UnrealHeaderTool/Private/UnrealTypeDefinitionInfo.cpp
--- a/Engine/Source/Programs/UnrealHeaderTool/Private/UnrealTypeDefinitionInfo.cpp
+++ b/Engine/Source/Programs/UnrealHeaderTool/Private/UnrealTypeDefinitionInfo.cpp
@@ -5,18 +5,26 @@
 #include "ClassMaps.h"
 #include "HeaderParser.h"
 #include "NativeClassExporter.h"
+#include "PropertyTypes.h"
 #include "Scope.h"
 #include "StringUtils.h"
 #include "UnrealHeaderTool.h"
 #include "UnrealSourceFile.h"
+#include "Specifiers/CheckedMetadataSpecifiers.h"
 #include "Specifiers/ClassMetadataSpecifiers.h"
 
 #include "Algo/FindSortedStringCaseInsensitive.h"
+#include "Math/UnitConversion.h"
 #include "Misc/PackageName.h"
 #include "UObject/ErrorException.h"
 #include "UObject/Interface.h"
 #include "UObject/ObjectRedirector.h"
 
+// Globals for common class definitions
+extern FUnrealClassDefinitionInfo* GUObjectDef;
+extern FUnrealClassDefinitionInfo* GUClassDef;
+extern FUnrealClassDefinitionInfo* GUInterfaceDef;
+
 namespace
 {
 	const FName NAME_IgnoreCategoryKeywordsInSubclasses(TEXT("IgnoreCategoryKeywordsInSubclasses"));
@@ -91,9 +99,9 @@ namespace
 	}
 
 	template <typename T>
-	FString GetTypePackageName_Inner(const T* Field)
+	FString GetTypePackageNameHelper(const T& Field)
 	{
-		FString PackageName = Field->GetMetaData(NAME_ReplaceConverted);
+		FString PackageName = Field.GetMetaData(NAME_ReplaceConverted);
 		if (PackageName.Len())
 		{
 			int32 ObjectDotIndex = INDEX_NONE;
@@ -105,7 +113,7 @@ namespace
 		}
 		else
 		{
-			PackageName = Field->GetOutermost()->GetName();
+			PackageName = Field.GetPackageDef().GetName();
 		}
 		return PackageName;
 	}
@@ -144,6 +152,38 @@ namespace
 	}
 }
 
+void FUHTMetaData::RemapMetaData(TMap<FName, FString>& MetaData)
+{
+	// Evaluate any key redirects on the passed in pairs
+	for (TPair<FName, FString>& Pair : MetaData)
+	{
+		FName& CurrentKey = Pair.Key;
+		FName NewKey = UMetaData::GetRemappedKeyName(CurrentKey);
+
+		if (NewKey != NAME_None)
+		{
+			UE_LOG_WARNING_UHT(TEXT("Remapping old metadata key '%s' to new key '%s', please update the declaration."), *CurrentKey.ToString(), *NewKey.ToString());
+			CurrentKey = NewKey;
+		}
+	}
+}
+
+const FString& FUHTMetaData::GetMetaData(const FName& Key) const
+{
+	// if not found, return a static empty string
+	static FString EmptyString;
+
+	// every key needs to be valid and meta data needs to exist
+	if (Key == NAME_None)
+	{
+		return EmptyString;
+	}
+
+	// find and return either the located string or an empty string
+	const FString* ValuePtr = FindMetaData(Key);
+	return ValuePtr ? *ValuePtr : EmptyString;
+}
+
 FUnrealPropertyDefinitionInfo* FUnrealTypeDefinitionInfo::AsProperty()
 {
 	return nullptr;
@@ -238,9 +278,191 @@ bool FUnrealTypeDefinitionInfo::IsDynamic(const FField* Field)
 	return Field->HasMetaData(NAME_ReplaceConverted);
 }
 
+void FUnrealTypeDefinitionInfo::ValidateMetaDataFormat(const FName InKey, const FString& InValue) const
+{
+	ValidateMetaDataFormat(InKey, GetCheckedMetadataSpecifier(InKey), InValue);
+}
+
+void FUnrealTypeDefinitionInfo::ValidateMetaDataFormat(const FName InKey, ECheckedMetadataSpecifier InCheckedMetadataSpecifier, const FString& InValue) const
+{
+	switch (InCheckedMetadataSpecifier)
+	{
+	default:
+	{
+		// Don't need to validate this specifier
+	}
+	break;
+
+	case ECheckedMetadataSpecifier::UIMin:
+	case ECheckedMetadataSpecifier::UIMax:
+	case ECheckedMetadataSpecifier::ClampMin:
+	case ECheckedMetadataSpecifier::ClampMax:
+	{
+		if (!InValue.IsNumeric())
+		{
+			FError::Throwf(TEXT("Metadata value for '%s' is non-numeric : '%s'"), *InKey.ToString(), *InValue);
+		}
+	}
+	break;
+
+	case ECheckedMetadataSpecifier::BlueprintProtected:
+	{
+		if (const FUnrealFunctionDefinitionInfo* FuncDef = UHTCast<FUnrealFunctionDefinitionInfo>(this))
+		{
+			UFunction* Function = FuncDef->GetFunction();
+			if (FuncDef->HasAnyFunctionFlags(FUNC_Static))
+			{
+				// Determine if it's a function library
+				FUnrealClassDefinitionInfo* ClassDef = FuncDef->GetOwnerClass();
+				for (; ClassDef != nullptr && ClassDef->GetSuperClass() != GUObjectDef; ClassDef = ClassDef->GetSuperClass())
+				{
+				}
+
+				if (ClassDef != nullptr && ClassDef->GetName() == TEXT("BlueprintFunctionLibrary"))
+				{
+					FError::Throwf(TEXT("%s doesn't make sense on static method '%s' in a blueprint function library"), *InKey.ToString(), *FuncDef->GetName());
+				}
+			}
+		}
+	}
+	break;
+
+	case ECheckedMetadataSpecifier::CommutativeAssociativeBinaryOperator:
+	{
+		if (const FUnrealFunctionDefinitionInfo* FuncDef = UHTCast<FUnrealFunctionDefinitionInfo>(this))
+		{
+			bool bGoodParams = (FuncDef->GetProperties().Num() == 3);
+			if (bGoodParams)
+			{
+				FUnrealPropertyDefinitionInfo* FirstParam = nullptr;
+				FUnrealPropertyDefinitionInfo* SecondParam = nullptr;
+				FUnrealPropertyDefinitionInfo* ReturnValue = nullptr;
+				for (FUnrealPropertyDefinitionInfo* PropertyDef : FuncDef->GetProperties())
+				{
+					if (PropertyDef->HasAnyPropertyFlags(CPF_ReturnParm))
+					{
+						ReturnValue = PropertyDef;
+					}
+					else
+					{
+						if (FirstParam == nullptr)
+						{
+							FirstParam = PropertyDef;
+						}
+						else if (SecondParam == nullptr)
+						{
+							SecondParam = PropertyDef;
+						}
+					}
+				}
+
+				if (ReturnValue == nullptr || SecondParam == nullptr || !SecondParam->SameType(*FirstParam))
+				{
+					bGoodParams = false;
+				}
+			}
+
+			if (!bGoodParams)
+			{
+				UE_LOG_ERROR_UHT(TEXT("Commutative associative binary operators must have exactly 2 parameters of the same type and a return value."));
+			}
+		}
+	}
+	break;
+
+	case ECheckedMetadataSpecifier::ExpandBoolAsExecs:
+	case ECheckedMetadataSpecifier::ExpandEnumAsExecs:
+	{
+		if (const FUnrealFunctionDefinitionInfo* FuncDef = UHTCast<FUnrealFunctionDefinitionInfo>(this))
+		{
+			// multiple entry parsing in the same format as eg SetParam.
+			TArray<FString> RawGroupings;
+			InValue.ParseIntoArray(RawGroupings, TEXT(","), false);
+
+			FUnrealPropertyDefinitionInfo* FirstInputDef = nullptr;
+			for (const FString& RawGroup : RawGroupings)
+			{
+				TArray<FString> IndividualEntries;
+				RawGroup.ParseIntoArray(IndividualEntries, TEXT("|"));
+
+				for (const FString& Entry : IndividualEntries)
+				{
+					if (Entry.IsEmpty())
+					{
+						continue;
+					}
+
+					FUnrealPropertyDefinitionInfo* FoundFieldDef = FHeaderParser::FindProperty(*FuncDef, *Entry, false);
+					if (!FoundFieldDef)
+					{
+						UE_LOG_ERROR_UHT(TEXT("Function does not have a parameter named '%s'"), *Entry);
+					}
+					else
+					{
+						if (!FoundFieldDef->HasAnyPropertyFlags(CPF_ReturnParm) &&
+							(!FoundFieldDef->HasAnyPropertyFlags(CPF_OutParm) ||
+								FoundFieldDef->HasAnyPropertyFlags(CPF_ReferenceParm)))
+						{
+							if (!FirstInputDef)
+							{
+								FirstInputDef = FoundFieldDef;
+							}
+							else
+							{
+								UE_LOG_ERROR_UHT(TEXT("Function already specified an ExpandEnumAsExec input (%s), but '%s' is also an input parameter. Only one is permitted."), *FirstInputDef->GetName(), *Entry);
+							}
+						}
+					}
+				}
+			}
+		}
+	}
+	break;
+
+	case ECheckedMetadataSpecifier::DevelopmentStatus:
+	{
+		const FString EarlyAccessValue(TEXT("EarlyAccess"));
+		const FString ExperimentalValue(TEXT("Experimental"));
+		if ((InValue != EarlyAccessValue) && (InValue != ExperimentalValue))
+		{
+			FError::Throwf(TEXT("'%s' metadata was '%s' but it must be %s or %s"), *InKey.ToString(), *InValue, *ExperimentalValue, *EarlyAccessValue);
+		}
+	}
+	break;
+
+	case ECheckedMetadataSpecifier::Units:
+	{
+		// Check for numeric property
+		if (const FUnrealPropertyDefinitionInfo* PropDef = UHTCast<FUnrealPropertyDefinitionInfo>(this))
+		{
+			if (!PropDef->IsNumericOrNumericStaticArray() && !PropDef->IsStructOrStructStaticArray())
+			{
+				FError::Throwf(TEXT("'Units' meta data can only be applied to numeric and struct properties"));
+			}
+		}
+
+		if (!FUnitConversion::UnitFromString(*InValue))
+		{
+			FError::Throwf(TEXT("Unrecognized units (%s) specified for property '%s'"), *InValue, *GetFullName());
+		}
+	}
+	break;
+
+	case ECheckedMetadataSpecifier::DocumentationPolicy:
+	{
+		const TCHAR* StrictValue = TEXT("Strict");
+		if (InValue != StrictValue)
+		{
+			FError::Throwf(TEXT("'%s' metadata was '%s' but it must be %s"), *InKey.ToString(), *InValue, *StrictValue);
+		}
+	}
+	break;
+	}
+}
+
 void FUnrealPropertyDefinitionInfo::PostParseFinalize()
 {
-	TypePackageName = GetTypePackageName_Inner(GetProperty());
+	TypePackageName = GetTypePackageNameHelper(*this);
 }
 
 void FUnrealPropertyDefinitionInfo::AddMetaData(TMap<FName, FString>&& InMetaData)
@@ -263,7 +485,7 @@ void FUnrealPropertyDefinitionInfo::AddMetaData(TMap<FName, FString>&& InMetaDat
 
 bool FUnrealPropertyDefinitionInfo::IsDynamic() const
 {
-	return GetProperty()->HasMetaData(NAME_ReplaceConverted);
+	return HasMetaData(NAME_ReplaceConverted);
 }
 
 bool FUnrealPropertyDefinitionInfo::IsOwnedByDynamicType() const
@@ -289,16 +511,99 @@ void FUnrealPropertyDefinitionInfo::SetDelegateFunctionSignature(FUnrealFunction
 	PropertyBase.FunctionDef = &DelegateFunctionDef;
 }
 
+FString FUnrealPropertyDefinitionInfo::GetEngineClassName() const
+{
+#if UHT_ENABLE_ENGINE_TYPE_CHECKS
+	check(FPropertyTraits::GetEngineClassName(*this) == GetProperty()->GetClass()->GetName()); // Validation
+#endif
+	return FPropertyTraits::GetEngineClassName(*this);
+}
 
-FUnrealPackageDefinitionInfo& FUnrealObjectDefinitionInfo::GetPackageDef() const
+FString FUnrealPropertyDefinitionInfo::GetPathName() const
+{
+	TStringBuilder<256> ResultString;
+	GetPathName(ResultString);
+	return FString(FStringView(ResultString));
+}
+
+void FUnrealPropertyDefinitionInfo::GetPathName(FStringBuilderBase& ResultString) const
+{
+	TArray<FName, TInlineAllocator<16>> ParentFields;
+	for (FUnrealTypeDefinitionInfo* LocalOuter = GetOuter(); LocalOuter; LocalOuter = LocalOuter->GetOuter())
+	{
+		if (UHTCast<FUnrealPropertyDefinitionInfo>(LocalOuter))
+		{
+			ParentFields.Add(LocalOuter->GetFName());
+		}
+		else
+		{
+			LocalOuter->GetPathName(ResultString);
+			ResultString << SUBOBJECT_DELIMITER_CHAR;
+			break;
+		}
+	}
+
+	for (int FieldIndex = ParentFields.Num() - 1; FieldIndex >= 0; --FieldIndex)
+	{
+		ParentFields[FieldIndex].AppendString(ResultString);
+		ResultString << TEXT(".");
+	}
+	GetFName().AppendString(ResultString);
+#if UHT_ENABLE_ENGINE_TYPE_CHECKS
+	TStringBuilder<256> OtherResultString;
+	GetProperty()->GetPathName(nullptr, OtherResultString); // Validation
+	check(FCString::Strcmp(OtherResultString.ToString(), ResultString.ToString()) == 0);
+#endif
+}
+
+FString FUnrealPropertyDefinitionInfo::GetFullName() const
+{
+	FString FullName = GetEngineClassName();
+	FullName += TEXT(" ");
+	FullName += GetPathName();
+#if UHT_ENABLE_ENGINE_TYPE_CHECKS
+	check(FullName == GetProperty()->GetFullName()); // Validation
+#endif
+	return FullName;
+}
+
+FString FUnrealPropertyDefinitionInfo::GetCPPType(FString* ExtendedTypeText/* = nullptr*/, uint32 CPPExportFlags/* = 0*/) const
+{
+	FString Out = FPropertyTraits::GetCPPType(*this, ExtendedTypeText, CPPExportFlags);
+#if UHT_ENABLE_ENGINE_TYPE_CHECKS
+	FString ExtOutTemp;
+	FString* ExtOutTempPtr = ExtendedTypeText ? &ExtOutTemp : nullptr;
+	check(Out == GetProperty()->GetCPPType(ExtOutTempPtr, CPPExportFlags) && (ExtendedTypeText == nullptr || *ExtendedTypeText == ExtOutTemp)); // Validation
+#endif
+	return Out;
+}
+
+FString FUnrealPropertyDefinitionInfo::GetCPPTypeForwardDeclaration() const
+{
+	FString Out = FPropertyTraits::GetCPPTypeForwardDeclaration(*this);
+#if UHT_ENABLE_ENGINE_TYPE_CHECKS
+	check(Out == GetProperty()->GetCPPTypeForwardDeclaration()); // Validation
+#endif
+	return Out;
+}
+
+FUnrealPackageDefinitionInfo& FUnrealPropertyDefinitionInfo::GetPackageDef() const
 {
 	if (HasSource())
 	{
 		return GetUnrealSourceFile().GetPackageDef();
 	}
-	return GTypeDefinitionInfoMap.FindChecked<FUnrealPackageDefinitionInfo>(GetObject());
+	return GTypeDefinitionInfoMap.FindChecked<FUnrealPackageDefinitionInfo>(GetProperty()->GetOutermost());
 }
 
+FUnrealPackageDefinitionInfo& FUnrealObjectDefinitionInfo::GetPackageDef() const
+{
+	if (HasSource())
+	{
+		return GetUnrealSourceFile().GetPackageDef();
+	}
+	return GTypeDefinitionInfoMap.FindChecked<FUnrealPackageDefinitionInfo>(GetObject()->GetPackage());
+}
 
 FUnrealPackageDefinitionInfo::FUnrealPackageDefinitionInfo(const FManifestModule& InModule, UPackage* InPackage)
 	: FUnrealObjectDefinitionInfo(FString())
@@ -347,7 +652,7 @@ void FUnrealFieldDefinitionInfo::PostParseFinalize()
 	SingletonName[1] = Out;
 	SingletonNameChopped[1] = SingletonName[1].LeftChop(2);
 
-	TypePackageName = GetTypePackageName_Inner(GetField());
+	TypePackageName = GetTypePackageNameHelper(*this);
 }
 
 void FUnrealFieldDefinitionInfo::AddCrossModuleReference(TSet<FString>* UniqueCrossModuleReferences, bool bRequiresValidObject) const
@@ -444,20 +749,17 @@ void FUnrealStructDefinitionInfo::AddProperty(FUnrealPropertyDefinitionInfo& Pro
 	// update the optimization flags
 	if (!bContainsDelegates)
 	{
-		FProperty* Prop = PropertyDef.GetProperty();
-		if (Prop->IsA(FDelegateProperty::StaticClass()) || Prop->IsA(FMulticastDelegateProperty::StaticClass()))
+		const FPropertyBase& PropertyBase = PropertyDef.GetPropertyBase();
+		if (PropertyDef.IsDelegateOrDelegateStaticArray() || PropertyDef.IsMulticastDelegateOrMulticastDelegateStaticArray())
 		{
 			bContainsDelegates = true;
 		}
-		else
+		else if (PropertyDef.IsDynamicArray())
 		{
-			FArrayProperty* ArrayProp = CastField<FArrayProperty>(Prop);
-			if (ArrayProp != NULL)
+			const FUnrealPropertyDefinitionInfo& ValuePropertyDef = PropertyDef.GetValuePropDef();
+			if (ValuePropertyDef.IsDelegateOrDelegateStaticArray() || ValuePropertyDef.IsMulticastDelegateOrMulticastDelegateStaticArray())
 			{
-				if (ArrayProp->Inner->IsA(FDelegateProperty::StaticClass()) || ArrayProp->Inner->IsA(FMulticastDelegateProperty::StaticClass()))
-				{
-					bContainsDelegates = true;
-				}
+				bContainsDelegates = true;
 			}
 		}
 	}
@@ -1230,10 +1532,9 @@ FString FUnrealClassDefinitionInfo::GetNameWithPrefix(EEnforceInterfacePrefix En
 
 void FUnrealFunctionDefinitionInfo::AddProperty(FUnrealPropertyDefinitionInfo& PropertyDef)
 {
-	const FProperty* Prop = PropertyDef.GetProperty();
-	check((Prop->PropertyFlags & CPF_Parm) != 0);
+	check(PropertyDef.HasAnyPropertyFlags(CPF_Parm));
 
-	if ((Prop->PropertyFlags & CPF_ReturnParm) != 0)
+	if (PropertyDef.HasAnyPropertyFlags(CPF_ReturnParm))
 	{
 		check(ReturnProperty == nullptr);
 		ReturnProperty = &PropertyDef;
diff --git a/Engine/Source/Programs/UnrealHeaderTool/Private/UnrealTypeDefinitionInfo.h b/Engine/Source/Programs/UnrealHeaderTool/Private/UnrealTypeDefinitionInfo.h
--- a/Engine/Source/Programs/UnrealHeaderTool/Private/UnrealTypeDefinitionInfo.h
+++ b/Engine/Source/Programs/UnrealHeaderTool/Private/UnrealTypeDefinitionInfo.h
@@ -16,6 +16,7 @@
 class FScope;
 class FUnrealSourceFile;
 struct FManifestModule;
+enum class ECheckedMetadataSpecifier : int32;
 
 // These are declared in this way to allow swapping out the classes for something more optimized in the future
 typedef FStringOutputDevice FUHTStringBuilder;
@@ -66,6 +67,90 @@ enum class ESerializerArchiveType
 };
 ENUM_CLASS_FLAGS(ESerializerArchiveType)
 
+/** The category of variable declaration being parsed */
+enum class EVariableCategory
+{
+	RegularParameter,
+	ReplicatedParameter,
+	Return,
+	Member
+};
+
+/**
+ * UHT Specific implementation of meta data support
+ */
+class FUHTMetaData
+{
+public:
+
+	static void RemapMetaData(TMap<FName, FString>& MetaData);
+
+	template <typename Def>
+	static void RemapAndAddMetaData(Def& TypeDef, TMap<FName, FString>&& InMetaData)
+	{
+		RemapMetaData(InMetaData);
+		TypeDef.ValidateMetaDataFormat(InMetaData);
+		TypeDef.AddMetaData(MoveTemp(InMetaData));
+	}
+
+	template <typename Def>
+	static void RemapAndAddMetaData(Def& TypeDef, TMap<FName, FString>& InMetaData)
+	{
+		RemapMetaData(InMetaData);
+		TypeDef.ValidateMetaDataFormat(InMetaData);
+		TypeDef.AddMetaData(MoveTemp(InMetaData));
+	}
+
+	/**
+	* Determines if the property has any metadata associated with the key
+	*
+	* @param Key The key to lookup in the metadata
+	* @return true if there is a (possibly blank) value associated with this key
+	*/
+	bool HasMetaData(const TCHAR* Key) const 
+	{
+		return FindMetaData(Key) != nullptr;
+	}
+	bool HasMetaData(const FName& Key) const 
+	{
+		return FindMetaData(Key) != nullptr;
+	}
+
+	/**
+	* Find the metadata value associated with the key
+	*
+	* @param Key The key to lookup in the metadata
+	* @return The value associated with the key if it exists, null otherwise
+	*/
+	const FString* FindMetaData(const TCHAR * Key) const
+	{
+		return FindMetaData(FName(Key, FNAME_Find));
+	}
+	virtual const FString* FindMetaData(const FName & Key) const
+	{
+		TMap<FName, FString>* MetaDataMap = GetMetaDataMap();
+		return MetaDataMap ? MetaDataMap->Find(Key) : nullptr;
+	}
+
+	/**
+	* Find the metadata value associated with the key
+	*
+	* @param Key The key to lookup in the metadata
+	* @return The value associated with the key
+	*/
+	const FString& GetMetaData(const TCHAR* Key) const
+	{
+		return GetMetaData(FName(Key, FNAME_Find));
+	}
+	virtual const FString& GetMetaData(const FName& Key) const;
+
+protected:
+	virtual TMap<FName, FString>* GetMetaDataMap() const = 0;
+#if UHT_ENABLE_ENGINE_TYPE_CHECKS
+	virtual void CheckFindMetaData(const FName& Key, const FString* ValuePtr) const = 0;
+#endif
+};
+
 /**
  * Class that stores information about type definitions.
  */
@@ -261,6 +346,13 @@ public:
 
 	virtual FString GetFullName() const = 0;
 
+	/**
+	 * Returns the fully qualified pathname for this object, in the format:
+	 * 'Outermost.[Outer:]Name'
+	 */
+	virtual FString GetPathName() const = 0;
+	virtual void GetPathName(FStringBuilderBase& ResultString) const = 0;
+
 	/**
 	 * Return true if this type has source information
 	 */
@@ -360,6 +452,22 @@ public:
 	 */
 	virtual TMap<FName, FString> GenerateMetadataMap() const = 0;
 
+	/**
+	 * Ensures at script compile time that the metadata formatting is correct
+	 * @param	InKey			the metadata key being added
+	 * @param	InValue			the value string that will be associated with the InKey
+	 */
+	void ValidateMetaDataFormat(const FName InKey, const FString& InValue) const;
+
+	// Ensures at script compile time that the metadata formatting is correct
+	virtual void ValidateMetaDataFormat(const TMap<FName, FString>& MetaData) const
+	{
+		for (const TPair<FName, FString>& Pair : MetaData)
+		{
+			ValidateMetaDataFormat(Pair.Key, Pair.Value);
+		}
+	}
+
 protected:
 	explicit FUnrealTypeDefinitionInfo(FString&& InNameCPP)
 		: NameCPP(MoveTemp(InNameCPP))
@@ -372,6 +480,8 @@ protected:
 		, LineNumber(InLineNumber)
 	{ }
 
+	virtual void ValidateMetaDataFormat(const FName InKey, ECheckedMetadataSpecifier InCheckedMetadataSpecifier, const FString& InValue) const;
+
 private:
 	FString NameCPP;
 	FUnrealTypeDefinitionInfo* Outer = nullptr;
@@ -385,12 +495,16 @@ private:
  */
 class FUnrealPropertyDefinitionInfo
 	: public FUnrealTypeDefinitionInfo
+	, public FUHTMetaData
 {
 public:
-	FUnrealPropertyDefinitionInfo(FUnrealSourceFile& InSourceFile, int32 InLineNumber, int32 InParsePosition, const FPropertyBase& InVarProperty, FString&& InNameCPP, FUnrealTypeDefinitionInfo& InOuter)
-		: FUnrealTypeDefinitionInfo(InSourceFile, InLineNumber, MoveTemp(InNameCPP), InOuter)
+	FUnrealPropertyDefinitionInfo(FUnrealSourceFile& InSourceFile, int32 InLineNumber, int32 InParsePosition, const FPropertyBase& InVarProperty, EVariableCategory InVariableCategory, const TCHAR* Dimensions, FName InName, FUnrealTypeDefinitionInfo& InOuter)
+		: FUnrealTypeDefinitionInfo(InSourceFile, InLineNumber, InName.ToString(), InOuter)
+		, Name(InName)
 		, PropertyBase(InVarProperty)
+		, ArrayDimensions(InVarProperty.ArrayType == EArrayType::Static ? FString(Dimensions) : FString())
 		, ParsePosition(InParsePosition)
+		, VariableCategory(InVariableCategory)
 	{ }
 
 	virtual FUnrealPropertyDefinitionInfo* AsProperty() override
@@ -445,41 +559,39 @@ public:
 	 */
 	virtual FString GetName() const override
 	{
-		return GetProperty()->GetName();
+#if UHT_ENABLE_ENGINE_TYPE_CHECKS
+		check(GetProperty()->GetName() == GetNameCPP());
+#endif
+		return GetNameCPP();
 	}
 
 	/** Returns the logical name of this object */
 	virtual FName GetFName() const override
 	{
-		return GetProperty()->GetFName();
+#if UHT_ENABLE_ENGINE_TYPE_CHECKS
+		check(GetProperty()->GetFName() == Name);
+#endif
+		return Name;
 	}
 
-	virtual FString GetFullName() const override
-	{
-		return GetProperty()->GetFullName();
-	}
+	virtual FString GetFullName() const override;
 
 	/**
 	 * Returns the fully qualified pathname for this object, in the format:
 	 * 'Outermost.[Outer:]Name'
-	 *
-	 * @param	StopOuter	if specified, indicates that the output string should be relative to this object.  if StopOuter
-	 *						does not exist in this object's Outer chain, the result would be the same as passing NULL.
-	 *
-	 * @note	safe to call on NULL object pointers!
 	 */
-	FString GetPathName(const UObject* StopOuter = nullptr)
-	{
-		return GetProperty()->GetPathName(StopOuter);
-	}
+	virtual FString GetPathName() const override;
+	virtual void GetPathName(FStringBuilderBase& ResultString) const override;
 
 	/**
 	 * Return the engine class name
 	 */
-	FString GetEngineClassName()
-	{
-		return GetProperty()->GetClass()->GetName();
-	}
+	FString GetEngineClassName() const;
+
+	/**
+	 * Return the package associated with this object
+	 */
+	FUnrealPackageDefinitionInfo& GetPackageDef() const;
 
 	/**
 	* Finds the localized display name or native display name as a fallback.
@@ -509,9 +621,13 @@ public:
 	 *
 	 * @return C++ name of property
 	 */
-	FString GetPropertyNameCPP() const
+	FString GetNameWithDeprecated() const
 	{
-		return GetProperty()->GetNameCPP();
+		FString OutName = HasAnyPropertyFlags(CPF_Deprecated) ? GetName() + TEXT("_DEPRECATED") : GetName();
+#if UHT_ENABLE_ENGINE_TYPE_CHECKS
+		check(GetProperty()->GetNameCPP() == OutName);
+#endif
+		return OutName;
 	}
 
 	/**
@@ -520,15 +636,9 @@ public:
 	 * @param	ExtendedTypeText	for property types which use templates, will be filled in with the type
 	 * @param	CPPExportFlags		flags for modifying the behavior of the export
 	 */
-	FString GetCPPType(FString* ExtendedTypeText = NULL, uint32 CPPExportFlags = 0) const
-	{
-		return GetProperty()->GetCPPType(ExtendedTypeText, CPPExportFlags);
-	}
+	FString GetCPPType(FString* ExtendedTypeText = nullptr, uint32 CPPExportFlags = 0) const;
 
-	FString GetCPPTypeForwardDeclaration() const
-	{
-		return GetProperty()->GetCPPTypeForwardDeclaration();
-	}
+	FString GetCPPTypeForwardDeclaration() const;
 
 	/**
 	 * Returns this property's propertyflags
@@ -607,53 +717,6 @@ public:
 		return (PropertyBase.PropertyFlags & FlagsToCheck) == ExpectedFlags;
 	}
 
-	/**
-	 * Get the cast flags
-	 */
-	uint64 GetCastFlags() const
-	{
-		return GetProperty()->GetCastFlags();
-	}
-
-	/**
-	* Determines if the property has any metadata associated with the key
-	*
-	* @param Key The key to lookup in the metadata
-	* @return true if there is a (possibly blank) value associated with this key
-	*/
-	bool HasMetaData(const TCHAR* Key) const { return FindMetaData(Key) != nullptr; }
-	bool HasMetaData(const FName& Key) const { return FindMetaData(Key) != nullptr; }
-
-	/**
-	* Find the metadata value associated with the key
-	*
-	* @param Key The key to lookup in the metadata
-	* @return The value associated with the key if it exists, null otherwise
-	*/
-	const FString* FindMetaData(const TCHAR* Key) const
-	{
-		return GetProperty()->FindMetaData(Key);
-	}
-	const FString* FindMetaData(const FName& Key) const
-	{
-		return GetProperty()->FindMetaData(Key);
-	}
-
-	/**
-	* Find the metadata value associated with the key
-	*
-	* @param Key The key to lookup in the metadata
-	* @return The value associated with the key
-	*/
-	const FString& GetMetaData(const TCHAR* Key) const
-	{
-		return GetProperty()->GetMetaData(Key);
-	}
-	const FString& GetMetaData(const FName& Key) const
-	{
-		return GetProperty()->GetMetaData(Key);
-	}
-
 	/**
 	 * Return the meta data in map form
 	 */
@@ -677,7 +740,11 @@ public:
 	 */
 	bool IsEditorOnlyProperty() const
 	{
-		return GetProperty()->IsEditorOnlyProperty();
+		bool bResult = HasAnyPropertyFlags(CPF_DevelopmentAssets);
+#if UHT_ENABLE_ENGINE_TYPE_CHECKS
+		check(bResult == GetProperty()->IsEditorOnlyProperty());
+#endif
+		return bResult;
 	}
 
 	/**
@@ -688,8 +755,11 @@ public:
 	 */
 	bool ContainsInstancedObjectProperty() const
 	{
-		//return (PropertyFlags & (CPF_ContainsInstancedReference | CPF_InstancedReference)) != 0;
-		return GetProperty()->ContainsInstancedObjectProperty();
+		bool bResult = HasAnyPropertyFlags(CPF_ContainsInstancedReference | CPF_InstancedReference);
+#if UHT_ENABLE_ENGINE_TYPE_CHECKS
+		check(bResult == GetProperty()->ContainsInstancedObjectProperty());
+#endif
+		return bResult;
 	}
 
 	/**
@@ -812,6 +882,22 @@ public:
 		return PropertyBase.IsStructOrStructStaticArray();
 	}
 
+	/**
+	 * Return true if this is a delegate or a delegate static array
+	 */
+	bool IsDelegateOrDelegateStaticArray() const
+	{
+		return PropertyBase.IsDelegateOrDelegateStaticArray();
+	}
+
+	/**
+	 * Return true if this is a multicast delegate or a multicast delegate static array
+	 */
+	bool IsMulticastDelegateOrMulticastDelegateStaticArray() const
+	{
+		return PropertyBase.IsMulticastDelegateOrMulticastDelegateStaticArray();
+	}	
+
 	/**
 	 * Return true if this is a structure and it has no op constructor
 	 */
@@ -825,20 +911,19 @@ public:
 	}
 
 	/**
-	 * Set the string that represents the array dimensions
+	 * Get the string that represents the array dimensions.  A nullptr is returned if the property doesn't have any dimensions
 	 */
-	void SetArrayDimensions(const FString& InArrayDimensions)
+	const TCHAR* GetArrayDimensions() const
 	{
-		ArrayDimensions = InArrayDimensions;
-		check(!ArrayDimensions.IsEmpty());
+		return ArrayDimensions.IsEmpty() ? nullptr : *ArrayDimensions;
 	}
-
+	
 	/**
-	 * Get the string that represents the array dimensions.  A nullptr is returned if the property doesn't have any dimensions
+	 * Get the category of this property
 	 */
-	const TCHAR* GetArrayDimensions() const
+	EVariableCategory GetVariableCategory() const
 	{
-		return ArrayDimensions.IsEmpty() ? nullptr : *ArrayDimensions;
+		return VariableCategory;
 	}
 
 	/**
@@ -973,24 +1058,33 @@ public:
 	 */
 	FName GetRepNotifyFunc() const
 	{
-		return GetProperty()->RepNotifyFunc;
+#if UHT_ENABLE_ENGINE_TYPE_CHECKS
+		check(PropertyBase.RepNotifyName == GetProperty()->RepNotifyFunc);
+#endif
+		return PropertyBase.RepNotifyName;
 	}
 
 	/**
-	 * Set the replication notification function name
+	 * Set the delegate function signature
 	 */
-	void SetRepNotifyFunc(FName InRepNotifyFunc)
+	void SetDelegateFunctionSignature(FUnrealFunctionDefinitionInfo& DelegateFunctionDef);
+
+protected:
+	virtual TMap<FName, FString>* GetMetaDataMap() const override
 	{
-		GetProperty()->RepNotifyFunc = InRepNotifyFunc;
-		PropertyBase.RepNotifyName = InRepNotifyFunc;
+		return const_cast<TMap<FName, FString>*>(GetProperty()->GetMetaDataMap());
 	}
 
-	/**
-	 * Set the delegate function signature
-	 */
-	void SetDelegateFunctionSignature(FUnrealFunctionDefinitionInfo& DelegateFunctionDef);
+#if UHT_ENABLE_ENGINE_TYPE_CHECKS
+	virtual void CheckFindMetaData(const FName& Key, const FString* ValuePtr) const
+	{
+		const FString* CheckPtr = GetProperty()->FindMetaData(Key);
+		check((CheckPtr == nullptr && ValuePtr == nullptr) || (CheckPtr != nullptr && ValuePtr != nullptr && *CheckPtr == *ValuePtr));
+	}
+#endif
 
 private:
+	FName Name;
 	FPropertyBase PropertyBase;
 	FString ArrayDimensions;
 	FString TypePackageName;
@@ -998,6 +1092,7 @@ private:
 	FUnrealPropertyDefinitionInfo* ValuePropDef = nullptr;
 	FProperty* Property = nullptr;
 	int32 ParsePosition;
+	EVariableCategory VariableCategory = EVariableCategory::Member;
 	EAllocatorType AllocatorType = EAllocatorType::Default;
 	bool bUnsized = false;
 };
@@ -1056,15 +1151,15 @@ public:
 	/**
 	 * Returns the fully qualified pathname for this object, in the format:
 	 * 'Outermost.[Outer:]Name'
-	 *
-	 * @param	StopOuter	if specified, indicates that the output string should be relative to this object.  if StopOuter
-	 *						does not exist in this object's Outer chain, the result would be the same as passing NULL.
-	 *
-	 * @note	safe to call on NULL object pointers!
 	 */
-	FString GetPathName(const FUnrealObjectDefinitionInfo* StopOuter = nullptr) const
+	virtual FString GetPathName() const override
+	{
+		return GetObject()->GetPathName(nullptr);
+	}
+	
+	virtual void GetPathName(FStringBuilderBase& ResultString) const override
 	{
-		return GetObject()->GetPathName(StopOuter ? StopOuter->GetObject() : nullptr);
+		return GetObject()->GetPathName(nullptr, ResultString);
 	}
 
 	/**
@@ -2272,6 +2367,9 @@ public:
 	 */
 	bool HasAnyClassFlags(EClassFlags FlagsToCheck) const
 	{
+		bool bA = EnumHasAnyFlags(ParsedClassFlags, FlagsToCheck);
+		bool bB = GetClass()->HasAnyClassFlags(FlagsToCheck);
+		check(!bA || bB);
 		return EnumHasAnyFlags(ParsedClassFlags, FlagsToCheck) || GetClass()->HasAnyClassFlags(FlagsToCheck);
 	}
 
@@ -2283,6 +2381,9 @@ public:
 	 */
 	bool HasAllClassFlags(EClassFlags FlagsToCheck) const
 	{
+		bool bA = EnumHasAllFlags(ParsedClassFlags, FlagsToCheck);
+		bool bB = GetClass()->HasAllClassFlags(FlagsToCheck);
+		check(!bA || bB);
 		return EnumHasAllFlags(ParsedClassFlags, FlagsToCheck) || GetClass()->HasAllClassFlags(FlagsToCheck);
 	}
 
