Properties engine types no longer need to be created during parsing.

#rnx

diff --git a/Engine/Source/Programs/UnrealHeaderTool/Private/CodeGenerator.cpp b/Engine/Source/Programs/UnrealHeaderTool/Private/CodeGenerator.cpp
--- a/Engine/Source/Programs/UnrealHeaderTool/Private/CodeGenerator.cpp
+++ b/Engine/Source/Programs/UnrealHeaderTool/Private/CodeGenerator.cpp
@@ -175,7 +175,7 @@ namespace
 	{
 		if (!ClassDef.HasAnyClassFlags(CLASS_ReplicationDataIsSetUp))
 		{
-			for (FUnrealPropertyDefinitionInfo* PropertyDef : ClassDef.GetProperties())
+			for (TSharedRef<FUnrealPropertyDefinitionInfo> PropertyDef : ClassDef.GetProperties())
 			{
 				if (PropertyDef->HasAnyPropertyFlags(CPF_Net))
 				{
@@ -514,15 +514,15 @@ struct FParmsAndReturnProperties
 FParmsAndReturnProperties GetFunctionParmsAndReturn(FUnrealFunctionDefinitionInfo& FunctionDef)
 {
 	FParmsAndReturnProperties Result;
-	for (FUnrealPropertyDefinitionInfo* PropertyDef : FunctionDef.GetProperties())
+	for (TSharedRef<FUnrealPropertyDefinitionInfo> PropertyDef : FunctionDef.GetProperties())
 	{
 		if (PropertyDef->HasSpecificPropertyFlags(CPF_Parm | CPF_ReturnParm, CPF_Parm))
 		{
-			Result.Parms.Add(PropertyDef);
+			Result.Parms.Add(&*PropertyDef);
 		}
 		else if (PropertyDef->HasAnyPropertyFlags(CPF_ReturnParm))
 		{
-			Result.Return = PropertyDef;
+			Result.Return = &*PropertyDef;
 		}
 	}
 	return Result;
@@ -708,7 +708,7 @@ void FNativeClassHeaderGenerator::ExportProperties(FOutputDevice& Out, FUnrealSt
 	FMacroBlockEmitter WithEditorOnlyData(Out, TEXT("WITH_EDITORONLY_DATA"));
 
 	// Iterate over all properties in this struct.
-	for (FUnrealPropertyDefinitionInfo* PropertyDef : StructDef.GetProperties())
+	for (TSharedRef<FUnrealPropertyDefinitionInfo> PropertyDef : StructDef.GetProperties())
 	{
 		WithEditorOnlyData(PropertyDef->IsEditorOnlyProperty());
 
@@ -1613,7 +1613,7 @@ void FNativeClassHeaderGenerator::OutputProperty(FOutputDevice& DeclOut, FOutput
 			FPropertyObjectFlags,
 			*ArrayDim,
 			OffsetStr,
-			*FString::Printf(TEXT("&F%s::StaticClass"), *PropertyBase.PropertyPathClass->GetName()),
+			*FString::Printf(TEXT("&F%s::StaticClass"), *PropertyBase.FieldClassName.ToString()),
 			*MetaDataParams,
 			*PropTag
 		);
@@ -1698,7 +1698,7 @@ TTuple<FString, FString> FNativeClassHeaderGenerator::OutputProperties(FOutputDe
 		FMacroBlockEmitter WithEditorOnlyMacroEmitter(Out, TEXT("WITH_EDITORONLY_DATA"));
 		FMacroBlockEmitter WithEditorOnlyMacroEmitterDecl(DeclOut, TEXT("WITH_EDITORONLY_DATA"));
 
-		for (FUnrealPropertyDefinitionInfo* PropertyDef : StructDef.GetProperties())
+		for (TSharedRef<FUnrealPropertyDefinitionInfo> PropertyDef : StructDef.GetProperties())
 		{
 			bool bRequiresHasEditorOnlyMacro = IsEditorOnlyDataProperty(*PropertyDef);
 			if (!bRequiresHasEditorOnlyMacro)
@@ -1796,7 +1796,7 @@ static void FindNoExportStructsRecursive(TArray<FUnrealScriptStructDefinitionInf
 			}
 		}
 
-		for (FUnrealPropertyDefinitionInfo* PropertyDef : StartDef->GetProperties())
+		for (TSharedRef<FUnrealPropertyDefinitionInfo> PropertyDef : StartDef->GetProperties())
 		{
 			const FPropertyBase& PropertyBase = PropertyDef->GetPropertyBase();
 			if (UHTCast<FUnrealScriptStructDefinitionInfo>(PropertyBase.TypeDef) != nullptr)
@@ -1912,7 +1912,7 @@ void FNativeClassHeaderGenerator::ExportNativeGeneratedInitCode(FOutputDevice& O
 	TSet<FName> AlreadyIncludedNames;
 	TArray<FUnrealFunctionDefinitionInfo*> FunctionsToExport;
 	bool bAllEditorOnlyFunctions = true;
-	for (FUnrealFunctionDefinitionInfo* LocalFuncDef : ClassDef.GetFunctions())
+	for (TSharedRef<FUnrealFunctionDefinitionInfo> LocalFuncDef : ClassDef.GetFunctions())
 	{
 		FName TrueName = FNativeClassHeaderGenerator::GetOverriddenFName(*LocalFuncDef);
 		bool bAlreadyIncluded = false;
@@ -1930,7 +1930,7 @@ void FNativeClassHeaderGenerator::ExportNativeGeneratedInitCode(FOutputDevice& O
 		{
 			bAllEditorOnlyFunctions &= LocalFuncDef->HasAnyFunctionFlags(FUNC_EditorOnly);
 		}
-		FunctionsToExport.Add(LocalFuncDef);
+		FunctionsToExport.Add(&*LocalFuncDef);
 	}
 
 	// Sort the list of functions
@@ -2462,12 +2462,12 @@ void FNativeClassHeaderGenerator::ExportNatives(FOutputDevice& Out, FUnrealClass
 		bool bAllEditorOnly = true;
 
 		TArray<TTuple<FUnrealFunctionDefinitionInfo*, FString>> NamedFunctionsToExport;
-		for (FUnrealFunctionDefinitionInfo* FunctionDef : ClassDef.GetFunctions())
+		for (TSharedRef<FUnrealFunctionDefinitionInfo> FunctionDef : ClassDef.GetFunctions())
 		{
 			if (FunctionDef->HasSpecificFunctionFlags(FUNC_Native | FUNC_NetRequest, FUNC_Native))
 			{
-				FString OverriddenName = FNativeClassHeaderGenerator::GetUTF8OverriddenNameForLiteral(FunctionDef);
-				NamedFunctionsToExport.Emplace(FunctionDef, MoveTemp(OverriddenName));
+				FString OverriddenName = FNativeClassHeaderGenerator::GetUTF8OverriddenNameForLiteral(&*FunctionDef);
+				NamedFunctionsToExport.Emplace(&*FunctionDef, MoveTemp(OverriddenName));
 
 				if (!FunctionDef->HasAnyFunctionFlags(FUNC_EditorOnly))
 				{
@@ -2684,11 +2684,11 @@ void FNativeClassHeaderGenerator::ExportClassFromSourceFileInner(
 
 	// Get Callback functions
 	TArray<FUnrealFunctionDefinitionInfo*> CallbackFunctions;
-	for (FUnrealFunctionDefinitionInfo* FunctionDef : ClassDef.GetFunctions())
+	for (TSharedRef<FUnrealFunctionDefinitionInfo> FunctionDef : ClassDef.GetFunctions())
 	{
 		if (FunctionDef->HasAnyFunctionFlags(FUNC_Event) && FunctionDef->GetSuperFunction() == nullptr)
 		{
-			CallbackFunctions.Add(FunctionDef);
+			CallbackFunctions.Add(&*FunctionDef);
 		}
 	}
 
@@ -4045,7 +4045,7 @@ void FNativeClassHeaderGenerator::ExportMirrorsForNoexportStruct(FOutputDevice&
 
 bool FNativeClassHeaderGenerator::WillExportEventParms(FUnrealFunctionDefinitionInfo& FunctionDef)
 {
-	const TArray<FUnrealPropertyDefinitionInfo*>& Properties = FunctionDef.GetProperties();
+	const TArray<TSharedRef<FUnrealPropertyDefinitionInfo>>& Properties = FunctionDef.GetProperties();
 	return Properties.Num() > 0 && Properties[0]->HasAnyPropertyFlags(CPF_Parm);
 }
 
@@ -4217,7 +4217,7 @@ void FNativeClassHeaderGenerator::ExportEventParm(FUHTStringBuilder& Out, TSet<F
 	Out.Logf(TEXT("%sstruct %s\r\n"), FCString::Tab(Indent), *EventParmStructName);
 	Out.Logf(TEXT("%s{\r\n"), FCString::Tab(Indent));
 
-	for (FUnrealPropertyDefinitionInfo* PropDef : FunctionDef.GetProperties())
+	for (TSharedRef<FUnrealPropertyDefinitionInfo> PropDef : FunctionDef.GetProperties())
 	{
 		FPropertyBase& PropertyBase = PropDef->GetPropertyBase();
 		if (!PropDef->HasAnyPropertyFlags(CPF_Parm))
@@ -4652,7 +4652,7 @@ void FNativeClassHeaderGenerator::ExportNativeFunctionHeader(
 		++ParmCount;
 	}
 
-	for (FUnrealPropertyDefinitionInfo* PropertyDef : FunctionDef.GetProperties())
+	for (TSharedRef<FUnrealPropertyDefinitionInfo> PropertyDef : FunctionDef.GetProperties())
 	{
 		if (!PropertyDef->HasSpecificPropertyFlags(CPF_Parm | CPF_ReturnParm, CPF_Parm))
 		{
@@ -4920,7 +4920,12 @@ FString FNativeClassHeaderGenerator::GetFunctionParameterString(FUnrealFunctionD
 	FString ParameterList;
 	FUHTStringBuilder PropertyText;
 
-	for (FUnrealPropertyDefinitionInfo* PropertyDef : FunctionDef.GetProperties())
+	if (FunctionDef.GetReturn() != nullptr)
+	{
+		ParameterList = ParameterList;
+	}
+
+	for (TSharedRef<FUnrealPropertyDefinitionInfo> PropertyDef : FunctionDef.GetProperties())
 	{
 		OutReferenceGatherers.ForwardDeclarations.Add(PropertyDef->GetCPPTypeForwardDeclaration());
 
@@ -4985,7 +4990,7 @@ void FNativeClassHeaderGenerator::ExportNativeFunctions(FOutputDevice& OutGenera
 		FUnrealScriptStructDefinitionInfo* SparseClassDataStructDef = GTypeDefinitionInfoMap.FindByName<FUnrealScriptStructDefinitionInfo>(*SparseClassDataString);
 		while (SparseClassDataStructDef != nullptr)
 		{
-			for (FUnrealPropertyDefinitionInfo* PropertyDef : SparseClassDataStructDef->GetProperties())
+			for (TSharedRef<FUnrealPropertyDefinitionInfo> PropertyDef : SparseClassDataStructDef->GetProperties())
 			{
 				FString ReturnExtendedType;
 				FString VarType = PropertyDef->GetCPPType(&ReturnExtendedType, EPropertyExportCPPFlags::CPPF_ArgumentOrReturnValue | EPropertyExportCPPFlags::CPPF_Implementation);
@@ -5032,12 +5037,12 @@ void FNativeClassHeaderGenerator::ExportNativeFunctions(FOutputDevice& OutGenera
 		}
 	}
 
-	TArray<FUnrealFunctionDefinitionInfo*> Functions;
+	TArray<TSharedRef<FUnrealFunctionDefinitionInfo>> Functions;
 	Algo::Copy(ClassDef.GetFunctions(), Functions);
 	Algo::Reverse(Functions);
 
 	// export the C++ stubs
-	for (FUnrealFunctionDefinitionInfo* FunctionDef : Functions)
+	for (TSharedRef<FUnrealFunctionDefinitionInfo> FunctionDef : Functions)
 	{
 		const FFuncInfo& FunctionData = FunctionDef->GetFunctionData();
 		if (!FunctionDef->HasAnyFunctionFlags(FUNC_Native))
@@ -6644,62 +6649,20 @@ void PrepareTypesForExport(TArray<FUnrealPackageDefinitionInfo*>& PackageDefs)
 	{
 		FResults::Try([PackageDef]()
 			{
-				// Collect all of the classes from the sources
-				TArray<FUnrealClassDefinitionInfo*> PackageClasses;
-				PackageClasses.Reserve(128); // Arbitrary number
-				for (TSharedRef<FUnrealSourceFile>& SourceFile : PackageDef->GetAllSourceFiles())
-				{
-					for (TSharedRef<FUnrealTypeDefinitionInfo>& TypeDef : SourceFile->GetDefinedClasses())
-					{
-						PackageClasses.Add(&TypeDef->AsClassChecked());
-					}
-				}
-
-				// Validate the sparse class data for all classes in the current package
-				for (const FUnrealClassDefinitionInfo* ClassDef : PackageClasses)
-				{
-					FHeaderParser::CheckSparseClassData(*ClassDef);
-				}
-
-				// Check to see if we need write the classes 
-				const bool bPackageHasAnyExportClasses = PackageClasses.ContainsByPredicate([](FUnrealClassDefinitionInfo* ClassDef)
-					{
-						return ClassDef->HasAnyClassFlags(CLASS_Native) && !ClassDef->HasAnyClassFlags(CLASS_NoExport | CLASS_Intrinsic);
-					});
-				if (bPackageHasAnyExportClasses)
-				{
-					for (FUnrealClassDefinitionInfo* ClassDef : PackageClasses)
-					{
-						UClass* Class = ClassDef->GetClass();
-						if (!ClassDef->HasAnyClassFlags(CLASS_Native))
-						{
-							Class->UnMark(EObjectMark(OBJECTMARK_TagImp | OBJECTMARK_TagExp));
-						}
-						else if (!ClassDef->HasAnyClassFlags(CLASS_NoExport))
-						{
-							PackageDef->SetWriteClassesH(true);
-							Class->UnMark(OBJECTMARK_TagImp);
-							Class->Mark(OBJECTMARK_TagExp);
-						}
-					}
-				}
-
-				// This needs to be done outside of parallel blocks because it will modify UClass memory.
-				// Later calls to SetUpUhtReplicationData inside parallel blocks should be fine, because
-				// they will see the memory has already been set up, and just return the parent pointer.
-				for (FUnrealClassDefinitionInfo* ClassDef : PackageClasses)
-				{
-					if (ClassHasReplicatedProperties(*ClassDef))
-					{
-						ClassDef->GetClass()->SetUpUhtReplicationData();
-					}
-				}
+				PackageDef->PostParseFinalize(true);
 			}
 		);
 	}
 	FResults::WaitForErrorTasks();
 
-	GTypeDefinitionInfoMap.ForAllTypes([](FUnrealTypeDefinitionInfo& TypeDef) { TypeDef.PostParseFinalize(); });
+	// Until we fix the issue with types not having sources, we need to do this
+	GTypeDefinitionInfoMap.ForAllTypes([](FUnrealTypeDefinitionInfo& TypeDef) 
+		{ 
+			if (!TypeDef.HasSource())
+			{
+				TypeDef.PostParseFinalize(true);
+			}
+		});
 }
 
 void Export(TArray<FUnrealPackageDefinitionInfo*>& PackageDefs, TArray<FUnrealSourceFile*>& OrderedSourceFiles)
diff --git a/Engine/Source/Programs/UnrealHeaderTool/Private/HeaderParser.cpp b/Engine/Source/Programs/UnrealHeaderTool/Private/HeaderParser.cpp
--- a/Engine/Source/Programs/UnrealHeaderTool/Private/HeaderParser.cpp
+++ b/Engine/Source/Programs/UnrealHeaderTool/Private/HeaderParser.cpp
@@ -2119,7 +2119,7 @@ FUnrealScriptStructDefinitionInfo& FHeaderParser::CompileStructDeclaration()
 	CheckSparseClassData(StructDef);
 
 	// Link the properties within the struct
-	Struct->StaticLink(true);
+	Struct->PrepareCppStructOps();
 
 	return StructDef;
 }
@@ -2291,11 +2291,6 @@ void FHeaderParser::PopNest(ENestType NestType, const TCHAR* Descr)
 		bLinkProps = !TopClass->HasAnyClassFlags(CLASS_Intrinsic);
 	}
 
-	if (NestType != ENestType::GlobalScope)
-	{
-		GetCurrentClass()->StaticLink(bLinkProps);
-	}
-
 	// Pop the nesting level.
 	NestType = TopNest->NestType;
 	NestLevel--;
@@ -2313,7 +2308,7 @@ void FHeaderParser::PopNest(ENestType NestType, const TCHAR* Descr)
 
 void FHeaderParser::FixupDelegateProperties(FUnrealStructDefinitionInfo& StructDef, FScope& Scope, TMap<FName, FUnrealFunctionDefinitionInfo*>& DelegateCache )
 {
-	for (FUnrealPropertyDefinitionInfo* PropertyDef : StructDef.GetProperties())
+	for (TSharedRef<FUnrealPropertyDefinitionInfo> PropertyDef : StructDef.GetProperties())
 	{
 		FPropertyBase& PropertyBase = PropertyDef->GetPropertyBase();
 
@@ -2408,7 +2403,7 @@ void FHeaderParser::FixupDelegateProperties(FUnrealStructDefinitionInfo& StructD
 
 		// successfully found the delegate function that this delegate property corresponds to
 
-			// save this into the delegate cache for faster lookup later
+		// save this into the delegate cache for faster lookup later
 		DelegateCache.Add(DelegatePropertyBase.DelegateName, SourceDelegateFunctionDef);
 
 		// bind it to the delegate property
@@ -2431,7 +2426,7 @@ void FHeaderParser::FixupDelegateProperties(FUnrealStructDefinitionInfo& StructD
 
 				if (PropertyDef->HasAnyPropertyFlags(CPF_BlueprintAssignable | CPF_BlueprintCallable))
 				{
-					for (FUnrealPropertyDefinitionInfo* FuncParamDef : SourceDelegateFunctionDef->GetProperties())
+					for (TSharedRef<FUnrealPropertyDefinitionInfo> FuncParamDef : SourceDelegateFunctionDef->GetProperties())
 					{
 						if (!FPropertyTraits::IsSupportedByBlueprint(*FuncParamDef, false))
 						{
@@ -2461,7 +2456,7 @@ void FHeaderParser::FixupDelegateProperties(FUnrealStructDefinitionInfo& StructD
 	}
 
 	// Functions might have their own delegate properties which need to be validated
-	for (FUnrealFunctionDefinitionInfo* Function : StructDef.GetFunctions())
+	for (TSharedRef<FUnrealFunctionDefinitionInfo> Function : StructDef.GetFunctions())
 	{
 		FixupDelegateProperties(*Function, Scope, DelegateCache);
 	}
@@ -2568,7 +2563,7 @@ void FHeaderParser::VerifyBlueprintPropertyGetter(FUnrealPropertyDefinitionInfo&
 {
 	check(TargetFuncDef);
 
-	const TArray<FUnrealPropertyDefinitionInfo*>& Properties = TargetFuncDef->GetProperties();
+	const TArray<TSharedRef<FUnrealPropertyDefinitionInfo>>& Properties = TargetFuncDef->GetProperties();
 	FUnrealPropertyDefinitionInfo* ReturnPropDef = TargetFuncDef->GetReturn();
 	if (Properties.Num() > 1 || (Properties.Num() == 1 && ReturnPropDef == nullptr))
 	{
@@ -2603,7 +2598,7 @@ void FHeaderParser::VerifyBlueprintPropertySetter(FUnrealPropertyDefinitionInfo&
 	}
 	else
 	{
-		const TArray<FUnrealPropertyDefinitionInfo*>& Properties = TargetFuncDef->GetProperties();
+		const TArray<TSharedRef<FUnrealPropertyDefinitionInfo>>& Properties = TargetFuncDef->GetProperties();
 		if (TargetFuncDef->GetProperties().Num() != 1 || !PropertyDef.SameType(*Properties[0]))
 		{
 			FString ExtendedCPPType;
@@ -2638,7 +2633,7 @@ void FHeaderParser::VerifyRepNotifyCallback(FUnrealPropertyDefinitionInfo& Prope
 		const bool bIsArrayProperty = PropertyDef.IsStaticArray() || PropertyDef.IsDynamicArray();
 		const int32 MaxParms = bIsArrayProperty ? 2 : 1;
 
-		const TArray<FUnrealPropertyDefinitionInfo*>& Properties = TargetFuncDef->GetProperties();
+		const TArray<TSharedRef<FUnrealPropertyDefinitionInfo>>& Properties = TargetFuncDef->GetProperties();
 		if (Properties.Num() > MaxParms)
 		{
 			UE_LOG_ERROR_UHT(TEXT("Replication notification function %s has too many parameters."), *TargetFuncDef->GetName());
@@ -2646,7 +2641,7 @@ void FHeaderParser::VerifyRepNotifyCallback(FUnrealPropertyDefinitionInfo& Prope
 
 		if (Properties.Num() >= 1)
 		{
-			const FUnrealPropertyDefinitionInfo* ParmDef = Properties[0];
+			const FUnrealPropertyDefinitionInfo* ParmDef = &*Properties[0];
 			// First parameter is always the old value:
 			if (!PropertyDef.SameType(*ParmDef))
 			{
@@ -2658,7 +2653,7 @@ void FHeaderParser::VerifyRepNotifyCallback(FUnrealPropertyDefinitionInfo& Prope
 
 		if (TargetFuncDef->GetProperties().Num() >= 2)
 		{
-			FUnrealPropertyDefinitionInfo* ParmDef = Properties[1];
+			FUnrealPropertyDefinitionInfo* ParmDef = &*Properties[1];
 			FPropertyBase& ParmBase = ParmDef->GetPropertyBase();
 			// A 2nd parameter for arrays can be specified as a const TArray<uint8>&. This is a list of element indices that have changed
 			if (!(ParmBase.Type == CPT_Byte && ParmBase.ArrayType == EArrayType::Dynamic && ParmDef->HasAllPropertyFlags(CPF_ConstParm | CPF_ReferenceParm)))
@@ -2676,23 +2671,23 @@ void FHeaderParser::VerifyRepNotifyCallback(FUnrealPropertyDefinitionInfo& Prope
 void FHeaderParser::VerifyPropertyMarkups(FUnrealClassDefinitionInfo& TargetClassDef)
 {
 	// Iterate over all properties, looking for those flagged as CPF_RepNotify
-	for (FUnrealPropertyDefinitionInfo* PropertyDef : TargetClassDef.GetProperties())
+	for (TSharedRef<FUnrealPropertyDefinitionInfo> PropertyDef : TargetClassDef.GetProperties())
 	{
-		auto FindTargetFunction = [&](const FName FuncName)
+		auto FindTargetFunction = [&](const FName FuncName) -> FUnrealFunctionDefinitionInfo*
 		{
 			// Search through this class and its super classes looking for the specified callback
 			for (FUnrealClassDefinitionInfo* SearchClassDef = &TargetClassDef; SearchClassDef; SearchClassDef = SearchClassDef->GetSuperClass())
 			{
 				// Since the function map is not valid yet, we have to iterate over the fields to look for the function
-				for (FUnrealFunctionDefinitionInfo* FunctionDef : SearchClassDef->GetFunctions())
+				for (TSharedRef<FUnrealFunctionDefinitionInfo> FunctionDef : SearchClassDef->GetFunctions())
 				{
 					if (FNativeClassHeaderGenerator::GetOverriddenFName(*FunctionDef) == FuncName)
 					{
-						return FunctionDef;
+						return &*FunctionDef;
 					}
 				}
 			}
-			return static_cast<FUnrealFunctionDefinitionInfo*>(nullptr);
+			return nullptr;
 		};
 
 		FPropertyBase& PropertyToken = PropertyDef->GetPropertyBase();
@@ -3919,7 +3914,7 @@ void FHeaderParser::GetVarType(
 	{
 		RequireSymbol( TEXT('<'), TEXT("'TFieldPath'") );
 
-		FFieldClass* PropertyClass = nullptr;
+		FName FieldClassName = NAME_None;
 		FToken PropertyTypeToken;
 		if (!GetToken(PropertyTypeToken, /* bNoConsts = */ true))
 		{
@@ -3927,12 +3922,8 @@ void FHeaderParser::GetVarType(
 		}
 		else
 		{
-			FFieldClass** PropertyClassPtr = FFieldClass::GetNameToFieldClassMap().Find(PropertyTypeToken.Identifier + 1);
-			if (PropertyClassPtr)
-			{
-				PropertyClass = *PropertyClassPtr;
-			}
-			else
+			FieldClassName = FName(PropertyTypeToken.Identifier + 1, FNAME_Add);
+			if (!FPropertyTraits::IsValidFieldClass(FieldClassName))
 			{
 				FError::Throwf(TEXT("Undefined property type: %s"), PropertyTypeToken.Identifier);
 			}
@@ -3940,7 +3931,7 @@ void FHeaderParser::GetVarType(
 
 		RequireSymbol(TEXT('>'), VarType.Identifier, ESymbolParseOption::CloseTemplateBracket);
 
-		VarProperty = FPropertyBase(PropertyClass, CPT_FieldPath);
+		VarProperty = FPropertyBase(FieldClassName, CPT_FieldPath);
 	}
 	else if (FUnrealEnumDefinitionInfo* EnumDef = GTypeDefinitionInfoMap.FindByName<FUnrealEnumDefinitionInfo>(VarType.Identifier))
 	{
@@ -4475,12 +4466,6 @@ FUnrealPropertyDefinitionInfo& FHeaderParser::GetVarNameAndDim
 	ELayoutMacroType        LayoutMacroType
 )
 {
-	EObjectFlags ObjectFlags = RF_Public;
-	if (VariableCategory == EVariableCategory::Member && CurrentAccessSpecifier == ACCESS_Private)
-	{
-		ObjectFlags = RF_NoFlags;
-	}
-
 	const TCHAR* HintText = GetHintText(VariableCategory);
 
 	AddModuleRelativePathToMetadata(ParentStruct, VarProperty.MetaData);
@@ -4742,11 +4727,11 @@ FUnrealPropertyDefinitionInfo& FHeaderParser::GetVarNameAndDim
 
 	// Create the property
 	FName PropertyName(Identifier);
-	FUnrealPropertyDefinitionInfo& PropDef = FPropertyTraits::CreateProperty(VarProperty, ParentStruct, PropertyName, ObjectFlags, VariableCategory, Dimensions.String, SourceFile, InputLine, InputPos);
+	TSharedRef<FUnrealPropertyDefinitionInfo> PropDefRef = FPropertyTraits::CreateProperty(VarProperty, ParentStruct, PropertyName, VariableCategory, CurrentAccessSpecifier, Dimensions.String, SourceFile, InputLine, InputPos);
 
 	// Add the property to the parent
-	ParentStruct.AddProperty(PropDef);
-	return PropDef;
+	ParentStruct.AddProperty(PropDefRef);
+	return *PropDefRef;
 }
 
 /*-----------------------------------------------------------------------------
@@ -4906,7 +4891,7 @@ bool FHeaderParser::CompileDeclaration(TArray<FUnrealFunctionDefinitionInfo*>& D
 
 	if (Token.Matches(TEXT("UFUNCTION"), ESearchCase::CaseSensitive))
 	{
-		CompileFunctionDeclaration();
+		FUnrealFunctionDefinitionInfo& FunctionDef = CompileFunctionDeclaration();
 		return true;
 	}
 
@@ -5360,7 +5345,7 @@ void PostParsingClassSetup(FUnrealClassDefinitionInfo& ClassDef)
 
 	if (!HasAllOptimizationClassFlags())
 	{
-		for (FUnrealPropertyDefinitionInfo* PropertyDef : ClassDef.GetProperties())
+		for (TSharedRef<FUnrealPropertyDefinitionInfo> PropertyDef : ClassDef.GetProperties())
 		{
 			if (PropertyDef->HasAnyPropertyFlags(CPF_Config))
 			{
@@ -6321,9 +6306,6 @@ FUnrealFunctionDefinitionInfo& FHeaderParser::CompileDelegateDeclaration(const T
 	// Optionally consume a semicolon, it's not required for the delegate macro since it contains one internally
 	MatchSemi();
 
-	// Bind the function.
-	DelegateSignatureFunction->Bind();
-
 	// End the nesting
 	PostPopFunctionDeclaration(DelegateSignatureFunctionDef);
 
@@ -6344,7 +6326,7 @@ FUnrealFunctionDefinitionInfo& FHeaderParser::CompileDelegateDeclaration(const T
 /**
  * Parses and compiles a function declaration
  */
-void FHeaderParser::CompileFunctionDeclaration()
+FUnrealFunctionDefinitionInfo& FHeaderParser::CompileFunctionDeclaration()
 {
 	CheckAllow(TEXT("'Function'"), ENestAllowFlags::Function);
 
@@ -6669,7 +6651,7 @@ void FHeaderParser::CompileFunctionDeclaration()
 	bool bHasAnyOutputs = bHasReturnValue;
 	if (!bHasAnyOutputs)
 	{
-		for (FUnrealPropertyDefinitionInfo* PropertyDef : FuncDef.GetProperties())
+		for (TSharedRef<FUnrealPropertyDefinitionInfo> PropertyDef : FuncDef.GetProperties())
 		{
 			if (PropertyDef->HasSpecificPropertyFlags(CPF_ReturnParm | CPF_OutParm, CPF_OutParm))
 			{
@@ -6776,9 +6758,6 @@ void FHeaderParser::CompileFunctionDeclaration()
 	//@TODO: UCREMOVAL: Ideally the flags didn't get copied midway thru parsing the function declaration, and we could avoid this
 	TopFunction->FunctionFlags |= FuncDefFuncInfo.FunctionFlags;
 
-	// Bind the function.
-	TopFunction->Bind();
-	
 	// Make sure that the replication flags set on an overridden function match the parent function
 	if (UFunction* SuperFunc = TopFunction->GetSuperFunction())
 	{
@@ -6797,7 +6776,7 @@ void FHeaderParser::CompileFunctionDeclaration()
 
 	if (TopFunction->FunctionFlags & (FUNC_BlueprintCallable | FUNC_BlueprintEvent))
 	{
-		for (FUnrealPropertyDefinitionInfo* PropertyDef : FuncDef.GetProperties())
+		for (TSharedRef<FUnrealPropertyDefinitionInfo> PropertyDef : FuncDef.GetProperties())
 		{
 			if (PropertyDef->IsStaticArray())
 			{
@@ -6843,6 +6822,8 @@ void FHeaderParser::CompileFunctionDeclaration()
 
 	// perform documentation policy tests
 	CheckDocumentationPolicyForFunc(UHTCastChecked<FUnrealClassDefinitionInfo>(GetCurrentClassDef()), FuncDef);
+
+	return FuncDef;
 }
 
 /** Parses optional metadata text. */
@@ -7003,7 +6984,7 @@ bool FHeaderParser::ValidateScriptStructOkForNet(const FString& OriginStructName
 		}
 	}
 
-	for (FUnrealPropertyDefinitionInfo* PropertyDef : InStructDef.GetProperties())
+	for (TSharedRef<FUnrealPropertyDefinitionInfo> PropertyDef : InStructDef.GetProperties())
 	{
 		if (PropertyDef->IsSet())
 		{
@@ -7341,14 +7322,14 @@ bool FHeaderParser::CompileStatement(TArray<FUnrealFunctionDefinitionInfo*>& Del
 void FHeaderParser::ComputeFunctionParametersSize(FUnrealClassDefinitionInfo& ClassDef)
 {
 	// Recurse with all child states in this class.
-	for (FUnrealFunctionDefinitionInfo* FunctionDef : ClassDef.GetFunctions())
+	for (TSharedRef<FUnrealFunctionDefinitionInfo> FunctionDef : ClassDef.GetFunctions())
 	{
 		UFunction* ThisFunction = FunctionDef->GetFunction();
 
 		// Fix up any structs that were used as a parameter in a delegate before being defined
 		if (FunctionDef->HasAnyFunctionFlags(FUNC_Delegate))
 		{
-			for (FUnrealPropertyDefinitionInfo* PropertyDef : FunctionDef->GetProperties())
+			for (TSharedRef<FUnrealPropertyDefinitionInfo> PropertyDef : FunctionDef->GetProperties())
 			{
 				if (PropertyDef->IsStructOrStructStaticArray())
 				{
@@ -7365,7 +7346,7 @@ void FHeaderParser::ComputeFunctionParametersSize(FUnrealClassDefinitionInfo& Cl
 		// Compute the function parameter size, propagate some flags to the outer function, and save the return offset
 		// Must be done in a second phase, as StaticLink resets various fields again!
 		ThisFunction->ParmsSize = 0;
-		for (FUnrealPropertyDefinitionInfo* PropertyDef : FunctionDef->GetProperties())
+		for (TSharedRef<FUnrealPropertyDefinitionInfo> PropertyDef : FunctionDef->GetProperties())
 		{
 			if (PropertyDef->HasSpecificPropertyFlags(CPF_ReturnParm | CPF_OutParm, CPF_OutParm))
 			{
@@ -7442,7 +7423,7 @@ void FHeaderParser::FinalizeScriptExposedFunctions(FUnrealClassDefinitionInfo& C
 {
 	// Finalize all of the children introduced in this class
 	UClass* Class = ClassDef.GetClass();
-	for (FUnrealFunctionDefinitionInfo* FunctionDef : ClassDef.GetFunctions())
+	for (TSharedRef<FUnrealFunctionDefinitionInfo> FunctionDef : ClassDef.GetFunctions())
 	{
 		UFunction* Function = FunctionDef->GetFunction();
 		Class->AddFunctionToFunctionMap(Function, Function->GetFName());
@@ -7570,9 +7551,6 @@ ECompilationResult::Type FHeaderParser::ParseHeader()
 			UClass* Class = ClassDef.GetClass();
 			PostParsingClassSetup(ClassDef);
 
-			// Clean up and exit.
-			Class->Bind();
-
 			// Finalize functions
 			FinalizeScriptExposedFunctions(ClassDef);
 
@@ -7743,6 +7721,11 @@ ECompilationResult::Type FHeaderParser::Parse(
 		}
 	}
 #if !PLATFORM_EXCEPTIONS_DISABLED
+	catch (const FUHTException& Ex)
+	{
+		FResults::LogError(Ex);
+		Result = ECompilationResult::OtherCompilationError;
+	}
 	catch (TCHAR* ErrorMsg)
 	{
 		Warn->Log(ELogVerbosity::Error, ErrorMsg);
@@ -8808,17 +8791,17 @@ void FHeaderParser::PostPopNestClass(FUnrealClassDefinitionInfo& CurrentClassDef
 
 					for (int32 Index = 0, End = ClassFunctionDef->GetProperties().Num(); Index != End; ++Index)
 					{
-						FUnrealPropertyDefinitionInfo* InterfaceParamDef = InterfaceFunctionDef->GetProperties()[Index];
-						FUnrealPropertyDefinitionInfo* ClassParamDef = ClassFunctionDef->GetProperties()[Index];
-						if (!InterfaceParamDef->MatchesType(*ClassParamDef, true))
+						FUnrealPropertyDefinitionInfo& InterfaceParamDef = *InterfaceFunctionDef->GetProperties()[Index];
+						FUnrealPropertyDefinitionInfo& ClassParamDef = *ClassFunctionDef->GetProperties()[Index];
+						if (!InterfaceParamDef.MatchesType(ClassParamDef, true))
 						{
-							if (InterfaceParamDef->HasAnyPropertyFlags(CPF_ReturnParm))
+							if (InterfaceParamDef.HasAnyPropertyFlags(CPF_ReturnParm))
 							{
 								FError::Throwf(TEXT("Implementation of function '%s' conflicts only by return type with interface '%s'"), *InterfaceFunctionDef->GetName(), *InterfaceDef->GetName());
 							}
 							else
 							{
-								FError::Throwf(TEXT("Implementation of function '%s' conflicts with interface '%s' - parameter %i '%s'"), *InterfaceFunctionDef->GetName(), *InterfaceDef->GetName(), End, *InterfaceParamDef->GetName());
+								FError::Throwf(TEXT("Implementation of function '%s' conflicts with interface '%s' - parameter %i '%s'"), *InterfaceFunctionDef->GetName(), *InterfaceDef->GetName(), End, *InterfaceParamDef.GetName());
 							}
 						}
 					}
@@ -8964,7 +8947,7 @@ void FHeaderParser::CheckDocumentationPolicyForStruct(FUnrealStructDefinitionInf
 	if (DocumentationPolicy.bMemberToolTipsRequired)
 	{
 		TMap<FString, FName> ToolTipToPropertyName;
-		for (FUnrealPropertyDefinitionInfo* PropertyDef : StructDef.GetProperties())
+		for (TSharedRef<FUnrealPropertyDefinitionInfo> PropertyDef : StructDef.GetProperties())
 		{
 			FString ToolTip = PropertyDef->GetToolTipText().ToString();
 			if (ToolTip.IsEmpty() || ToolTip.Equals(PropertyDef->GetDisplayNameText().ToString()))
@@ -8983,7 +8966,7 @@ void FHeaderParser::CheckDocumentationPolicyForStruct(FUnrealStructDefinitionInf
 
 	if (DocumentationPolicy.bFloatRangesRequired)
 	{
-		for (FUnrealPropertyDefinitionInfo* PropertyDef : StructDef.GetProperties())
+		for (TSharedRef<FUnrealPropertyDefinitionInfo> PropertyDef : StructDef.GetProperties())
 		{
 			if(DoesCPPTypeRequireDocumentation(PropertyDef->GetCPPType()))
 			{
@@ -9004,7 +8987,7 @@ void FHeaderParser::CheckDocumentationPolicyForStruct(FUnrealStructDefinitionInf
 		if (FUnrealClassDefinitionInfo* ClassDef = UHTCast<FUnrealClassDefinitionInfo>(StructDef))
 		{
 			TMap<FString, FName> ToolTipToFunc;
-			for (FUnrealFunctionDefinitionInfo* FunctionDef : ClassDef->GetFunctions())
+			for (TSharedRef<FUnrealFunctionDefinitionInfo> FunctionDef : ClassDef->GetFunctions())
 			{
 				UFunction* Func = FunctionDef->GetFunction();
 				FString ToolTip = Func->GetToolTipText().ToString();
@@ -9069,7 +9052,7 @@ void FHeaderParser::CheckDocumentationPolicyForFunc(FUnrealClassDefinitionInfo&
 		{
 			// ensure each parameter has a tooltip
 			TSet<FName> ExistingFields;
-			for (FUnrealPropertyDefinitionInfo* PropertyDef : FunctionDef.GetProperties())
+			for (TSharedRef<FUnrealPropertyDefinitionInfo> PropertyDef : FunctionDef.GetProperties())
 			{
 				FName ParamName = PropertyDef->GetFName();
 				if (ParamName == NAME_ReturnValue)
@@ -9221,10 +9204,14 @@ FUnrealFunctionDefinitionInfo& CreateFunctionImpl(FUnrealSourceFile& SourceFile,
 		FUnrealStructDefinitionInfo& StructDef = ((FStructScope*)CurrentScope)->GetStructDef();
 		UStruct* Struct = StructDef.GetStruct();
 
-		StructDef.AddFunction(FuncDef);
+		StructDef.AddFunction(FuncDefRef);
 		Function->Next = Struct->Children;
 		Struct->Children = Function;
 	}
+	else
+	{
+		SourceFile.AddDefinedFunction(FuncDefRef);
+	}
 
 	return FuncDef;
 }
diff --git a/Engine/Source/Programs/UnrealHeaderTool/Private/HeaderParser.h b/Engine/Source/Programs/UnrealHeaderTool/Private/HeaderParser.h
--- a/Engine/Source/Programs/UnrealHeaderTool/Private/HeaderParser.h
+++ b/Engine/Source/Programs/UnrealHeaderTool/Private/HeaderParser.h
@@ -439,7 +439,7 @@ protected:
 
 	FUnrealClassDefinitionInfo& CompileClassDeclaration();
 	FUnrealFunctionDefinitionInfo& CompileDelegateDeclaration(const TCHAR* DelegateIdentifier, EDelegateSpecifierAction::Type SpecifierAction = EDelegateSpecifierAction::DontParse);
-	void CompileFunctionDeclaration();
+	FUnrealFunctionDefinitionInfo& CompileFunctionDeclaration();
 	void CompileVariableDeclaration (FUnrealStructDefinitionInfo& StructDef);
 	void CompileInterfaceDeclaration();
 	void CompileRigVMMethodDeclaration(FUnrealStructDefinitionInfo& StructDef);
diff --git a/Engine/Source/Programs/UnrealHeaderTool/Private/ParserHelper.cpp b/Engine/Source/Programs/UnrealHeaderTool/Private/ParserHelper.cpp
--- a/Engine/Source/Programs/UnrealHeaderTool/Private/ParserHelper.cpp
+++ b/Engine/Source/Programs/UnrealHeaderTool/Private/ParserHelper.cpp
@@ -11,35 +11,28 @@
 /////////////////////////////////////////////////////
 // FPropertyBase
 
+FPropertyBase::FPropertyBase(EPropertyType InType)
+	: Type(InType)
+	, IntType(GetSizedIntTypeFromPropertyType(InType))
+{
+}
+
+FPropertyBase::FPropertyBase(EPropertyType InType, EIntType InIntType)
+	: Type(InType)
+	, IntType(InIntType)
+{
+}
+
 FPropertyBase::FPropertyBase(FUnrealEnumDefinitionInfo& InEnumDef, EPropertyType InType)
 	: Type(InType)
-	, ArrayType(EArrayType::None)
-	, PropertyFlags(CPF_None)
-	, ImpliedPropertyFlags(CPF_None)
-	, RefQualifier(ERefQualifier::None)
-	, PropertyExportFlags(PROPEXPORT_Public)
 	, EnumDef(&InEnumDef)
-	, DelegateName(NAME_None)
-	, DelegateSignatureOwnerClass(nullptr)
-	, RepNotifyName(NAME_None)
-	, PointerType(EPointerType::None)
 	, IntType(GetSizedIntTypeFromPropertyType(InType))
 {
 }
 
 FPropertyBase::FPropertyBase(FUnrealClassDefinitionInfo& InClassDef, EPropertyType InType, bool bWeakIsAuto/* = false*/)
 	: Type(InType)
-	, ArrayType(EArrayType::None)
-	, PropertyFlags(CPF_None)
-	, ImpliedPropertyFlags(CPF_None)
-	, RefQualifier(ERefQualifier::None)
-	, PropertyExportFlags(PROPEXPORT_Public)
 	, ClassDef(&InClassDef)
-	, DelegateName(NAME_None)
-	, DelegateSignatureOwnerClass(nullptr)
-	, RepNotifyName(NAME_None)
-	, PointerType(EPointerType::None)
-	, IntType(EIntType::None)
 {
 	if ((Type == CPT_WeakObjectReference) && bWeakIsAuto)
 	{
@@ -49,17 +42,14 @@ FPropertyBase::FPropertyBase(FUnrealClassDefinitionInfo& InClassDef, EPropertyTy
 
 FPropertyBase::FPropertyBase(FUnrealScriptStructDefinitionInfo& InStructDef)
 	: Type(CPT_Struct)
-	, ArrayType(EArrayType::None)
-	, PropertyFlags(CPF_None)
-	, ImpliedPropertyFlags(CPF_None)
-	, RefQualifier(ERefQualifier::None)
-	, PropertyExportFlags(PROPEXPORT_Public)
 	, ScriptStructDef(&InStructDef)
-	, DelegateName(NAME_None)
-	, DelegateSignatureOwnerClass(nullptr)
-	, RepNotifyName(NAME_None)
-	, PointerType(EPointerType::None)
-	, IntType(EIntType::None)
+{
+}
+
+FPropertyBase::FPropertyBase(FName InFieldClassName, EPropertyType InType)
+	: Type(InType)
+	, FieldClassName(InFieldClassName)
+	, IntType(GetSizedIntTypeFromPropertyType(InType))
 {
 }
 
diff --git a/Engine/Source/Programs/UnrealHeaderTool/Private/ParserHelper.h b/Engine/Source/Programs/UnrealHeaderTool/Private/ParserHelper.h
--- a/Engine/Source/Programs/UnrealHeaderTool/Private/ParserHelper.h
+++ b/Engine/Source/Programs/UnrealHeaderTool/Private/ParserHelper.h
@@ -169,19 +169,19 @@ class FPropertyBase : public TSharedFromThis<FPropertyBase>
 {
 public:
 	// Variables.
-	EPropertyType       Type;
-	EArrayType::Type    ArrayType;
+	EPropertyType       Type = EPropertyType::CPT_None;
+	EArrayType::Type    ArrayType = EArrayType::None;
 	EAllocatorType      AllocatorType = EAllocatorType::Default;
-	EPropertyFlags      PropertyFlags;
-	EPropertyFlags      ImpliedPropertyFlags;
-	ERefQualifier::Type RefQualifier; // This is needed because of legacy stuff - FString mangles the flags for reasons that have become lost in time but we need this info for testing for invalid replicated function signatures.
+	EPropertyFlags      PropertyFlags = CPF_None;
+	EPropertyFlags      ImpliedPropertyFlags = CPF_None;
+	ERefQualifier::Type RefQualifier = ERefQualifier::None; // This is needed because of legacy stuff - FString mangles the flags for reasons that have become lost in time but we need this info for testing for invalid replicated function signatures.
 
 	TSharedPtr<FPropertyBase> MapKeyProp;
 
 	/**
 	 * A mask of EPropertyHeaderExportFlags which are used for modifying how this property is exported to the native class header
 	 */
-	uint32 PropertyExportFlags;
+	uint32 PropertyExportFlags = PROPEXPORT_Public;
 
 	union
 	{
@@ -191,12 +191,12 @@ public:
 		FUnrealClassDefinitionInfo* ClassDef;
 		FUnrealFunctionDefinitionInfo* FunctionDef;
 	};
-	FFieldClass* PropertyPathClass = nullptr;
+	FName FieldClassName = NAME_None;
 	FUnrealClassDefinitionInfo* MetaClassDef = nullptr;
 
-	FName   DelegateName;
-	UClass*	DelegateSignatureOwnerClass;
-	FName   RepNotifyName;
+	FName   DelegateName = NAME_None;
+	UClass*	DelegateSignatureOwnerClass = nullptr;
+	FName   RepNotifyName = NAME_None;
 
 	/** Raw string (not type-checked) used for specifying special text when exporting a property to the *Classes.h file */
 	FString	ExportInfo;
@@ -204,8 +204,8 @@ public:
 	/** Map of key value pairs that will be added to the package's UMetaData for this property */
 	TMap<FName, FString> MetaData;
 
-	EPointerType::Type PointerType;
-	EIntType IntType;
+	EPointerType::Type PointerType = EPointerType::None;
+	EIntType IntType = EIntType::None;
 
 public:
 	/** @name Constructors */
@@ -215,35 +215,9 @@ public:
 	FPropertyBase& operator=(const FPropertyBase&) = default;
 	FPropertyBase& operator=(FPropertyBase&&) = default;
 
-	explicit FPropertyBase(EPropertyType InType)
-	: Type                (InType)
-	, ArrayType           (EArrayType::None)
-	, PropertyFlags       (CPF_None)
-	, ImpliedPropertyFlags(CPF_None)
-	, RefQualifier        (ERefQualifier::None)
-	, PropertyExportFlags (PROPEXPORT_Public)
-	, DelegateName        (NAME_None)
-	, DelegateSignatureOwnerClass(nullptr)
-	, RepNotifyName       (NAME_None)
-	, PointerType         (EPointerType::None)
-	, IntType             (GetSizedIntTypeFromPropertyType(InType))
-	{
-	}
+	explicit FPropertyBase(EPropertyType InType);
 
-	explicit FPropertyBase(EPropertyType InType, EIntType InIntType)
-	: Type                (InType)
-	, ArrayType           (EArrayType::None)
-	, PropertyFlags       (CPF_None)
-	, ImpliedPropertyFlags(CPF_None)
-	, RefQualifier        (ERefQualifier::None)
-	, PropertyExportFlags (PROPEXPORT_Public)
-	, DelegateName        (NAME_None)
-	, DelegateSignatureOwnerClass(nullptr)
-	, RepNotifyName       (NAME_None)
-	, PointerType         (EPointerType::None)
-	, IntType             (InIntType)
-	{
-	}
+	explicit FPropertyBase(EPropertyType InType, EIntType InIntType);
 
 	explicit FPropertyBase(FUnrealEnumDefinitionInfo& InEnumDef, EPropertyType InType);
 
@@ -251,21 +225,7 @@ public:
 
 	explicit FPropertyBase(FUnrealScriptStructDefinitionInfo& InStructDef);
 
-	explicit FPropertyBase(FFieldClass* InPropertyClass, EPropertyType InType)
-		: Type(InType)
-		, ArrayType(EArrayType::None)
-		, PropertyFlags(CPF_None)
-		, ImpliedPropertyFlags(CPF_None)
-		, RefQualifier(ERefQualifier::None)
-		, PropertyExportFlags(PROPEXPORT_Public)
-		, PropertyPathClass(InPropertyClass)
-		, DelegateName(NAME_None)
-		, DelegateSignatureOwnerClass(nullptr)
-		, RepNotifyName(NAME_None)
-		, PointerType(EPointerType::None)
-		, IntType(GetSizedIntTypeFromPropertyType(InType))
-	{
-	}
+	explicit FPropertyBase(FName InFieldClassName, EPropertyType InType);
 	//@}
 
 	/** @name Functions */
diff --git a/Engine/Source/Programs/UnrealHeaderTool/Private/PropertyTypes.cpp b/Engine/Source/Programs/UnrealHeaderTool/Private/PropertyTypes.cpp
--- a/Engine/Source/Programs/UnrealHeaderTool/Private/PropertyTypes.cpp
+++ b/Engine/Source/Programs/UnrealHeaderTool/Private/PropertyTypes.cpp
@@ -101,6 +101,56 @@ namespace
 		}
 	};
 
+	//----------------------------------------------------------------------------------------------------------------------------------------------------------
+	//----------------------------------------------------------------------------------------------------------------------------------------------------------
+	// Field class bootstrap system
+	//----------------------------------------------------------------------------------------------------------------------------------------------------------
+	//----------------------------------------------------------------------------------------------------------------------------------------------------------
+
+	class FFieldClassInit
+	{
+	public:
+		FFieldClassInit()
+		{
+			FProperty::StaticClass();
+			FNumericProperty::StaticClass();
+			FByteProperty::StaticClass();
+			FInt8Property::StaticClass();
+			FInt16Property::StaticClass();
+			FIntProperty::StaticClass();
+			FInt64Property::StaticClass();
+			FUInt16Property::StaticClass();
+			FUInt32Property::StaticClass();
+			FUInt64Property::StaticClass();
+			FFloatProperty::StaticClass();
+			FDoubleProperty::StaticClass();
+			FLargeWorldCoordinatesRealProperty::StaticClass();
+			FBoolProperty::StaticClass();
+			FEnumProperty::StaticClass();
+			//TObjectPropertyBase::StaticClass();
+			FObjectProperty::StaticClass();
+			FClassProperty::StaticClass();
+			FClassPtrProperty::StaticClass();
+			FObjectPtrProperty::StaticClass();
+			FWeakObjectProperty::StaticClass();
+			FLazyObjectProperty::StaticClass();
+			FSoftObjectProperty::StaticClass();
+			FSoftClassProperty::StaticClass();
+			FInterfaceProperty::StaticClass();
+			FNameProperty::StaticClass();
+			FStrProperty::StaticClass();
+			FTextProperty::StaticClass();
+			FStructProperty::StaticClass();
+			FMulticastDelegateProperty::StaticClass();
+			FMulticastSparseDelegateProperty::StaticClass();
+			FMulticastInlineDelegateProperty::StaticClass();
+			FFieldPathProperty::StaticClass();
+			FArrayProperty::StaticClass();
+			FSetProperty::StaticClass();
+			FMapProperty::StaticClass();
+		}
+	};
+
 	//----------------------------------------------------------------------------------------------------------------------------------------------------------
 	//----------------------------------------------------------------------------------------------------------------------------------------------------------
 	// Dispatch system
@@ -202,9 +252,9 @@ namespace
 	template<typename TraitsType>
 	struct CreatePropertyDispatch
 	{
-		FProperty* operator()(FUnrealPropertyDefinitionInfo& PropDef, FFieldVariant Scope, const FName& Name, EObjectFlags ObjectFlags)
+		void operator()(FUnrealPropertyDefinitionInfo& PropDef)
 		{
-			return TraitsType::CreateProperty(PropDef, Scope, Name, ObjectFlags);
+			TraitsType::CreateProperty(PropDef);
 		}
 	};
 
@@ -213,7 +263,7 @@ namespace
 	{
 		FProperty* operator()(FUnrealPropertyDefinitionInfo& PropDef, FFieldVariant Scope, const FName& Name, EObjectFlags ObjectFlags)
 		{
-			return TraitsType::CreateProperty(PropDef, Scope, Name, ObjectFlags);
+			return TraitsType::CreateEngineType(PropDef, Scope, Name, ObjectFlags);
 		}
 	};
 
@@ -260,9 +310,15 @@ namespace
 	//----------------------------------------------------------------------------------------------------------------------------------------------------------
 
 	template <bool bHandleContainers>
-	FProperty* CreatePropertyHelper(FUnrealPropertyDefinitionInfo& PropDef, FFieldVariant Scope, const FName& Name, EObjectFlags ObjectFlags)
+	void CreatePropertyHelper(FUnrealPropertyDefinitionInfo& PropDef)
 	{
-		return PropertyTypeDispatch<CreatePropertyDispatch, bHandleContainers, FProperty*>(PropDef.GetPropertyBase(), std::ref(PropDef), Scope, std::ref(Name), ObjectFlags);
+		return PropertyTypeDispatch<CreatePropertyDispatch, bHandleContainers, void>(PropDef.GetPropertyBase(), std::ref(PropDef));
+	}
+
+	template <bool bHandleContainers>
+	FProperty* CreateEngineTypeHelper(FUnrealPropertyDefinitionInfo& PropDef, FFieldVariant Scope, const FName& Name, EObjectFlags ObjectFlags)
+	{
+		return PropertyTypeDispatch<CreateEngineTypeDispatch, bHandleContainers, FProperty*>(PropDef.GetPropertyBase(), std::ref(PropDef), Scope, std::ref(Name), ObjectFlags);
 	}
 
 	bool IsSupportedByBlueprintSansContainers(const FUnrealPropertyDefinitionInfo& PropDef, bool bMemberVariable)
@@ -296,12 +352,8 @@ struct FPropertyTypeTraitsBase
 	 * NOTE: This method MUST be implemented by each of the property types.
 	 *
 	 * @param PropDef The definition of the property
-	 * @param Scope The parent object owning the property
-	 * @param Name The name of the property
-	 * @param ObjectFlags The flags associated with the property
-	 * @return The pointer to the newly created property.  It will be attached to the definition by the caller
 	 */
-	 //static FProperty* CreateProperty(FUnrealPropertyDefinitionInfo& PropDef, FFieldVariant Scope, const FName& Name, EObjectFlags ObjectFlags)
+	 //static void CreateProperty(FUnrealPropertyDefinitionInfo& PropDef)
 
 	/**
 	 * Given a property definition create the underlying engine type
@@ -314,7 +366,7 @@ struct FPropertyTypeTraitsBase
 	 * @param ObjectFlags The flags associated with the property
 	 * @return The pointer to the newly created property.  It will be attached to the definition by the caller
 	 */
-	 //static FProperty* CreateProperty(FUnrealPropertyDefinitionInfo& PropDef, FFieldVariant Scope, const FName& Name, EObjectFlags ObjectFlags)
+	 //static FProperty* CreateEngineType(FUnrealPropertyDefinitionInfo& PropDef, FFieldVariant Scope, const FName& Name, EObjectFlags ObjectFlags)
 
 	/**
 	 * Returns true if this property is supported by blueprints
@@ -370,13 +422,17 @@ struct FPropertyTypeTraitsByte : public FPropertyTypeTraitsNumericBase
 		return false;
 	}
 
-	static FProperty* CreateProperty(FUnrealPropertyDefinitionInfo& PropDef, FFieldVariant Scope, const FName& Name, EObjectFlags ObjectFlags)
+	static void CreateProperty(FUnrealPropertyDefinitionInfo& PropDef)
 	{
-		const FPropertyBase& VarProperty = PropDef.GetPropertyBase();
-		UEnum* Enum = VarProperty.EnumDef ? VarProperty.EnumDef->GetEnum() : nullptr;
 #if UHT_ENABLE_VALUE_PROPERTY_TAG
-		PropDef.GetUnrealSourceFile().AddTypeDefIncludeIfNeeded(VarProperty.EnumDef);
+		PropDef.GetUnrealSourceFile().AddTypeDefIncludeIfNeeded(PropDef.GetPropertyBase().EnumDef);
 #endif
+	}
+
+	static FProperty* CreateEngineType(FUnrealPropertyDefinitionInfo& PropDef, FFieldVariant Scope, const FName& Name, EObjectFlags ObjectFlags)
+	{
+		const FPropertyBase& VarProperty = PropDef.GetPropertyBase();
+		UEnum* Enum = VarProperty.EnumDef ? VarProperty.EnumDef->GetEnum() : nullptr;
 		FByteProperty* Result = new FByteProperty(Scope, Name, ObjectFlags);
 		Result->Enum = Enum;
 		check(VarProperty.IntType == EIntType::Sized);
@@ -440,7 +496,11 @@ struct FPropertyTypeTraitsByte : public FPropertyTypeTraitsNumericBase
 
 struct FPropertyTypeTraitsInt8 : public FPropertyTypeTraitsNumericBase
 {
-	static FProperty* CreateProperty(FUnrealPropertyDefinitionInfo& PropDef, FFieldVariant Scope, const FName& Name, EObjectFlags ObjectFlags)
+	static void CreateProperty(FUnrealPropertyDefinitionInfo& PropDef)
+	{
+	}
+
+	static FProperty* CreateEngineType(FUnrealPropertyDefinitionInfo& PropDef, FFieldVariant Scope, const FName& Name, EObjectFlags ObjectFlags)
 	{
 		const FPropertyBase& VarProperty = PropDef.GetPropertyBase();
 		FInt8Property* Result = new FInt8Property(Scope, Name, ObjectFlags);
@@ -461,7 +521,11 @@ struct FPropertyTypeTraitsInt8 : public FPropertyTypeTraitsNumericBase
 
 struct FPropertyTypeTraitsInt16 : public FPropertyTypeTraitsNumericBase
 {
-	static FProperty* CreateProperty(FUnrealPropertyDefinitionInfo& PropDef, FFieldVariant Scope, const FName& Name, EObjectFlags ObjectFlags)
+	static void CreateProperty(FUnrealPropertyDefinitionInfo& PropDef)
+	{
+	}
+
+	static FProperty* CreateEngineType(FUnrealPropertyDefinitionInfo& PropDef, FFieldVariant Scope, const FName& Name, EObjectFlags ObjectFlags)
 	{
 		const FPropertyBase& VarProperty = PropDef.GetPropertyBase();
 		FInt16Property* Result = new FInt16Property(Scope, Name, ObjectFlags);
@@ -493,7 +557,11 @@ struct FPropertyTypeTraitsInt : public FPropertyTypeTraitsNumericBase
 		return false;
 	}
 
-	static FProperty* CreateProperty(FUnrealPropertyDefinitionInfo& PropDef, FFieldVariant Scope, const FName& Name, EObjectFlags ObjectFlags)
+	static void CreateProperty(FUnrealPropertyDefinitionInfo& PropDef)
+	{
+	}
+
+	static FProperty* CreateEngineType(FUnrealPropertyDefinitionInfo& PropDef, FFieldVariant Scope, const FName& Name, EObjectFlags ObjectFlags)
 	{
 		const FPropertyBase& VarProperty = PropDef.GetPropertyBase();
 		FIntProperty* Result = new FIntProperty(Scope, Name, ObjectFlags);
@@ -530,7 +598,11 @@ struct FPropertyTypeTraitsInt64 : public FPropertyTypeTraitsNumericBase
 		return false;
 	}
 
-	static FProperty* CreateProperty(FUnrealPropertyDefinitionInfo& PropDef, FFieldVariant Scope, const FName& Name, EObjectFlags ObjectFlags)
+	static void CreateProperty(FUnrealPropertyDefinitionInfo& PropDef)
+	{
+	}
+
+	static FProperty* CreateEngineType(FUnrealPropertyDefinitionInfo& PropDef, FFieldVariant Scope, const FName& Name, EObjectFlags ObjectFlags)
 	{
 		const FPropertyBase& VarProperty = PropDef.GetPropertyBase();
 		FInt64Property* Result = new FInt64Property(Scope, Name, ObjectFlags);
@@ -556,7 +628,11 @@ struct FPropertyTypeTraitsInt64 : public FPropertyTypeTraitsNumericBase
 
 struct FPropertyTypeTraitsUInt16 : public FPropertyTypeTraitsNumericBase
 {
-	static FProperty* CreateProperty(FUnrealPropertyDefinitionInfo& PropDef, FFieldVariant Scope, const FName& Name, EObjectFlags ObjectFlags)
+	static void CreateProperty(FUnrealPropertyDefinitionInfo& PropDef)
+	{
+	}
+
+	static FProperty* CreateEngineType(FUnrealPropertyDefinitionInfo& PropDef, FFieldVariant Scope, const FName& Name, EObjectFlags ObjectFlags)
 	{
 		const FPropertyBase& VarProperty = PropDef.GetPropertyBase();
 		FUInt16Property* Result = new FUInt16Property(Scope, Name, ObjectFlags);
@@ -577,7 +653,11 @@ struct FPropertyTypeTraitsUInt16 : public FPropertyTypeTraitsNumericBase
 
 struct FPropertyTypeTraitsUInt32 : public FPropertyTypeTraitsNumericBase
 {
-	static FProperty* CreateProperty(FUnrealPropertyDefinitionInfo& PropDef, FFieldVariant Scope, const FName& Name, EObjectFlags ObjectFlags)
+	static void CreateProperty(FUnrealPropertyDefinitionInfo& PropDef)
+	{
+	}
+
+	static FProperty* CreateEngineType(FUnrealPropertyDefinitionInfo& PropDef, FFieldVariant Scope, const FName& Name, EObjectFlags ObjectFlags)
 	{
 		const FPropertyBase& VarProperty = PropDef.GetPropertyBase();
 		FUInt32Property* Result = new FUInt32Property(Scope, Name, ObjectFlags);
@@ -598,7 +678,11 @@ struct FPropertyTypeTraitsUInt32 : public FPropertyTypeTraitsNumericBase
 
 struct FPropertyTypeTraitsUInt64 : public FPropertyTypeTraitsNumericBase
 {
-	static FProperty* CreateProperty(FUnrealPropertyDefinitionInfo& PropDef, FFieldVariant Scope, const FName& Name, EObjectFlags ObjectFlags)
+	static void CreateProperty(FUnrealPropertyDefinitionInfo& PropDef)
+	{
+	}
+
+	static FProperty* CreateEngineType(FUnrealPropertyDefinitionInfo& PropDef, FFieldVariant Scope, const FName& Name, EObjectFlags ObjectFlags)
 	{
 		const FPropertyBase& VarProperty = PropDef.GetPropertyBase();
 		FUInt64Property* Result = new FUInt64Property(Scope, Name, ObjectFlags);
@@ -630,7 +714,11 @@ struct FPropertyTypeTraitsFloat : public FPropertyTypeTraitsNumericBase
 		return false;
 	}
 
-	static FProperty* CreateProperty(FUnrealPropertyDefinitionInfo& PropDef, FFieldVariant Scope, const FName& Name, EObjectFlags ObjectFlags)
+	static void CreateProperty(FUnrealPropertyDefinitionInfo& PropDef)
+	{
+	}
+
+	static FProperty* CreateEngineType(FUnrealPropertyDefinitionInfo& PropDef, FFieldVariant Scope, const FName& Name, EObjectFlags ObjectFlags)
 	{
 		FFloatProperty* Result = new FFloatProperty(Scope, Name, ObjectFlags);
 		return Result;
@@ -665,7 +753,11 @@ struct FPropertyTypeTraitsDouble : public FPropertyTypeTraitsNumericBase
 		return false;
 	}
 
-	static FProperty* CreateProperty(FUnrealPropertyDefinitionInfo& PropDef, FFieldVariant Scope, const FName& Name, EObjectFlags ObjectFlags)
+	static void CreateProperty(FUnrealPropertyDefinitionInfo& PropDef)
+	{
+	}
+
+	static FProperty* CreateEngineType(FUnrealPropertyDefinitionInfo& PropDef, FFieldVariant Scope, const FName& Name, EObjectFlags ObjectFlags)
 	{
 		FDoubleProperty* Result = new FDoubleProperty(Scope, Name, ObjectFlags);
 		return Result;
@@ -700,7 +792,11 @@ struct FPropertyTypeTraitsLargeWorldCoordinatesReal : public FPropertyTypeTraits
 		return false;
 	}
 
-	static FProperty* CreateProperty(FUnrealPropertyDefinitionInfo& PropDef, FFieldVariant Scope, const FName& Name, EObjectFlags ObjectFlags)
+	static void CreateProperty(FUnrealPropertyDefinitionInfo& PropDef)
+	{
+	}
+
+	static FProperty* CreateEngineType(FUnrealPropertyDefinitionInfo& PropDef, FFieldVariant Scope, const FName& Name, EObjectFlags ObjectFlags)
 	{
 		FLargeWorldCoordinatesRealProperty* Result = new FLargeWorldCoordinatesRealProperty(Scope, Name, ObjectFlags);
 		return Result;
@@ -733,8 +829,12 @@ struct FPropertyTypeTraitsLargeWorldCoordinatesReal : public FPropertyTypeTraits
  */
 struct FPropertyTypeTraitsBooleanBase : public FPropertyTypeTraitsBase
 {
+	static void CreateProperty(FUnrealPropertyDefinitionInfo& PropDef)
+	{
+	}
+
 	template <typename SizeType, bool bIsNativeBool>
-	static FProperty* CreatePropertyHelper(FUnrealPropertyDefinitionInfo& PropDef, FFieldVariant Scope, const FName& Name, EObjectFlags ObjectFlags)
+	static FProperty* CreateEngineTypeHelper(FUnrealPropertyDefinitionInfo& PropDef, FFieldVariant Scope, const FName& Name, EObjectFlags ObjectFlags)
 	{
 		FBoolProperty* Result = new FBoolProperty(Scope, Name, ObjectFlags);
 		bool bActsLikeNativeBool = bIsNativeBool || PropDef.GetVariableCategory() == EVariableCategory::Return;
@@ -788,9 +888,9 @@ struct FPropertyTypeTraitsBooleanBase : public FPropertyTypeTraitsBase
 
 struct FPropertyTypeTraitsBool : public FPropertyTypeTraitsBooleanBase
 {
-	static FProperty* CreateProperty(FUnrealPropertyDefinitionInfo& PropDef, FFieldVariant Scope, const FName& Name, EObjectFlags ObjectFlags)
+	static FProperty* CreateEngineType(FUnrealPropertyDefinitionInfo& PropDef, FFieldVariant Scope, const FName& Name, EObjectFlags ObjectFlags)
 	{
-		return CreatePropertyHelper<bool, true>(PropDef, Scope, Name, ObjectFlags);
+		return CreateEngineTypeHelper<bool, true>(PropDef, Scope, Name, ObjectFlags);
 	}
 
 	static FString GetCPPType(const FUnrealPropertyDefinitionInfo& PropDef, FString* ExtendedTypeText, uint32 CPPExportFlags)
@@ -801,9 +901,9 @@ struct FPropertyTypeTraitsBool : public FPropertyTypeTraitsBooleanBase
 
 struct FPropertyTypeTraitsBool8 : public FPropertyTypeTraitsBooleanBase
 {
-	static FProperty* CreateProperty(FUnrealPropertyDefinitionInfo& PropDef, FFieldVariant Scope, const FName& Name, EObjectFlags ObjectFlags)
+	static FProperty* CreateEngineType(FUnrealPropertyDefinitionInfo& PropDef, FFieldVariant Scope, const FName& Name, EObjectFlags ObjectFlags)
 	{
-		return CreatePropertyHelper<uint8, false>(PropDef, Scope, Name, ObjectFlags);
+		return CreateEngineTypeHelper<uint8, false>(PropDef, Scope, Name, ObjectFlags);
 	}
 
 	static FString GetCPPType(const FUnrealPropertyDefinitionInfo& PropDef, FString* ExtendedTypeText, uint32 CPPExportFlags)
@@ -814,9 +914,9 @@ struct FPropertyTypeTraitsBool8 : public FPropertyTypeTraitsBooleanBase
 
 struct FPropertyTypeTraitsBool16 : public FPropertyTypeTraitsBooleanBase
 {
-	static FProperty* CreateProperty(FUnrealPropertyDefinitionInfo& PropDef, FFieldVariant Scope, const FName& Name, EObjectFlags ObjectFlags)
+	static FProperty* CreateEngineType(FUnrealPropertyDefinitionInfo& PropDef, FFieldVariant Scope, const FName& Name, EObjectFlags ObjectFlags)
 	{
-		return CreatePropertyHelper<uint16, false>(PropDef, Scope, Name, ObjectFlags);
+		return CreateEngineTypeHelper<uint16, false>(PropDef, Scope, Name, ObjectFlags);
 	}
 
 	static FString GetCPPType(const FUnrealPropertyDefinitionInfo& PropDef, FString* ExtendedTypeText, uint32 CPPExportFlags)
@@ -827,9 +927,9 @@ struct FPropertyTypeTraitsBool16 : public FPropertyTypeTraitsBooleanBase
 
 struct FPropertyTypeTraitsBool32 : public FPropertyTypeTraitsBooleanBase
 {
-	static FProperty* CreateProperty(FUnrealPropertyDefinitionInfo& PropDef, FFieldVariant Scope, const FName& Name, EObjectFlags ObjectFlags)
+	static FProperty* CreateEngineType(FUnrealPropertyDefinitionInfo& PropDef, FFieldVariant Scope, const FName& Name, EObjectFlags ObjectFlags)
 	{
-		return CreatePropertyHelper<uint32, false>(PropDef, Scope, Name, ObjectFlags);
+		return CreateEngineTypeHelper<uint32, false>(PropDef, Scope, Name, ObjectFlags);
 	}
 
 	static FString GetCPPType(const FUnrealPropertyDefinitionInfo& PropDef, FString* ExtendedTypeText, uint32 CPPExportFlags)
@@ -840,9 +940,9 @@ struct FPropertyTypeTraitsBool32 : public FPropertyTypeTraitsBooleanBase
 
 struct FPropertyTypeTraitsBool64 : public FPropertyTypeTraitsBooleanBase
 {
-	static FProperty* CreateProperty(FUnrealPropertyDefinitionInfo& PropDef, FFieldVariant Scope, const FName& Name, EObjectFlags ObjectFlags)
+	static FProperty* CreateEngineType(FUnrealPropertyDefinitionInfo& PropDef, FFieldVariant Scope, const FName& Name, EObjectFlags ObjectFlags)
 	{
-		return CreatePropertyHelper<uint64, false>(PropDef, Scope, Name, ObjectFlags);
+		return CreateEngineTypeHelper<uint64, false>(PropDef, Scope, Name, ObjectFlags);
 	}
 
 	static FString GetCPPType(const FUnrealPropertyDefinitionInfo& PropDef, FString* ExtendedTypeText, uint32 CPPExportFlags)
@@ -876,7 +976,7 @@ struct FPropertyTypeTraitsEnum : public FPropertyTypeTraitsBase
 		return true;
 	}
 
-	static FProperty* CreateProperty(FUnrealPropertyDefinitionInfo& PropDef, FFieldVariant Scope, const FName& Name, EObjectFlags ObjectFlags)
+	static void CreateProperty(FUnrealPropertyDefinitionInfo& PropDef)
 	{
 		const FPropertyBase& VarProperty = PropDef.GetPropertyBase();
 
@@ -885,7 +985,7 @@ struct FPropertyTypeTraitsEnum : public FPropertyTypeTraitsBase
 		if (Enum->GetCppForm() != UEnum::ECppForm::EnumClass)
 		{
 			check(VarProperty.Type == EPropertyType::CPT_Byte);
-			return FPropertyTypeTraitsByte::CreateProperty(PropDef, Scope, Name, ObjectFlags);
+			return FPropertyTypeTraitsByte::CreateProperty(PropDef);
 		}
 
 #if UHT_ENABLE_VALUE_PROPERTY_TAG
@@ -917,12 +1017,25 @@ struct FPropertyTypeTraitsEnum : public FPropertyTypeTraitsBase
 			check(false);
 		}
 
+		PropDef.SetValuePropDef(FPropertyTraits::CreateProperty(UnderlyingProperty, PropDef, TEXT("UnderlyingType"), 
+			PropDef.GetVariableCategory(), ACCESS_Public, PropDef.GetArrayDimensions(), PropDef.GetUnrealSourceFile(), PropDef.GetLineNumber(), PropDef.GetParsePosition()));
+	}
+
+	static FProperty* CreateEngineType(FUnrealPropertyDefinitionInfo& PropDef, FFieldVariant Scope, const FName& Name, EObjectFlags ObjectFlags)
+	{
+		const FPropertyBase& VarProperty = PropDef.GetPropertyBase();
+
+		UEnum* Enum = VarProperty.EnumDef->GetEnum();
+
+		if (Enum->GetCppForm() != UEnum::ECppForm::EnumClass)
+		{
+			check(VarProperty.Type == EPropertyType::CPT_Byte);
+			return FPropertyTypeTraitsByte::CreateEngineType(PropDef, Scope, Name, ObjectFlags);
+		}
+
 		FEnumProperty* Result = new FEnumProperty(Scope, Name, ObjectFlags);
 		PropDef.SetProperty(Result);
-		FUnrealPropertyDefinitionInfo& SubProp = FPropertyTraits::CreateProperty(UnderlyingProperty, PropDef, TEXT("UnderlyingType"), ObjectFlags,
-			PropDef.GetVariableCategory(), PropDef.GetArrayDimensions(), PropDef.GetUnrealSourceFile(), PropDef.GetLineNumber(), PropDef.GetParsePosition());
-		PropDef.SetValuePropDef(SubProp);
-		Result->UnderlyingProp = CastFieldChecked<FNumericProperty>(SubProp.GetProperty());
+		Result->UnderlyingProp = CastFieldChecked<FNumericProperty>(FPropertyTraits::CreateEngineType(PropDef.GetValuePropDefRef()));
 		Result->Enum = Enum;
 		return Result;
 	}
@@ -1015,7 +1128,7 @@ struct FPropertyTypeTraitsObjectBase : public FPropertyTypeTraitsBase
 
 struct FPropertyTypeTraitsObjectReference : public FPropertyTypeTraitsObjectBase
 {
-	static FProperty* CreateProperty(FUnrealPropertyDefinitionInfo& PropDef, FFieldVariant Scope, const FName& Name, EObjectFlags ObjectFlags)
+	static void CreateProperty(FUnrealPropertyDefinitionInfo& PropDef)
 	{
 		FPropertyBase& VarProperty = PropDef.GetPropertyBase();
 		check(VarProperty.ClassDef);
@@ -1027,10 +1140,6 @@ struct FPropertyTypeTraitsObjectReference : public FPropertyTypeTraitsObjectBase
 
 		if (PropertyClass->IsChildOf(UClass::StaticClass()))
 		{
-			FClassProperty* Result = new FClassProperty(Scope, Name, ObjectFlags);
-			Result->MetaClass = VarProperty.MetaClassDef ? VarProperty.MetaClassDef->GetClass() : nullptr;
-			Result->PropertyClass = PropertyClass;
-			return Result;
 		}
 		else
 		{
@@ -1039,7 +1148,24 @@ struct FPropertyTypeTraitsObjectReference : public FPropertyTypeTraitsObjectBase
 				VarProperty.PropertyFlags |= CPF_InstancedReference;
 				AddEditInlineMetaData(VarProperty.MetaData);
 			}
+		}
+	}
+
+	static FProperty* CreateEngineType(FUnrealPropertyDefinitionInfo& PropDef, FFieldVariant Scope, const FName& Name, EObjectFlags ObjectFlags)
+	{
+		FPropertyBase& VarProperty = PropDef.GetPropertyBase();
+		check(VarProperty.ClassDef);
+		UClass* PropertyClass = VarProperty.ClassDef->GetClass();
 
+		if (PropertyClass->IsChildOf(UClass::StaticClass()))
+		{
+			FClassProperty* Result = new FClassProperty(Scope, Name, ObjectFlags);
+			Result->MetaClass = VarProperty.MetaClassDef ? VarProperty.MetaClassDef->GetClass() : nullptr;
+			Result->PropertyClass = PropertyClass;
+			return Result;
+		}
+		else
+		{
 			FObjectProperty* Result = new FObjectProperty(Scope, Name, ObjectFlags);
 			Result->PropertyClass = PropertyClass;
 			return Result;
@@ -1108,18 +1234,21 @@ struct FPropertyTypeTraitsObjectReference : public FPropertyTypeTraitsObjectBase
 
 struct FPropertyTypeTraitsWeakObjectReference : public FPropertyTypeTraitsObjectBase
 {
-	static FProperty* CreateProperty(FUnrealPropertyDefinitionInfo& PropDef, FFieldVariant Scope, const FName& Name, EObjectFlags ObjectFlags)
+	static void CreateProperty(FUnrealPropertyDefinitionInfo& PropDef)
 	{
 		const FPropertyBase& VarProperty = PropDef.GetPropertyBase();
 		check(VarProperty.ClassDef);
-		UClass* PropertyClass = VarProperty.ClassDef->GetClass();
-
 #if UHT_ENABLE_PTR_PROPERTY_TAG
 		PropDef.GetUnrealSourceFile().AddTypeDefIncludeIfNeeded(VarProperty.ClassDef);
 #endif
+	}
 
+	static FProperty* CreateEngineType(FUnrealPropertyDefinitionInfo& PropDef, FFieldVariant Scope, const FName& Name, EObjectFlags ObjectFlags)
+	{
+		const FPropertyBase& VarProperty = PropDef.GetPropertyBase();
+		check(VarProperty.ClassDef);
 		FWeakObjectProperty* Result = new FWeakObjectProperty(Scope, Name, ObjectFlags);
-		Result->PropertyClass = PropertyClass;
+		Result->PropertyClass = VarProperty.ClassDef->GetClass();
 		return Result;
 	}
 
@@ -1157,18 +1286,21 @@ struct FPropertyTypeTraitsWeakObjectReference : public FPropertyTypeTraitsObject
 
 struct FPropertyTypeTraitsLazyObjectReference : public FPropertyTypeTraitsObjectBase
 {
-	static FProperty* CreateProperty(FUnrealPropertyDefinitionInfo& PropDef, FFieldVariant Scope, const FName& Name, EObjectFlags ObjectFlags)
+	static void CreateProperty(FUnrealPropertyDefinitionInfo& PropDef)
 	{
 		const FPropertyBase& VarProperty = PropDef.GetPropertyBase();
 		check(VarProperty.ClassDef);
-		UClass* PropertyClass = VarProperty.ClassDef->GetClass();
-
 #if UHT_ENABLE_PTR_PROPERTY_TAG
 		PropDef.GetUnrealSourceFile().AddTypeDefIncludeIfNeeded(VarProperty.ClassDef);
 #endif
+	}
 
+	static FProperty* CreateEngineType(FUnrealPropertyDefinitionInfo& PropDef, FFieldVariant Scope, const FName& Name, EObjectFlags ObjectFlags)
+	{
+		const FPropertyBase& VarProperty = PropDef.GetPropertyBase();
+		check(VarProperty.ClassDef);
 		FLazyObjectProperty* Result = new FLazyObjectProperty(Scope, Name, ObjectFlags);
-		Result->PropertyClass = PropertyClass;
+		Result->PropertyClass = VarProperty.ClassDef->GetClass();
 		return Result;
 	}
 
@@ -1196,7 +1328,7 @@ struct FPropertyTypeTraitsLazyObjectReference : public FPropertyTypeTraitsObject
 
 struct FPropertyTypeTraitsObjectPtrReference : public FPropertyTypeTraitsObjectBase
 {
-	static FProperty* CreateProperty(FUnrealPropertyDefinitionInfo& PropDef, FFieldVariant Scope, const FName& Name, EObjectFlags ObjectFlags)
+	static void CreateProperty(FUnrealPropertyDefinitionInfo& PropDef)
 	{
 		FPropertyBase& VarProperty = PropDef.GetPropertyBase();
 		check(VarProperty.ClassDef);
@@ -1208,10 +1340,6 @@ struct FPropertyTypeTraitsObjectPtrReference : public FPropertyTypeTraitsObjectB
 
 		if (PropertyClass->IsChildOf(UClass::StaticClass()))
 		{
-			FClassPtrProperty* Result = new FClassPtrProperty(Scope, Name, ObjectFlags);
-			Result->MetaClass = VarProperty.MetaClassDef ? VarProperty.MetaClassDef->GetClass() : nullptr;
-			Result->PropertyClass = PropertyClass;
-			return Result;
 		}
 		else
 		{
@@ -1220,7 +1348,24 @@ struct FPropertyTypeTraitsObjectPtrReference : public FPropertyTypeTraitsObjectB
 				VarProperty.PropertyFlags |= CPF_InstancedReference;
 				AddEditInlineMetaData(VarProperty.MetaData);
 			}
+		}
+	}
+
+	static FProperty* CreateEngineType(FUnrealPropertyDefinitionInfo& PropDef, FFieldVariant Scope, const FName& Name, EObjectFlags ObjectFlags)
+	{
+		FPropertyBase& VarProperty = PropDef.GetPropertyBase();
+		check(VarProperty.ClassDef);
+		UClass* PropertyClass = VarProperty.ClassDef->GetClass();
 
+		if (PropertyClass->IsChildOf(UClass::StaticClass()))
+		{
+			FClassPtrProperty* Result = new FClassPtrProperty(Scope, Name, ObjectFlags);
+			Result->MetaClass = VarProperty.MetaClassDef ? VarProperty.MetaClassDef->GetClass() : nullptr;
+			Result->PropertyClass = PropertyClass;
+			return Result;
+		}
+		else
+		{
 			FObjectPtrProperty* Result = new FObjectPtrProperty(Scope, Name, ObjectFlags);
 			Result->PropertyClass = PropertyClass;
 			return Result;
@@ -1281,15 +1426,20 @@ struct FPropertyTypeTraitsObjectPtrReference : public FPropertyTypeTraitsObjectB
 
 struct FPropertyTypeTraitsSoftObjectReference : public FPropertyTypeTraitsObjectBase
 {
-	static FProperty* CreateProperty(FUnrealPropertyDefinitionInfo& PropDef, FFieldVariant Scope, const FName& Name, EObjectFlags ObjectFlags)
+	static void CreateProperty(FUnrealPropertyDefinitionInfo& PropDef)
 	{
 		FPropertyBase& VarProperty = PropDef.GetPropertyBase();
 		check(VarProperty.ClassDef);
-		UClass* PropertyClass = VarProperty.ClassDef->GetClass();
-
 #if UHT_ENABLE_PTR_PROPERTY_TAG
 		PropDef.GetUnrealSourceFile().AddTypeDefIncludeIfNeeded(VarProperty.ClassDef);
 #endif
+	}
+
+	static FProperty* CreateEngineType(FUnrealPropertyDefinitionInfo& PropDef, FFieldVariant Scope, const FName& Name, EObjectFlags ObjectFlags)
+	{
+		FPropertyBase& VarProperty = PropDef.GetPropertyBase();
+		check(VarProperty.ClassDef);
+		UClass* PropertyClass = VarProperty.ClassDef->GetClass();
 
 		if (PropertyClass->IsChildOf(UClass::StaticClass()))
 		{
@@ -1361,7 +1511,7 @@ struct FPropertyTypeTraitsSoftObjectReference : public FPropertyTypeTraitsObject
 
 struct FPropertyTypeTraitsInterface : public FPropertyTypeTraitsBase
 {
-	static FProperty* CreateProperty(FUnrealPropertyDefinitionInfo& PropDef, FFieldVariant Scope, const FName& Name, EObjectFlags ObjectFlags)
+	static void CreateProperty(FUnrealPropertyDefinitionInfo& PropDef)
 	{
 		FPropertyBase& VarProperty = PropDef.GetPropertyBase();
 		check(VarProperty.ClassDef);
@@ -1371,7 +1521,14 @@ struct FPropertyTypeTraitsInterface : public FPropertyTypeTraitsBase
 #if UHT_ENABLE_PTR_PROPERTY_TAG
 		PropDef.GetUnrealSourceFile().AddTypeDefIncludeIfNeeded(VarProperty.ClassDef);
 #endif
+	}
 
+	static FProperty* CreateEngineType(FUnrealPropertyDefinitionInfo& PropDef, FFieldVariant Scope, const FName& Name, EObjectFlags ObjectFlags)
+	{
+		FPropertyBase& VarProperty = PropDef.GetPropertyBase();
+		check(VarProperty.ClassDef);
+		UClass* PropertyClass = VarProperty.ClassDef->GetClass();
+		check(PropertyClass->HasAnyClassFlags(CLASS_Interface));
 		FInterfaceProperty* Result = new  FInterfaceProperty(Scope, Name, ObjectFlags);
 		Result->InterfaceClass = PropertyClass;
 		return Result;
@@ -1439,7 +1596,11 @@ struct FPropertyTypeTraitsName : public FPropertyTypeTraitsBase
 		return FDefaultValueHelper::StringFromCppString(CppForm, TEXT("FName"), OutForm);
 	}
 
-	static FProperty* CreateProperty(FUnrealPropertyDefinitionInfo& PropDef, FFieldVariant Scope, const FName& Name, EObjectFlags ObjectFlags)
+	static void CreateProperty(FUnrealPropertyDefinitionInfo& PropDef)
+	{
+	}
+
+	static FProperty* CreateEngineType(FUnrealPropertyDefinitionInfo& PropDef, FFieldVariant Scope, const FName& Name, EObjectFlags ObjectFlags)
 	{
 		FNameProperty* Result = new FNameProperty(Scope, Name, ObjectFlags);
 		return Result;
@@ -1473,7 +1634,11 @@ struct FPropertyTypeTraitsString : public FPropertyTypeTraitsBase
 		return FDefaultValueHelper::StringFromCppString(CppForm, TEXT("FString"), OutForm);
 	}
 
-	static FProperty* CreateProperty(FUnrealPropertyDefinitionInfo& PropDef, FFieldVariant Scope, const FName& Name, EObjectFlags ObjectFlags)
+	static void CreateProperty(FUnrealPropertyDefinitionInfo& PropDef)
+	{
+	}
+
+	static FProperty* CreateEngineType(FUnrealPropertyDefinitionInfo& PropDef, FFieldVariant Scope, const FName& Name, EObjectFlags ObjectFlags)
 	{
 		FStrProperty* Result = new FStrProperty(Scope, Name, ObjectFlags);
 		return Result;
@@ -1544,7 +1709,11 @@ struct FPropertyTypeTraitsText : public FPropertyTypeTraitsBase
 		return true;
 	}
 
-	static FProperty* CreateProperty(FUnrealPropertyDefinitionInfo& PropDef, FFieldVariant Scope, const FName& Name, EObjectFlags ObjectFlags)
+	static void CreateProperty(FUnrealPropertyDefinitionInfo& PropDef)
+	{
+	}
+
+	static FProperty* CreateEngineType(FUnrealPropertyDefinitionInfo& PropDef, FFieldVariant Scope, const FName& Name, EObjectFlags ObjectFlags)
 	{
 		FTextProperty* Result = new FTextProperty(Scope, Name, ObjectFlags);
 		return Result;
@@ -1777,7 +1946,7 @@ struct FPropertyTypeTraitsStruct : public FPropertyTypeTraitsBase
 		return !OutForm.IsEmpty();
 	}
 
-	static FProperty* CreateProperty(FUnrealPropertyDefinitionInfo& PropDef, FFieldVariant Scope, const FName& Name, EObjectFlags ObjectFlags)
+	static void CreateProperty(FUnrealPropertyDefinitionInfo& PropDef)
 	{
 		FPropertyBase& VarProperty = PropDef.GetPropertyBase();
 		UScriptStruct* Struct = PropDef.GetPropertyBase().ScriptStructDef->GetScriptStruct();
@@ -1790,9 +1959,13 @@ struct FPropertyTypeTraitsStruct : public FPropertyTypeTraitsBase
 		{
 			VarProperty.PropertyFlags |= CPF_ContainsInstancedReference;
 		}
+	}
 
+	static FProperty* CreateEngineType(FUnrealPropertyDefinitionInfo& PropDef, FFieldVariant Scope, const FName& Name, EObjectFlags ObjectFlags)
+	{
+		FPropertyBase& VarProperty = PropDef.GetPropertyBase();
 		FStructProperty* Result = new FStructProperty(Scope, Name, ObjectFlags);
-		Result->Struct = Struct;
+		Result->Struct = PropDef.GetPropertyBase().ScriptStructDef->GetScriptStruct();
 		return Result;
 	}
 
@@ -1828,9 +2001,21 @@ struct FPropertyTypeTraitsStruct : public FPropertyTypeTraitsBase
 
 struct FPropertyTypeTraitsDelegate : public FPropertyTypeTraitsBase
 {
-	static FProperty* CreateProperty(FUnrealPropertyDefinitionInfo& PropDef, FFieldVariant Scope, const FName& Name, EObjectFlags ObjectFlags)
+	static void CreateProperty(FUnrealPropertyDefinitionInfo& PropDef)
 	{
+	}
+
+	static FProperty* CreateEngineType(FUnrealPropertyDefinitionInfo& PropDef, FFieldVariant Scope, const FName& Name, EObjectFlags ObjectFlags)
+	{
+		const FPropertyBase& VarProperty = PropDef.GetPropertyBase();
+
 		FDelegateProperty* Result = new FDelegateProperty(Scope, Name, ObjectFlags);
+
+		// Once we fully transition to post parse engine type creation, this should always be set.
+		if (VarProperty.FunctionDef != nullptr)
+		{
+			Result->SignatureFunction = VarProperty.FunctionDef->GetFunction();
+		}
 		return Result;
 	}
 
@@ -1893,7 +2078,11 @@ struct FPropertyTypeTraitsDelegate : public FPropertyTypeTraitsBase
 
 struct FPropertyTypeTraitsMulticastDelegate : public FPropertyTypeTraitsBase
 {
-	static FProperty* CreateProperty(FUnrealPropertyDefinitionInfo& PropDef, FFieldVariant Scope, const FName& Name, EObjectFlags ObjectFlags)
+	static void CreateProperty(FUnrealPropertyDefinitionInfo& PropDef)
+	{
+	}
+
+	static FProperty* CreateEngineType(FUnrealPropertyDefinitionInfo& PropDef, FFieldVariant Scope, const FName& Name, EObjectFlags ObjectFlags)
 	{
 		const FPropertyBase& VarProperty = PropDef.GetPropertyBase();
 		UFunction* Function = PropDef.GetPropertyBase().FunctionDef->GetFunction();
@@ -1985,11 +2174,17 @@ struct FPropertyTypeTraitsMulticastDelegate : public FPropertyTypeTraitsBase
 
 struct FPropertyTypeTraitsFieldPath : public FPropertyTypeTraitsBase
 {
-	static FProperty* CreateProperty(FUnrealPropertyDefinitionInfo& PropDef, FFieldVariant Scope, const FName& Name, EObjectFlags ObjectFlags)
+	static void CreateProperty(FUnrealPropertyDefinitionInfo& PropDef)
+	{
+	}
+
+	static FProperty* CreateEngineType(FUnrealPropertyDefinitionInfo& PropDef, FFieldVariant Scope, const FName& Name, EObjectFlags ObjectFlags)
 	{
 		const FPropertyBase& VarProperty = PropDef.GetPropertyBase();
 		FFieldPathProperty* Result = new FFieldPathProperty(Scope, Name, ObjectFlags);
-		Result->PropertyClass = VarProperty.PropertyPathClass;
+		FFieldClass** FieldClass = FFieldClass::GetNameToFieldClassMap().Find(VarProperty.FieldClassName);
+		check(FieldClass);
+		Result->PropertyClass = *FieldClass;
 		return Result;
 	}
 
@@ -2010,7 +2205,7 @@ struct FPropertyTypeTraitsFieldPath : public FPropertyTypeTraitsBase
 			const FPropertyBase& VarProperty = PropDef.GetPropertyBase();
 			FString& InnerTypeText = *ExtendedTypeText;
 			InnerTypeText = TEXT("<F");
-			InnerTypeText += VarProperty.PropertyPathClass->GetName();
+			InnerTypeText += VarProperty.FieldClassName.ToString();
 			InnerTypeText += TEXT(">");
 		}
 		return TEXT("TFieldPath");
@@ -2018,7 +2213,7 @@ struct FPropertyTypeTraitsFieldPath : public FPropertyTypeTraitsBase
 
 	static FString GetCPPTypeForwardDeclaration(const FUnrealPropertyDefinitionInfo& PropDef)
 	{
-		return FString::Printf(TEXT("class F%s;"), *PropDef.GetPropertyBase().PropertyPathClass->GetName());
+		return FString::Printf(TEXT("class F%s;"), *PropDef.GetPropertyBase().FieldClassName.ToString());
 	}
 };
 
@@ -2030,9 +2225,14 @@ struct FPropertyTypeTraitsFieldPath : public FPropertyTypeTraitsBase
 
 struct FPropertyTypeTraitsStaticArray : public FPropertyTypeTraitsBase
 {
-	static FProperty* CreateProperty(FUnrealPropertyDefinitionInfo& PropDef, FFieldVariant Scope, const FName& Name, EObjectFlags ObjectFlags)
+	static void CreateProperty(FUnrealPropertyDefinitionInfo& PropDef)
+	{
+		CreatePropertyHelper<false>(PropDef);
+	}
+
+	static FProperty* CreateEngineType(FUnrealPropertyDefinitionInfo& PropDef, FFieldVariant Scope, const FName& Name, EObjectFlags ObjectFlags)
 	{
-		FProperty* Property = CreatePropertyHelper<false>(PropDef, Scope, Name, ObjectFlags);
+		FProperty* Property = CreateEngineTypeHelper<false>(PropDef, Scope, Name, ObjectFlags);
 		Property->ArrayDim = 2;
 		return Property;
 	}
@@ -2060,31 +2260,35 @@ struct FPropertyTypeTraitsStaticArray : public FPropertyTypeTraitsBase
 
 struct FPropertyTypeTraitsDynamicArray : public FPropertyTypeTraitsBase
 {
-	static FProperty* CreateProperty(FUnrealPropertyDefinitionInfo& PropDef, FFieldVariant Scope, const FName& Name, EObjectFlags ObjectFlags)
+	static void CreateProperty(FUnrealPropertyDefinitionInfo& PropDef)
 	{
 		FPropertyBase& VarProperty = PropDef.GetPropertyBase();
 
-		FArrayProperty* Array = new FArrayProperty(Scope, Name, ObjectFlags);
-		PropDef.SetProperty(Array);
-
 		FPropertyBase InnerVarProperty = PropDef.GetPropertyBase();
 		InnerVarProperty.ArrayType = EArrayType::None;
 		InnerVarProperty.RepNotifyName = NAME_None;
 		InnerVarProperty.MetaData.Empty();
-		FUnrealPropertyDefinitionInfo& InnerPropDef = FPropertyTraits::CreateProperty(InnerVarProperty, PropDef, Name, RF_Public, 
-			PropDef.GetVariableCategory(), PropDef.GetArrayDimensions(), PropDef.GetUnrealSourceFile(), PropDef.GetLineNumber(), PropDef.GetParsePosition());
-		FProperty* InnerProp = InnerPropDef.GetProperty();
+		PropDef.SetValuePropDef(FPropertyTraits::CreateProperty(InnerVarProperty, PropDef, PropDef.GetFName(), 
+			PropDef.GetVariableCategory(), ACCESS_Public, PropDef.GetArrayDimensions(), PropDef.GetUnrealSourceFile(), PropDef.GetLineNumber(), PropDef.GetParsePosition()));
+		FUnrealPropertyDefinitionInfo& InnerPropDef = PropDef.GetValuePropDef();
 
-		Array->Inner = InnerProp;
-		VarProperty.PropertyFlags = Array->Inner->PropertyFlags;
+		VarProperty.PropertyFlags = InnerPropDef.GetPropertyBase().PropertyFlags;
 		VarProperty.MetaData.Append(MoveTemp(InnerPropDef.GetPropertyBase().MetaData)); // Move any added meta data items to the container
 		PropDef.SetAllocatorType(VarProperty.AllocatorType);
-		PropDef.SetValuePropDef(InnerPropDef);
 
 		// Propagate flags
-		InnerPropDef.GetPropertyBase().PropertyFlags = InnerProp->PropertyFlags = InnerProp->PropertyFlags & CPF_PropagateToArrayInner;
+		InnerPropDef.GetPropertyBase().PropertyFlags &= CPF_PropagateToArrayInner;
 
 		PropagateFlagsFromInnerAndHandlePersistentInstanceMetadata(VarProperty.PropertyFlags, PropDef.GetPropertyBase().MetaData, InnerPropDef);
+	}
+
+	static FProperty* CreateEngineType(FUnrealPropertyDefinitionInfo& PropDef, FFieldVariant Scope, const FName& Name, EObjectFlags ObjectFlags)
+	{
+		FPropertyBase& VarProperty = PropDef.GetPropertyBase();
+
+		FArrayProperty* Array = new FArrayProperty(Scope, Name, ObjectFlags);
+		PropDef.SetProperty(Array);
+		Array->Inner = FPropertyTraits::CreateEngineType(PropDef.GetValuePropDefRef());
 		return Array;
 	}
 
@@ -2127,30 +2331,34 @@ struct FPropertyTypeTraitsDynamicArray : public FPropertyTypeTraitsBase
 
 struct FPropertyTypeTraitsSet : public FPropertyTypeTraitsBase
 {
-	static FProperty* CreateProperty(FUnrealPropertyDefinitionInfo& PropDef, FFieldVariant Scope, const FName& Name, EObjectFlags ObjectFlags)
+	static void CreateProperty(FUnrealPropertyDefinitionInfo& PropDef)
 	{
 		FPropertyBase& VarProperty = PropDef.GetPropertyBase();
 
-		FSetProperty* Set = new FSetProperty(Scope, Name, ObjectFlags);
-		PropDef.SetProperty(Set);
-
 		FPropertyBase InnerVarProperty = PropDef.GetPropertyBase();
 		InnerVarProperty.ArrayType = EArrayType::None;
 		InnerVarProperty.RepNotifyName = NAME_None;
 		InnerVarProperty.MetaData.Empty();
-		FUnrealPropertyDefinitionInfo& InnerPropDef = FPropertyTraits::CreateProperty(InnerVarProperty, PropDef, Name, RF_Public, 
-			PropDef.GetVariableCategory(), PropDef.GetArrayDimensions(), PropDef.GetUnrealSourceFile(), PropDef.GetLineNumber(), PropDef.GetParsePosition());
-		FProperty* InnerProp = InnerPropDef.GetProperty();
+		PropDef.SetValuePropDef(FPropertyTraits::CreateProperty(InnerVarProperty, PropDef, PropDef.GetFName(),
+			PropDef.GetVariableCategory(), ACCESS_Public, PropDef.GetArrayDimensions(), PropDef.GetUnrealSourceFile(), PropDef.GetLineNumber(), PropDef.GetParsePosition()));
+		FUnrealPropertyDefinitionInfo& InnerPropDef = PropDef.GetValuePropDef();
 
-		Set->ElementProp = InnerProp;
-		VarProperty.PropertyFlags = InnerProp->PropertyFlags;
+		VarProperty.PropertyFlags = InnerPropDef.GetPropertyBase().PropertyFlags;
 		VarProperty.MetaData.Append(MoveTemp(InnerPropDef.GetPropertyBase().MetaData)); // Move any added meta data items to the container
-		PropDef.SetValuePropDef(InnerPropDef);
 
 		// Propagate flags
-		InnerPropDef.GetPropertyBase().PropertyFlags = InnerProp->PropertyFlags = InnerProp->PropertyFlags & CPF_PropagateToSetElement;
+		InnerPropDef.GetPropertyBase().PropertyFlags &= CPF_PropagateToSetElement;
 
 		PropagateFlagsFromInnerAndHandlePersistentInstanceMetadata(VarProperty.PropertyFlags, PropDef.GetPropertyBase().MetaData, InnerPropDef);
+	}
+
+	static FProperty* CreateEngineType(FUnrealPropertyDefinitionInfo& PropDef, FFieldVariant Scope, const FName& Name, EObjectFlags ObjectFlags)
+	{
+		FPropertyBase& VarProperty = PropDef.GetPropertyBase();
+
+		FSetProperty* Set = new FSetProperty(Scope, Name, ObjectFlags);
+		PropDef.SetProperty(Set);
+		Set->ElementProp = FPropertyTraits::CreateEngineType(PropDef.GetValuePropDefRef());
 		return Set;
 	}
 
@@ -2191,40 +2399,47 @@ struct FPropertyTypeTraitsSet : public FPropertyTypeTraitsBase
 
 struct FPropertyTypeTraitsMap : public FPropertyTypeTraitsBase
 {
-	static FProperty* CreateProperty(FUnrealPropertyDefinitionInfo& PropDef, FFieldVariant Scope, const FName& Name, EObjectFlags ObjectFlags)
+	static void CreateProperty(FUnrealPropertyDefinitionInfo& PropDef)
 	{
 		FPropertyBase& VarProperty = PropDef.GetPropertyBase();
 
-		FMapProperty* Map = new FMapProperty(Scope, Name, ObjectFlags);
-		PropDef.SetProperty(Map);
-
-		FUnrealPropertyDefinitionInfo& KeyPropDef = FPropertyTraits::CreateProperty(*PropDef.GetPropertyBase().MapKeyProp, PropDef, *(Name.ToString() + TEXT("_Key")), RF_Public,
-			PropDef.GetVariableCategory(), PropDef.GetArrayDimensions(), PropDef.GetUnrealSourceFile(), PropDef.GetLineNumber(), PropDef.GetParsePosition());
-		FProperty* KeyProp = KeyPropDef.GetProperty();
+		// The key shouldn't ever have meta data, so eliminate was came in as part of parsing and the remove anything that was added
+		// as part of property creation.  
+		PropDef.GetPropertyBase().MapKeyProp->MetaData.Empty();
+		PropDef.SetKeyPropDef(FPropertyTraits::CreateProperty(*PropDef.GetPropertyBase().MapKeyProp, PropDef, *(PropDef.GetFName().ToString() + TEXT("_Key")), 
+			PropDef.GetVariableCategory(), ACCESS_Public, PropDef.GetArrayDimensions(), PropDef.GetUnrealSourceFile(), PropDef.GetLineNumber(), PropDef.GetParsePosition()));
+		FUnrealPropertyDefinitionInfo& KeyPropDef = PropDef.GetKeyPropDef();
+		KeyPropDef.GetPropertyBase().MetaData.Empty();
 
 		FPropertyBase ValueVarProperty = PropDef.GetPropertyBase();
 		ValueVarProperty.ArrayType = EArrayType::None;
 		ValueVarProperty.MapKeyProp = nullptr;
 		ValueVarProperty.RepNotifyName = NAME_None;
 		ValueVarProperty.MetaData.Empty();
-		FUnrealPropertyDefinitionInfo& ValuePropDef = FPropertyTraits::CreateProperty(ValueVarProperty, PropDef, Name, RF_Public,
-			PropDef.GetVariableCategory(), PropDef.GetArrayDimensions(), PropDef.GetUnrealSourceFile(), PropDef.GetLineNumber(), PropDef.GetParsePosition());
-		FProperty* ValueProp = ValuePropDef.GetProperty();
+		PropDef.SetValuePropDef(FPropertyTraits::CreateProperty(ValueVarProperty, PropDef, PropDef.GetFName(), 
+			PropDef.GetVariableCategory(), ACCESS_Public, PropDef.GetArrayDimensions(), PropDef.GetUnrealSourceFile(), PropDef.GetLineNumber(), PropDef.GetParsePosition()));
+		FUnrealPropertyDefinitionInfo& ValuePropDef = PropDef.GetValuePropDef();
 
-		Map->KeyProp = KeyProp;
-		Map->ValueProp = ValueProp;
-		VarProperty.PropertyFlags = ValueProp->PropertyFlags;
+		VarProperty.PropertyFlags = ValuePropDef.GetPropertyBase().PropertyFlags;
 		VarProperty.MetaData.Append(MoveTemp(ValuePropDef.GetPropertyBase().MetaData)); // Move any added meta data items to the container
-		PropDef.SetAllocatorType(VarProperty.AllocatorType);
-		PropDef.SetKeyPropDef(KeyPropDef);
-		PropDef.SetValuePropDef(ValuePropDef);
+		PropDef.SetAllocatorType(VarProperty.AllocatorType);//@TODO - This is a duplicate now
 
 		// Propagate flags
-		KeyPropDef.GetPropertyBase().PropertyFlags = KeyProp->PropertyFlags = KeyProp->PropertyFlags & CPF_PropagateToMapKey;
-		ValuePropDef.GetPropertyBase().PropertyFlags = ValueProp->PropertyFlags = ValueProp->PropertyFlags & CPF_PropagateToMapValue;
+		KeyPropDef.GetPropertyBase().PropertyFlags &= CPF_PropagateToMapKey;
+		ValuePropDef.GetPropertyBase().PropertyFlags &= CPF_PropagateToMapValue;
 
 		PropagateFlagsFromInnerAndHandlePersistentInstanceMetadata(VarProperty.PropertyFlags, PropDef.GetPropertyBase().MapKeyProp->MetaData, KeyPropDef);
 		PropagateFlagsFromInnerAndHandlePersistentInstanceMetadata(VarProperty.PropertyFlags, PropDef.GetPropertyBase().MetaData, ValuePropDef);
+	}
+
+	static FProperty* CreateEngineType(FUnrealPropertyDefinitionInfo& PropDef, FFieldVariant Scope, const FName& Name, EObjectFlags ObjectFlags)
+	{
+		FPropertyBase& VarProperty = PropDef.GetPropertyBase();
+
+		FMapProperty* Map = new FMapProperty(Scope, Name, ObjectFlags);
+		PropDef.SetProperty(Map);
+		Map->KeyProp = FPropertyTraits::CreateEngineType(PropDef.GetKeyPropDefRef());
+		Map->ValueProp = FPropertyTraits::CreateEngineType(PropDef.GetValuePropDefRef());
 		return Map;
 	}
 
@@ -2292,22 +2507,47 @@ bool FPropertyTraits::DefaultValueStringCppFormatToInnerFormat(const FUnrealProp
 	return PropertyTypeDispatch<DefaultValueStringCppFormatToInnerFormatDispatch, false, bool>(PropDef.GetPropertyBase(), std::ref(PropDef), std::ref(CppForm), std::ref(OutForm));
 }
 
-FUnrealPropertyDefinitionInfo& FPropertyTraits::CreateProperty(const FPropertyBase& VarProperty, FUnrealTypeDefinitionInfo& Outer, const FName& Name, EObjectFlags ObjectFlags, EVariableCategory VariableCategory,
-	const TCHAR* Dimensions, FUnrealSourceFile& SourceFile, int LineNumber, int ParsePosition)
+TSharedRef<FUnrealPropertyDefinitionInfo> FPropertyTraits::CreateProperty(const FPropertyBase& VarProperty, FUnrealTypeDefinitionInfo& Outer, const FName& Name, EVariableCategory VariableCategory, EAccessSpecifier AccessSpecifier, const TCHAR* Dimensions, FUnrealSourceFile& SourceFile, int LineNumber, int ParsePosition)
 {
-	TSharedRef<FUnrealPropertyDefinitionInfo> PropDefRef = MakeShared<FUnrealPropertyDefinitionInfo>(SourceFile, LineNumber, ParsePosition, VarProperty, VariableCategory, Dimensions, Name, Outer);
+	TSharedRef<FUnrealPropertyDefinitionInfo> PropDefRef = MakeShared<FUnrealPropertyDefinitionInfo>(SourceFile, LineNumber, ParsePosition, VarProperty, VariableCategory, AccessSpecifier, Dimensions, Name, Outer);
 	FUnrealPropertyDefinitionInfo& PropDef = *PropDefRef;
 
 	// Create the property and attach to the definition
-	FFieldVariant Scope = Outer.AsProperty() ? FFieldVariant(Outer.AsProperty()->GetProperty()) : FFieldVariant(Outer.AsObject()->GetObject());
-	FProperty* Property = CreatePropertyHelper<true>(PropDef, Scope, Name, ObjectFlags);
+	CreatePropertyHelper<true>(PropDef);
+
+	// If we aren't a sub-property, then add any meta data to the property
+	// We do this because of the edit inline metadata which might be added.
+	// Technically, it should be added to the outermost property.
+	if (UHTCast<FUnrealPropertyDefinitionInfo>(Outer) == nullptr)
+	{
+		FUHTMetaData::RemapMetaData(PropDef.GetPropertyBase().MetaData);
+		PropDef.ValidateMetaDataFormat(PropDef.GetPropertyBase().MetaData);
+	}
+	return PropDefRef;
+}
+
+FProperty* FPropertyTraits::CreateEngineType(TSharedRef<FUnrealPropertyDefinitionInfo> PropDefRef)
+{
+	FUnrealPropertyDefinitionInfo& PropDef = *PropDefRef;
+	FUnrealTypeDefinitionInfo* Outer = PropDef.GetOuter();
+	check(Outer);
+
+	EObjectFlags ObjectFlags = RF_Public;
+	if (PropDef.GetVariableCategory() == EVariableCategory::Member && PropDef.GetAccessSpecifier() == ACCESS_Private)
+	{
+		ObjectFlags = RF_NoFlags;
+	}
+
+	// Create the property and attach to the definition
+	FFieldVariant Scope = Outer->AsProperty() ? FFieldVariant(Outer->AsProperty()->GetProperty()) : FFieldVariant(Outer->AsObject()->GetObject());
+	FProperty* Property = CreateEngineTypeHelper<true>(PropDef, Scope, PropDef.GetFName(), ObjectFlags);
 	PropDef.SetProperty(Property);
 
 	// Perform some final initialization from the property base data
 	Property->PropertyFlags = PropDef.GetPropertyBase().PropertyFlags;
 
 	// Special initialization for member variables
-	if (VariableCategory == EVariableCategory::Member)
+	if (PropDef.GetVariableCategory() == EVariableCategory::Member)
 	{
 		if (PropDef.HasAnyPropertyFlags(CPF_RepNotify))
 		{
@@ -2315,24 +2555,33 @@ FUnrealPropertyDefinitionInfo& FPropertyTraits::CreateProperty(const FPropertyBa
 		}
 	}
 
-	// If we aren't a sub-property, then add any meta data to the property
-	// We do this because of the edit inline metadata which might be added.
-	// Technically, it should be added to the outermost property.
-	if (UHTCast<FUnrealPropertyDefinitionInfo>(Outer) == nullptr)
+	// Add the meta data to the property
+	if (TMap<FName, FString>* MetaDataMap = PropDef.GetMetaDataMap())
 	{
-		FUHTMetaData::RemapAndAddMetaData(PropDef, PropDef.GetPropertyBase().MetaData);
+		for (TPair<FName, FString>& KVP : *MetaDataMap)
+		{
+			Property->SetMetaData(KVP.Key, FString(KVP.Value));
+		}
 	}
 
-	//FString a = GetCPPTypeForwardDeclaration(PropDef);
-	//FString b = PropDef.GetProperty()->GetCPPTypeForwardDeclaration();
-	//if (a != b)
-	//{
-	//	GetCPPType(PropDef, nullptr, 0);
-	//}
+	// If we are parent to a struct
+	if (FUnrealStructDefinitionInfo* StructDef = UHTCast<FUnrealStructDefinitionInfo>(Outer))
+	{
+
+		// Add to the end of the properties slist
+		FField** Prev = &StructDef->GetStruct()->ChildProperties;
+		for (; *Prev != nullptr; Prev = &(*Prev)->Next)
+		{
+			// No body
+		}
+		check(*Prev == nullptr);
+		Property->Next = nullptr;
+		*Prev = Property;
+	}
 
 	// Add the property to the lookup map
-	GTypeDefinitionInfoMap.Add(PropDef.GetProperty(), MoveTemp(PropDefRef));
-	return PropDef;
+	GTypeDefinitionInfoMap.Add(PropDef.GetProperty(), PropDefRef);
+	return Property;
 }
 
 bool FPropertyTraits::IsSupportedByBlueprint(const FUnrealPropertyDefinitionInfo& PropDef, bool bMemberVariable)
@@ -2354,3 +2603,67 @@ FString FPropertyTraits::GetCPPTypeForwardDeclaration(const FUnrealPropertyDefin
 {
 	return PropertyTypeDispatch<GetCPPTypeForwardDeclarationDispatch, true, FString>(PropDef.GetPropertyBase(), std::ref(PropDef));
 }
+
+bool FPropertyTraits::SameType(const FUnrealPropertyDefinitionInfo& Lhs, const FUnrealPropertyDefinitionInfo& Rhs)
+{
+	auto AdjustTypes = [](EUHTPropertyType PropertyType)
+	{
+		switch (PropertyType)
+		{
+		case EUHTPropertyType::ObjectPtrReference:
+			return EUHTPropertyType::ObjectReference;
+		case EUHTPropertyType::Bool8:
+		case EUHTPropertyType::Bool16:
+		case EUHTPropertyType::Bool32:
+		case EUHTPropertyType::Bool64:
+			return EUHTPropertyType::Bool;
+		default:
+			return PropertyType;
+		}
+	};
+
+	const FPropertyBase& LhsPropertyBase = Lhs.GetPropertyBase();
+	const FPropertyBase& RhsPropertyBase = Rhs.GetPropertyBase();
+
+	EUHTPropertyType LhsType = AdjustTypes(LhsPropertyBase.GetUHTPropertyType());
+	EUHTPropertyType RhsType = AdjustTypes(RhsPropertyBase.GetUHTPropertyType());
+	if (LhsType != RhsType)
+	{
+		return  false;
+	}
+
+	if (Lhs.HasKeyPropDef() && !SameType(Lhs.GetKeyPropDef(), Rhs.GetKeyPropDef()))
+	{
+		return false;
+	}
+
+
+	if (Lhs.HasValuePropDef() && !SameType(Lhs.GetValuePropDef(), Rhs.GetValuePropDef()))
+	{
+		return false;
+	}
+
+	if (LhsPropertyBase.TypeDef != RhsPropertyBase.TypeDef)
+	{
+		return false;
+	}
+
+	if (LhsPropertyBase.FieldClassName != RhsPropertyBase.FieldClassName)
+	{
+		return false;
+	}
+
+	if (LhsPropertyBase.MetaClassDef != RhsPropertyBase.MetaClassDef)
+	{
+		return false;
+	}
+
+	return true;
+}
+
+bool FPropertyTraits::IsValidFieldClass(FName FieldClassName)
+{
+	static FFieldClassInit FieldClassInit;
+
+	return FFieldClass::GetNameToFieldClassMap().Find(FieldClassName) != nullptr;
+}
diff --git a/Engine/Source/Programs/UnrealHeaderTool/Private/PropertyTypes.h b/Engine/Source/Programs/UnrealHeaderTool/Private/PropertyTypes.h
--- a/Engine/Source/Programs/UnrealHeaderTool/Private/PropertyTypes.h
+++ b/Engine/Source/Programs/UnrealHeaderTool/Private/PropertyTypes.h
@@ -32,13 +32,22 @@ struct FPropertyTraits
 	 * @param Name The name of the property
 	 * @param ObjectFlags The flags associated with the property
 	 * @param VariableCategory The parsing context of the property
+	 * @param AccessCategory The access of the property (i.e. public, private, etc...)
 	 * @param Dimensions When this is a static array, this represents the dimensions value
 	 * @param SourceFile The source file containing the property
 	 * @param LineNumber Line number of the property
 	 * @param ParsePosition Character position of the property in the header
 	 * @return The pointer to the newly created property.  It will be attached to the definition by the caller
 	 */
-	static FUnrealPropertyDefinitionInfo& CreateProperty(const FPropertyBase& VarProperty, FUnrealTypeDefinitionInfo& Outer, const FName& Name, EObjectFlags ObjectFlags, EVariableCategory VariableCategory, const TCHAR* Dimensions, FUnrealSourceFile& SourceFile, int LineNumber, int ParsePosition);
+	static TSharedRef<FUnrealPropertyDefinitionInfo> CreateProperty(const FPropertyBase& VarProperty, FUnrealTypeDefinitionInfo& Outer, const FName& Name, EVariableCategory VariableCategory, EAccessSpecifier AccessSpecifier, const TCHAR* Dimensions, FUnrealSourceFile& SourceFile, int LineNumber, int ParsePosition);
+
+	/**
+	 * Given a property, create the underlying engine types
+	 * 
+	 * @param PropDef The property in question
+	 * @param ObjectFlags The flags associated with the property
+	 */
+	static FProperty* CreateEngineType(TSharedRef<FUnrealPropertyDefinitionInfo> PropDefRef);
 
 	/**
 	 * Test to see if the property can be used in a blueprint
@@ -65,4 +74,14 @@ struct FPropertyTraits
 	static FString GetCPPType(const FUnrealPropertyDefinitionInfo& PropDef, FString* ExtendedTypeText = nullptr, uint32 CPPExportFlags = 0);
 
 	static FString GetCPPTypeForwardDeclaration(const FUnrealPropertyDefinitionInfo& PropDef);
+
+	/**
+	 * Tests to see if the two types are the same types
+	 */
+	static bool SameType(const FUnrealPropertyDefinitionInfo& Lhs, const FUnrealPropertyDefinitionInfo& Rhs);
+
+	/**
+	 * Tests to see if the given field class name is valid
+	 */
+	static bool IsValidFieldClass(FName FieldClassName);
 };
diff --git a/Engine/Source/Programs/UnrealHeaderTool/Private/UnrealSourceFile.cpp b/Engine/Source/Programs/UnrealHeaderTool/Private/UnrealSourceFile.cpp
--- a/Engine/Source/Programs/UnrealHeaderTool/Private/UnrealSourceFile.cpp
+++ b/Engine/Source/Programs/UnrealHeaderTool/Private/UnrealSourceFile.cpp
@@ -11,19 +11,28 @@
 
 void FUnrealSourceFile::AddDefinedClass(TSharedRef<FUnrealTypeDefinitionInfo> ClassDecl)
 {
+	DefinedTypes.Add(ClassDecl);
 	DefinedClasses.Add(MoveTemp(ClassDecl));
 }
 
 void FUnrealSourceFile::AddDefinedEnum(TSharedRef<FUnrealTypeDefinitionInfo> EnumDecl)
 {
+	DefinedTypes.Add(EnumDecl);
 	DefinedEnums.Add(MoveTemp(EnumDecl));
 }
 
 void FUnrealSourceFile::AddDefinedStruct(TSharedRef<FUnrealTypeDefinitionInfo> StructDecl)
 {
+	DefinedTypes.Add(StructDecl);
 	DefinedStructs.Add(MoveTemp(StructDecl));
 }
 
+void FUnrealSourceFile::AddDefinedFunction(TSharedRef<FUnrealTypeDefinitionInfo> FunctionDef)
+{
+	DefinedTypes.Add(FunctionDef);
+	DefinedFunctions.Add(MoveTemp(FunctionDef));
+}
+
 const FString& FUnrealSourceFile::GetFileId() const
 {
 	if (FileId.Len() == 0)
diff --git a/Engine/Source/Programs/UnrealHeaderTool/Private/UnrealSourceFile.h b/Engine/Source/Programs/UnrealHeaderTool/Private/UnrealSourceFile.h
--- a/Engine/Source/Programs/UnrealHeaderTool/Private/UnrealSourceFile.h
+++ b/Engine/Source/Programs/UnrealHeaderTool/Private/UnrealSourceFile.h
@@ -137,6 +137,59 @@ public:
 		return DefinedStructs.Num();
 	}
 
+	/**
+	 * Adds given function to function definition list for this source file.
+	 *
+	 * @param FunctionDef Declaration information about the struct
+	 */
+	void AddDefinedFunction(TSharedRef<FUnrealTypeDefinitionInfo> FunctionDef);
+
+	/**
+	 * Gets array with functions defined in this source file with parsing info.
+	 *
+	 * @returns Array with functions defined in this source file with parsing info.
+	 */
+	TArray<TSharedRef<FUnrealTypeDefinitionInfo>>& GetDefinedFunctions()
+	{
+		return DefinedFunctions;
+	}
+
+	const TArray<TSharedRef<FUnrealTypeDefinitionInfo>>& GetDefinedFunctions() const
+	{
+		return DefinedFunctions;
+	}
+
+	/**
+	 * Gets number of types defined in this source file.
+	 */
+	int32 GetDefinedFunctionsCount() const
+	{
+		return DefinedFunctions.Num();
+	}
+
+	/**
+	 * Gets array with types defined in this source file with parsing info.
+	 *
+	 * @returns Array with types defined in this source file with parsing info.
+	 */
+	TArray<TSharedRef<FUnrealTypeDefinitionInfo>>& GetDefinedTypes()
+	{
+		return DefinedTypes;
+	}
+
+	const TArray<TSharedRef<FUnrealTypeDefinitionInfo>>& GetDefinedTypes() const
+	{
+		return DefinedTypes;
+	}
+
+	/**
+	 * Gets number of types defined in this source file.
+	 */
+	int32 GetDefinedTypeCount() const
+	{
+		return DefinedTypes.Num();
+	}
+
 	/**
 	 * Gets generated header filename.
 	 */
@@ -281,11 +334,6 @@ public:
 	 */
 	void AddTypeDefIncludeIfNeeded(FUnrealTypeDefinitionInfo* TypeDef);
 
-	/**
-	 * Add an include for a type definition if required
-	 */
-	void AddTypeDefIncludeIfNeeded(UField* Field);
-
 	/**
 	 * Sets generated filename.
 	 */
@@ -496,6 +544,12 @@ private:
 	// List of structs defined in this source file along with parsing info.
 	TArray<TSharedRef<FUnrealTypeDefinitionInfo>> DefinedStructs;
 
+	// List of functions defined in this source file along with parsing info.
+	TArray<TSharedRef<FUnrealTypeDefinitionInfo>> DefinedFunctions;
+
+	// List of all defined types
+	TArray<TSharedRef<FUnrealTypeDefinitionInfo>> DefinedTypes;
+
 	// Collection of all singletons found during code generation
 	TArray<FUnrealFieldDefinitionInfo*> Singletons;
 };
diff --git a/Engine/Source/Programs/UnrealHeaderTool/Private/UnrealTypeDefinitionInfo.cpp b/Engine/Source/Programs/UnrealHeaderTool/Private/UnrealTypeDefinitionInfo.cpp
--- a/Engine/Source/Programs/UnrealHeaderTool/Private/UnrealTypeDefinitionInfo.cpp
+++ b/Engine/Source/Programs/UnrealHeaderTool/Private/UnrealTypeDefinitionInfo.cpp
@@ -338,21 +338,21 @@ void FUnrealTypeDefinitionInfo::ValidateMetaDataFormat(const FName InKey, ECheck
 				FUnrealPropertyDefinitionInfo* FirstParam = nullptr;
 				FUnrealPropertyDefinitionInfo* SecondParam = nullptr;
 				FUnrealPropertyDefinitionInfo* ReturnValue = nullptr;
-				for (FUnrealPropertyDefinitionInfo* PropertyDef : FuncDef->GetProperties())
+				for (TSharedRef<FUnrealPropertyDefinitionInfo> PropertyDef : FuncDef->GetProperties())
 				{
 					if (PropertyDef->HasAnyPropertyFlags(CPF_ReturnParm))
 					{
-						ReturnValue = PropertyDef;
+						ReturnValue = &*PropertyDef;
 					}
 					else
 					{
 						if (FirstParam == nullptr)
 						{
-							FirstParam = PropertyDef;
+							FirstParam = &*PropertyDef;
 						}
 						else if (SecondParam == nullptr)
 						{
-							SecondParam = PropertyDef;
+							SecondParam = &*PropertyDef;
 						}
 					}
 				}
@@ -461,27 +461,23 @@ void FUnrealTypeDefinitionInfo::ValidateMetaDataFormat(const FName InKey, ECheck
 	}
 }
 
-void FUnrealPropertyDefinitionInfo::PostParseFinalize()
+void FUnrealPropertyDefinitionInfo::PostParseFinalizeInternal(bool bCreateEngineTypes)
 {
-	TypePackageName = GetTypePackageNameHelper(*this);
-}
-
-void FUnrealPropertyDefinitionInfo::AddMetaData(TMap<FName, FString>&& InMetaData)
-{
-	// only add if we have some!
-	if (InMetaData.Num())
+	// Finalize the sub property defs, but don't create the engine type
+	if (KeyPropDef.IsValid())
 	{
-		check(Property);
-
-		UPackage* Package = Property->GetOutermost();
-		// get (or create) a metadata object for this package
-		UMetaData* MetaData = Package->GetMetaData();
+		KeyPropDef->PostParseFinalize(false);
+	}
+	if (ValuePropDef.IsValid())
+	{
+		ValuePropDef->PostParseFinalize(false);
+	}
 
-		for (TPair<FName, FString>& MetaKeyValue : InMetaData)
-		{
-			Property->SetMetaData(MetaKeyValue.Key, MoveTemp(MetaKeyValue.Value));
-		}
+	if (bCreateEngineTypes && GetPropertySafe() == nullptr)
+	{
+		FPropertyTraits::CreateEngineType(SharedThis(this));
 	}
+	TypePackageName = GetTypePackageNameHelper(*this);
 }
 
 bool FUnrealPropertyDefinitionInfo::IsDynamic() const
@@ -507,27 +503,33 @@ bool FUnrealPropertyDefinitionInfo::IsOwnedByDynamicType() const
 
 void FUnrealPropertyDefinitionInfo::SetDelegateFunctionSignature(FUnrealFunctionDefinitionInfo& DelegateFunctionDef)
 {
-	FDelegateProperty* DelegateProperty = CastFieldChecked<FDelegateProperty>(PropertyBase.ArrayType == EArrayType::None ? GetProperty() : GetValuePropDef().GetProperty());
-	DelegateProperty->SignatureFunction = DelegateFunctionDef.GetFunction();
+	if (GetPropertySafe() != nullptr)
+	{
+		FDelegateProperty* DelegateProperty = CastFieldChecked<FDelegateProperty>(PropertyBase.ArrayType == EArrayType::None ? GetProperty() : GetValuePropDef().GetProperty());
+		DelegateProperty->SignatureFunction = DelegateFunctionDef.GetFunction();
+	}
 	PropertyBase.FunctionDef = &DelegateFunctionDef;
 }
 
 FString FUnrealPropertyDefinitionInfo::GetEngineClassName() const
 {
 #if UHT_ENABLE_ENGINE_TYPE_CHECKS
-	check(FPropertyTraits::GetEngineClassName(*this) == GetProperty()->GetClass()->GetName()); // Validation
+	if (GetPropertySafe() != nullptr)
+	{
+		check(FPropertyTraits::GetEngineClassName(*this) == GetPropertySafe()->GetClass()->GetName()); // Validation
+	}
 #endif
 	return FPropertyTraits::GetEngineClassName(*this);
 }
 
-FString FUnrealPropertyDefinitionInfo::GetPathName() const
+FString FUnrealPropertyDefinitionInfo::GetPathName(FUnrealObjectDefinitionInfo* StopOuter) const
 {
 	TStringBuilder<256> ResultString;
-	GetPathName(ResultString);
+	GetPathName(StopOuter, ResultString);
 	return FString(FStringView(ResultString));
 }
 
-void FUnrealPropertyDefinitionInfo::GetPathName(FStringBuilderBase& ResultString) const
+void FUnrealPropertyDefinitionInfo::GetPathName(FUnrealObjectDefinitionInfo* StopOuter, FStringBuilderBase& ResultString) const
 {
 	TArray<FName, TInlineAllocator<16>> ParentFields;
 	for (FUnrealTypeDefinitionInfo* LocalOuter = GetOuter(); LocalOuter; LocalOuter = LocalOuter->GetOuter())
@@ -538,7 +540,7 @@ void FUnrealPropertyDefinitionInfo::GetPathName(FStringBuilderBase& ResultString
 		}
 		else
 		{
-			LocalOuter->GetPathName(ResultString);
+			LocalOuter->GetPathName(StopOuter, ResultString);
 			ResultString << SUBOBJECT_DELIMITER_CHAR;
 			break;
 		}
@@ -551,9 +553,12 @@ void FUnrealPropertyDefinitionInfo::GetPathName(FStringBuilderBase& ResultString
 	}
 	GetFName().AppendString(ResultString);
 #if UHT_ENABLE_ENGINE_TYPE_CHECKS
-	TStringBuilder<256> OtherResultString;
-	GetProperty()->GetPathName(nullptr, OtherResultString); // Validation
-	check(FCString::Strcmp(OtherResultString.ToString(), ResultString.ToString()) == 0);
+	if (GetPropertySafe() != nullptr)
+	{
+		TStringBuilder<256> OtherResultString;
+		GetPropertySafe()->GetPathName(StopOuter ? StopOuter->GetObject() : nullptr, OtherResultString); // Validation
+		check(FCString::Strcmp(OtherResultString.ToString(), ResultString.ToString()) == 0);
+	}
 #endif
 }
 
@@ -563,7 +568,10 @@ FString FUnrealPropertyDefinitionInfo::GetFullName() const
 	FullName += TEXT(" ");
 	FullName += GetPathName();
 #if UHT_ENABLE_ENGINE_TYPE_CHECKS
-	check(FullName == GetProperty()->GetFullName()); // Validation
+	if (GetPropertySafe() != nullptr)
+	{
+		check(FullName == GetPropertySafe()->GetFullName()); // Validation
+	}
 #endif
 	return FullName;
 }
@@ -572,9 +580,12 @@ FString FUnrealPropertyDefinitionInfo::GetCPPType(FString* ExtendedTypeText/* =
 {
 	FString Out = FPropertyTraits::GetCPPType(*this, ExtendedTypeText, CPPExportFlags);
 #if UHT_ENABLE_ENGINE_TYPE_CHECKS
-	FString ExtOutTemp;
-	FString* ExtOutTempPtr = ExtendedTypeText ? &ExtOutTemp : nullptr;
-	check(Out == GetProperty()->GetCPPType(ExtOutTempPtr, CPPExportFlags) && (ExtendedTypeText == nullptr || *ExtendedTypeText == ExtOutTemp)); // Validation
+	if (GetPropertySafe() != nullptr)
+	{
+		FString ExtOutTemp;
+		FString* ExtOutTempPtr = ExtendedTypeText ? &ExtOutTemp : nullptr;
+		check(Out == GetPropertySafe()->GetCPPType(ExtOutTempPtr, CPPExportFlags) && (ExtendedTypeText == nullptr || *ExtendedTypeText == ExtOutTemp)); // Validation
+	}
 #endif
 	return Out;
 }
@@ -583,11 +594,166 @@ FString FUnrealPropertyDefinitionInfo::GetCPPTypeForwardDeclaration() const
 {
 	FString Out = FPropertyTraits::GetCPPTypeForwardDeclaration(*this);
 #if UHT_ENABLE_ENGINE_TYPE_CHECKS
-	check(Out == GetProperty()->GetCPPTypeForwardDeclaration()); // Validation
+	if (GetPropertySafe() != nullptr)
+	{
+		check(Out == GetPropertySafe()->GetCPPTypeForwardDeclaration()); // Validation
+	}
 #endif
 	return Out;
 }
 
+struct FUHTFieldDisplayNameHelper
+{
+	static FUnrealObjectDefinitionInfo* GetOwnerObject(const FUnrealPropertyDefinitionInfo& Property)
+	{
+		for (FUnrealTypeDefinitionInfo* TypeDef = Property.GetOuter(); TypeDef; TypeDef = TypeDef->GetOuter())
+		{
+			if (FUnrealObjectDefinitionInfo* ObjectDef = UHTCast<FUnrealObjectDefinitionInfo>(TypeDef))
+			{
+				return ObjectDef;
+			}
+		}
+		return nullptr;
+	}
+
+	static FUnrealStructDefinitionInfo* GetOwnerStruct(const FUnrealPropertyDefinitionInfo& Property)
+	{
+		for (FUnrealTypeDefinitionInfo* TypeDef = Property.GetOuter(); TypeDef; TypeDef = TypeDef->GetOuter())
+		{
+			if (FUnrealStructDefinitionInfo* StructDef = UHTCast<FUnrealStructDefinitionInfo>(TypeDef))
+			{
+				return StructDef;
+			}
+		}
+		return nullptr;
+	}
+
+	static FString GetFullGroupName(const FUnrealPropertyDefinitionInfo& Property, bool bStartWithOuter)
+	{
+		if (bStartWithOuter)
+		{
+			if (FUnrealTypeDefinitionInfo* Owner = Property.GetOuter())
+			{
+				if (FUnrealObjectDefinitionInfo* ObjectOwner = UHTCast<FUnrealObjectDefinitionInfo>(Owner))
+				{
+					return ObjectOwner->GetPathName(&ObjectOwner->GetPackageDef());
+				}
+				else
+				{
+					FUnrealPropertyDefinitionInfo& PropertyOwner = UHTCastChecked<FUnrealPropertyDefinitionInfo>(Owner);
+					return PropertyOwner.GetPathName(&GetOwnerObject(PropertyOwner)->GetPackageDef());
+				}
+			}
+			else
+			{
+				return FString();
+			}
+		}
+		else
+		{
+			FUnrealObjectDefinitionInfo* ObjectOuter = GetOwnerObject(Property);
+			return Property.GetPathName(ObjectOuter ? &ObjectOuter->GetPackageDef() : nullptr);
+		}
+	}
+
+	static FString Get(const FUnrealPropertyDefinitionInfo& Property)
+	{
+		// The GetAuthoredNameForField only does something for user defined structures
+		//if (FUnrealStructDefinitionInfo* OwnerStruct = GetOwnerStruct(Property))
+		//{
+		//	OwnerStruct->GetStruct()->GetAuthoredNameForField(Property.GetProperty());
+		//}
+		return Property.GetName();
+	}
+};
+
+FText FUnrealPropertyDefinitionInfo::GetDisplayNameText() const
+{
+	FText LocalizedDisplayName;
+
+	static const FString Namespace = TEXT("UObjectDisplayNames");
+	static const FName NAME_DisplayName(TEXT("DisplayName"));
+
+	const FString Key = FUHTFieldDisplayNameHelper::GetFullGroupName(*this, false);
+
+	FString NativeDisplayName;
+	if (const FString* FoundMetaData = FindMetaData(NAME_DisplayName))
+	{
+		NativeDisplayName = *FoundMetaData;
+	}
+	else
+	{
+		NativeDisplayName = FName::NameToDisplayString(FUHTFieldDisplayNameHelper::Get(*this), IsBooleanOrBooleanStaticArray());
+	}
+
+	if (!(FText::FindText(Namespace, Key, /*OUT*/LocalizedDisplayName, &NativeDisplayName)))
+	{
+		LocalizedDisplayName = FText::FromString(NativeDisplayName);
+	}
+
+#if UHT_ENABLE_ENGINE_TYPE_CHECKS
+	if (GetPropertySafe())
+	{
+		check(LocalizedDisplayName.ToString() == GetPropertySafe()->GetDisplayNameText().ToString());
+	}
+#endif
+	return LocalizedDisplayName;
+}
+
+FText FUnrealPropertyDefinitionInfo::GetToolTipText(bool bShortTooltip) const
+{
+	bool bFoundShortTooltip = false;
+	static const FName NAME_Tooltip(TEXT("Tooltip"));
+	static const FName NAME_ShortTooltip(TEXT("ShortTooltip"));
+	FText LocalizedToolTip;
+	FString NativeToolTip;
+
+	if (bShortTooltip)
+	{
+		NativeToolTip = GetMetaData(NAME_ShortTooltip);
+		if (NativeToolTip.IsEmpty())
+		{
+			NativeToolTip = GetMetaData(NAME_Tooltip);
+		}
+		else
+		{
+			bFoundShortTooltip = true;
+		}
+	}
+	else
+	{
+		NativeToolTip = GetMetaData(NAME_Tooltip);
+	}
+
+	const FString Namespace = bFoundShortTooltip ? TEXT("UObjectShortTooltips") : TEXT("UObjectToolTips");
+	const FString Key = FUHTFieldDisplayNameHelper::GetFullGroupName(*this, false);
+	if (!FText::FindText(Namespace, Key, /*OUT*/LocalizedToolTip, &NativeToolTip))
+	{
+		if (NativeToolTip.IsEmpty())
+		{
+			NativeToolTip = FName::NameToDisplayString(FUHTFieldDisplayNameHelper::Get(*this), IsBooleanOrBooleanStaticArray());
+		}
+		else
+		{
+			static const FString DoxygenSee(TEXT("@see"));
+			static const FString TooltipSee(TEXT("See:"));
+			if (NativeToolTip.ReplaceInline(*DoxygenSee, *TooltipSee) > 0)
+			{
+				NativeToolTip.TrimEndInline();
+			}
+		}
+		LocalizedToolTip = FText::FromString(NativeToolTip);
+	}
+
+#if UHT_ENABLE_ENGINE_TYPE_CHECKS
+	if (GetPropertySafe())
+	{
+		check(LocalizedToolTip.ToString() == GetPropertySafe()->GetToolTipText(bShortTooltip).ToString());
+	}
+#endif
+	return LocalizedToolTip;
+}
+
 FUnrealPackageDefinitionInfo& FUnrealPropertyDefinitionInfo::GetPackageDef() const
 {
 	if (HasSource())
@@ -597,6 +763,18 @@ FUnrealPackageDefinitionInfo& FUnrealPropertyDefinitionInfo::GetPackageDef() con
 	return GTypeDefinitionInfoMap.FindChecked<FUnrealPackageDefinitionInfo>(GetProperty()->GetOutermost());
 }
 
+bool FUnrealPropertyDefinitionInfo::SameType(const FUnrealPropertyDefinitionInfo& Other) const
+{
+	bool bResults = FPropertyTraits::SameType(*this, Other);
+#if UHT_ENABLE_ENGINE_TYPE_CHECKS
+	if (GetPropertySafe() && Other.GetPropertySafe())
+	{
+		check(GetPropertySafe()->SameType(Other.GetPropertySafe()) == bResults)
+	}
+#endif
+	return bResults;
+}
+
 FUnrealPackageDefinitionInfo& FUnrealObjectDefinitionInfo::GetPackageDef() const
 {
 	if (HasSource())
@@ -606,6 +784,16 @@ FUnrealPackageDefinitionInfo& FUnrealObjectDefinitionInfo::GetPackageDef() const
 	return GTypeDefinitionInfoMap.FindChecked<FUnrealPackageDefinitionInfo>(GetObject()->GetPackage());
 }
 
+FString FUnrealObjectDefinitionInfo::GetPathName(FUnrealObjectDefinitionInfo* StopOuter) const
+{
+	return GetObject()->GetPathName(StopOuter ? StopOuter->GetObject() : nullptr);
+}
+
+void FUnrealObjectDefinitionInfo::GetPathName(FUnrealObjectDefinitionInfo* StopOuter, FStringBuilderBase& ResultString) const
+{
+	return GetObject()->GetPathName(StopOuter ? StopOuter->GetObject() : nullptr, ResultString);
+}
+
 FUnrealPackageDefinitionInfo::FUnrealPackageDefinitionInfo(const FManifestModule& InModule, UPackage* InPackage)
 	: FUnrealObjectDefinitionInfo(FString())
 	, Module(InModule)
@@ -615,7 +803,7 @@ FUnrealPackageDefinitionInfo::FUnrealPackageDefinitionInfo(const FManifestModule
 	SetObject(InPackage);
 }
 
-void FUnrealPackageDefinitionInfo::PostParseFinalize()
+void FUnrealPackageDefinitionInfo::PostParseFinalizeInternal(bool bCreateEngineTypes)
 {
 	UPackage* Package = GetPackage();
 
@@ -625,6 +813,14 @@ void FUnrealPackageDefinitionInfo::PostParseFinalize()
 	SingletonName.Appendf(TEXT("Z_Construct_UPackage_%s()"), *PackageName);
 	SingletonNameChopped = SingletonName.LeftChop(2);
 	ExternDecl.Appendf(TEXT("\tUPackage* %s;\r\n"), *SingletonName);
+
+	for (TSharedRef<FUnrealSourceFile>& LocalSourceFile : GetAllSourceFiles())
+	{
+		for (TSharedRef<FUnrealTypeDefinitionInfo>& TypeDef : LocalSourceFile->GetDefinedTypes())
+		{
+			TypeDef->PostParseFinalize(bCreateEngineTypes);
+		}
+	}
 }
 
 void FUnrealPackageDefinitionInfo::AddCrossModuleReference(TSet<FString>* UniqueCrossModuleReferences) const
@@ -635,7 +831,7 @@ void FUnrealPackageDefinitionInfo::AddCrossModuleReference(TSet<FString>* Unique
 	}
 }
 
-void FUnrealFieldDefinitionInfo::PostParseFinalize()
+void FUnrealFieldDefinitionInfo::PostParseFinalizeInternal(bool bCreateEngineTypes)
 {
 	const TCHAR* TypeStr = GetSimplifiedTypeClass();
 	UField* Field = GetField();
@@ -730,34 +926,21 @@ FUnrealEnumDefinitionInfo::FUnrealEnumDefinitionInfo(FUnrealSourceFile& InSource
 	: FUnrealFieldDefinitionInfo(InSourceFile, InLineNumber, MoveTemp(InNameCPP), InSourceFile.GetPackageDef())
 { }
 
-void FUnrealStructDefinitionInfo::AddProperty(FUnrealPropertyDefinitionInfo& PropertyDef)
+void FUnrealStructDefinitionInfo::AddProperty(TSharedRef<FUnrealPropertyDefinitionInfo> PropertyDef)
 {
-	Properties.Add(&PropertyDef);
-
-	// Add the property to the engine object
-	FProperty* Property = PropertyDef.GetProperty();
-
-	// Add to the end of the properties slist
-	FField** Prev = &GetStruct()->ChildProperties;
-	for (; *Prev != nullptr; Prev = &(*Prev)->Next)
-	{
-		// No body
-	}
-	check(*Prev == nullptr);
-	Property->Next = nullptr;
-	*Prev = Property;
+	Properties.Add(PropertyDef);
 
 	// update the optimization flags
 	if (!bContainsDelegates)
 	{
-		const FPropertyBase& PropertyBase = PropertyDef.GetPropertyBase();
-		if (PropertyDef.IsDelegateOrDelegateStaticArray() || PropertyDef.IsMulticastDelegateOrMulticastDelegateStaticArray())
+		const FPropertyBase& PropertyBase = PropertyDef->GetPropertyBase();
+		if (PropertyDef->IsDelegateOrDelegateStaticArray() || PropertyDef->IsMulticastDelegateOrMulticastDelegateStaticArray())
 		{
 			bContainsDelegates = true;
 		}
-		else if (PropertyDef.IsDynamicArray())
+		else if (PropertyDef->IsDynamicArray())
 		{
-			const FUnrealPropertyDefinitionInfo& ValuePropertyDef = PropertyDef.GetValuePropDef();
+			const FUnrealPropertyDefinitionInfo& ValuePropertyDef = PropertyDef->GetValuePropDef();
 			if (ValuePropertyDef.IsDelegateOrDelegateStaticArray() || ValuePropertyDef.IsMulticastDelegateOrMulticastDelegateStaticArray())
 			{
 				bContainsDelegates = true;
@@ -766,14 +949,63 @@ void FUnrealStructDefinitionInfo::AddProperty(FUnrealPropertyDefinitionInfo& Pro
 	}
 }
 
-void FUnrealStructDefinitionInfo::AddFunction(FUnrealFunctionDefinitionInfo& FunctionDef)
+void FUnrealStructDefinitionInfo::PostParseFinalizeInternal(bool bCreateEngineTypes)
 {
-	Functions.Add(&FunctionDef);
+	FUnrealFieldDefinitionInfo::PostParseFinalizeInternal(bCreateEngineTypes);
+
+	if (SuperStructInfo.Struct)
+	{
+		SuperStructInfo.Struct->PostParseFinalize(bCreateEngineTypes);
+	}
+
+	for (const FBaseStructInfo& Info : BaseStructInfo)
+	{
+		if (Info.Struct)
+		{
+			Info.Struct->PostParseFinalize(bCreateEngineTypes);
+		}
+	}
+
+	for (TSharedRef<FUnrealPropertyDefinitionInfo> PropertyDef : Properties)
+	{
+		PropertyDef->PostParseFinalize(bCreateEngineTypes);
+	}
+
+	for (TSharedRef<FUnrealFunctionDefinitionInfo> FunctionDef : Functions)
+	{
+		FunctionDef->PostParseFinalize(bCreateEngineTypes);
+	}
+
+	GetStruct()->Bind();
+
+	// Internals will assert of we are relinking an intrinsic 
+	bool bRelinkExistingProperties = true;
+	if (FUnrealClassDefinitionInfo* ClassDef = UHTCast<FUnrealClassDefinitionInfo>(this))
+	{
+		bRelinkExistingProperties = !ClassDef->HasAnyClassFlags(CLASS_Intrinsic);
+	}
+	GetStruct()->StaticLink(bRelinkExistingProperties);
+}
+
+void FUnrealStructDefinitionInfo::CreatePropertyEngineTypes()
+{
+	for (TSharedRef<FUnrealPropertyDefinitionInfo> PropertyDef : Properties)
+	{
+		if (PropertyDef->GetPropertySafe() == nullptr)
+		{
+			FPropertyTraits::CreateEngineType(PropertyDef);
+		}
+	}
+}
+
+void FUnrealStructDefinitionInfo::AddFunction(TSharedRef<FUnrealFunctionDefinitionInfo> FunctionDef)
+{
+	Functions.Add(FunctionDef);
 
 	// update the optimization flags
 	if (!bContainsDelegates)
 	{
-		if (FunctionDef.HasAnyFunctionFlags(FUNC_Delegate))
+		if (FunctionDef->HasAnyFunctionFlags(FUNC_Delegate))
 		{
 			bContainsDelegates = true;
 		}
@@ -923,8 +1155,10 @@ uint32 FUnrealClassDefinitionInfo::GetHash(bool bIncludeNoExport) const
 }
 
 
-void FUnrealClassDefinitionInfo::PostParseFinalize()
+void FUnrealClassDefinitionInfo::PostParseFinalizeInternal(bool bCreateEngineTypes)
 {
+	FUnrealStructDefinitionInfo::PostParseFinalizeInternal(bCreateEngineTypes);
+
 	if (IsInterface() && GetStructMetaData().ParsedInterface == EParsedInterface::ParsedUInterface)
 	{
 		FString UName = GetNameCPP();
@@ -932,7 +1166,23 @@ void FUnrealClassDefinitionInfo::PostParseFinalize()
 		FError::Throwf(TEXT("UInterface '%s' parsed without a corresponding '%s'"), *UName, *IName);
 	}
 
-	FUnrealStructDefinitionInfo::PostParseFinalize();
+	FHeaderParser::CheckSparseClassData(*this);
+
+	if (!HasAnyClassFlags(CLASS_Native))
+	{
+		//Class->UnMark(EObjectMark(OBJECTMARK_TagImp | OBJECTMARK_TagExp));
+	}
+	else if (!HasAnyClassFlags(CLASS_NoExport | CLASS_Intrinsic))
+	{
+		GetPackageDef().SetWriteClassesH(true);
+		//Class->UnMark(OBJECTMARK_TagImp);
+		//Class->Mark(OBJECTMARK_TagExp);
+	}
+
+	// This needs to be done outside of parallel blocks because it will modify UClass memory.
+	// Later calls to SetUpUhtReplicationData inside parallel blocks should be fine, because
+	// they will see the memory has already been set up, and just return the parent pointer.
+	GetClass()->SetUpUhtReplicationData();
 }
 
 void FUnrealClassDefinitionInfo::ParseClassProperties(TArray<FPropertySpecifier>&& InClassSpecifiers, const FString& InRequiredAPIMacroIfPresent)
@@ -1536,14 +1786,14 @@ FString FUnrealClassDefinitionInfo::GetNameWithPrefix(EEnforceInterfacePrefix En
 	return FString::Printf(TEXT("%s%s"), Prefix, *GetName());
 }
 
-void FUnrealFunctionDefinitionInfo::AddProperty(FUnrealPropertyDefinitionInfo& PropertyDef)
+void FUnrealFunctionDefinitionInfo::AddProperty(TSharedRef<FUnrealPropertyDefinitionInfo> PropertyDef)
 {
-	check(PropertyDef.HasAnyPropertyFlags(CPF_Parm));
+	check(PropertyDef->HasAnyPropertyFlags(CPF_Parm));
 
-	if (PropertyDef.HasAnyPropertyFlags(CPF_ReturnParm))
+	if (PropertyDef->HasAnyPropertyFlags(CPF_ReturnParm))
 	{
 		check(ReturnProperty == nullptr);
-		ReturnProperty = &PropertyDef;
+		ReturnProperty = PropertyDef;
 	}
 	FUnrealStructDefinitionInfo::AddProperty(PropertyDef);
 }
@@ -1552,3 +1802,11 @@ FUnrealFunctionDefinitionInfo* FUnrealFunctionDefinitionInfo::GetSuperFunction()
 {
 	return UHTCast<FUnrealFunctionDefinitionInfo>(GetSuperStruct());
 }
+
+void FUnrealFunctionDefinitionInfo::PostParseFinalizeInternal(bool bCreateEngineTypes)
+{
+	FUnrealStructDefinitionInfo::PostParseFinalizeInternal(bCreateEngineTypes);
+
+	GetFunction()->Bind();
+}
+
diff --git a/Engine/Source/Programs/UnrealHeaderTool/Private/UnrealTypeDefinitionInfo.h b/Engine/Source/Programs/UnrealHeaderTool/Private/UnrealTypeDefinitionInfo.h
--- a/Engine/Source/Programs/UnrealHeaderTool/Private/UnrealTypeDefinitionInfo.h
+++ b/Engine/Source/Programs/UnrealHeaderTool/Private/UnrealTypeDefinitionInfo.h
@@ -144,8 +144,9 @@ public:
 	}
 	virtual const FString& GetMetaData(const FName& Key) const;
 
-protected:
 	virtual TMap<FName, FString>* GetMetaDataMap() const = 0;
+
+protected:
 #if UHT_ENABLE_ENGINE_TYPE_CHECKS
 	virtual void CheckFindMetaData(const FName& Key, const FString* ValuePtr) const = 0;
 #endif
@@ -156,6 +157,14 @@ protected:
  */
 class FUnrealTypeDefinitionInfo : public TSharedFromThis<FUnrealTypeDefinitionInfo>
 {
+private:
+	enum class EFinalizeState : uint8
+	{
+		None,
+		InProgress,
+		Finished,
+	};
+
 public:
 	virtual ~FUnrealTypeDefinitionInfo() = default;
 
@@ -315,8 +324,22 @@ public:
 	/**
 	 * Perform any post parsing finalization and validation
 	 */
-	virtual void PostParseFinalize()
-	{}
+	void PostParseFinalize(bool bCreateEngineTypes)
+	{
+		switch (FinalizeState)
+		{
+		case EFinalizeState::None:
+			FinalizeState = EFinalizeState::InProgress;
+			PostParseFinalizeInternal(bCreateEngineTypes);
+			FinalizeState = EFinalizeState::Finished;
+			break;
+		case EFinalizeState::InProgress:
+			check(false);
+			break;
+		case EFinalizeState::Finished:
+			break;
+		}
+	}
 
 	/**
 	 * Return the C++ type name that represents this type (i.e. UClass)
@@ -350,8 +373,8 @@ public:
 	 * Returns the fully qualified pathname for this object, in the format:
 	 * 'Outermost.[Outer:]Name'
 	 */
-	virtual FString GetPathName() const = 0;
-	virtual void GetPathName(FStringBuilderBase& ResultString) const = 0;
+	virtual FString GetPathName(FUnrealObjectDefinitionInfo* StopOuter = nullptr) const = 0;
+	virtual void GetPathName(FUnrealObjectDefinitionInfo* StopOuter, FStringBuilderBase& ResultString) const = 0;
 
 	/**
 	 * Return true if this type has source information
@@ -482,12 +505,16 @@ protected:
 
 	virtual void ValidateMetaDataFormat(const FName InKey, ECheckedMetadataSpecifier InCheckedMetadataSpecifier, const FString& InValue) const;
 
+	virtual void PostParseFinalizeInternal(bool bCreateEngineTypes)
+	{}
+
 private:
 	FString NameCPP;
 	FUnrealTypeDefinitionInfo* Outer = nullptr;
 	FUnrealSourceFile* SourceFile = nullptr;
 	int32 LineNumber = 0;
 	std::atomic<uint32> Hash = 0;
+	EFinalizeState FinalizeState = EFinalizeState::None;
 };
 
 /**
@@ -498,13 +525,14 @@ class FUnrealPropertyDefinitionInfo
 	, public FUHTMetaData
 {
 public:
-	FUnrealPropertyDefinitionInfo(FUnrealSourceFile& InSourceFile, int32 InLineNumber, int32 InParsePosition, const FPropertyBase& InVarProperty, EVariableCategory InVariableCategory, const TCHAR* Dimensions, FName InName, FUnrealTypeDefinitionInfo& InOuter)
+	FUnrealPropertyDefinitionInfo(FUnrealSourceFile& InSourceFile, int32 InLineNumber, int32 InParsePosition, const FPropertyBase& InVarProperty, EVariableCategory InVariableCategory, EAccessSpecifier InAccessSpecifier, const TCHAR* Dimensions, FName InName, FUnrealTypeDefinitionInfo& InOuter)
 		: FUnrealTypeDefinitionInfo(InSourceFile, InLineNumber, InName.ToString(), InOuter)
 		, Name(InName)
 		, PropertyBase(InVarProperty)
 		, ArrayDimensions(InVarProperty.ArrayType == EArrayType::Static ? FString(Dimensions) : FString())
 		, ParsePosition(InParsePosition)
 		, VariableCategory(InVariableCategory)
+		, AccessSpecifier(InAccessSpecifier)
 	{ }
 
 	virtual FUnrealPropertyDefinitionInfo* AsProperty() override
@@ -520,11 +548,6 @@ public:
 		return TEXT("UProperty");
 	}
 
-	/**
-	 * Perform any post parsing finalization and validation
-	 */
-	virtual void PostParseFinalize() override;
-
 	/**
 	 * Return the Engine instance associated with the compiler instance
 	 */
@@ -534,6 +557,14 @@ public:
 		return Property;
 	}
 
+	/**
+	 * Return the Engine instance associated with the compiler instance, does not check the pointer, can be null
+	 */
+	FProperty* GetPropertySafe() const
+	{
+		return Property;
+	}
+
 	/**
 	 * Sets the engine type
 	 */
@@ -560,7 +591,7 @@ public:
 	virtual FString GetName() const override
 	{
 #if UHT_ENABLE_ENGINE_TYPE_CHECKS
-		check(GetProperty()->GetName() == GetNameCPP());
+		check(GetPropertySafe() == nullptr || GetPropertySafe()->GetName() == GetNameCPP());
 #endif
 		return GetNameCPP();
 	}
@@ -569,7 +600,7 @@ public:
 	virtual FName GetFName() const override
 	{
 #if UHT_ENABLE_ENGINE_TYPE_CHECKS
-		check(GetProperty()->GetFName() == Name);
+		check(GetPropertySafe() == nullptr || GetPropertySafe()->GetFName() == Name);
 #endif
 		return Name;
 	}
@@ -580,8 +611,8 @@ public:
 	 * Returns the fully qualified pathname for this object, in the format:
 	 * 'Outermost.[Outer:]Name'
 	 */
-	virtual FString GetPathName() const override;
-	virtual void GetPathName(FStringBuilderBase& ResultString) const override;
+	virtual FString GetPathName(FUnrealObjectDefinitionInfo* StopOuter = nullptr) const override;
+	virtual void GetPathName(FUnrealObjectDefinitionInfo* StopOuter, FStringBuilderBase& ResultString) const override;
 
 	/**
 	 * Return the engine class name
@@ -598,10 +629,7 @@ public:
 	*
 	* @return The display name for this object.
 	*/
-	FText GetDisplayNameText() const
-	{
-		return GetProperty()->GetDisplayNameText();
-	}
+	FText GetDisplayNameText() const;
 
 	/**
 	* Finds the localized tooltip or native tooltip as a fallback.
@@ -610,10 +638,7 @@ public:
 	*
 	* @return The tooltip for this object.
 	*/
-	FText GetToolTipText(bool bShortTooltip = false) const
-	{
-		return GetProperty()->GetToolTipText(bShortTooltip);
-	}
+	FText GetToolTipText(bool bShortTooltip = false) const;
 
 	/**
 	 * Returns the C++ name of the property, including the _DEPRECATED suffix if the
@@ -625,7 +650,7 @@ public:
 	{
 		FString OutName = HasAnyPropertyFlags(CPF_Deprecated) ? GetName() + TEXT("_DEPRECATED") : GetName();
 #if UHT_ENABLE_ENGINE_TYPE_CHECKS
-		check(GetProperty()->GetNameCPP() == OutName);
+		check(GetPropertySafe() == nullptr || GetPropertySafe()->GetNameCPP() == OutName);
 #endif
 		return OutName;
 	}
@@ -646,7 +671,7 @@ public:
 	EPropertyFlags GetPropertyFlags() const
 	{
 #if UHT_ENABLE_ENGINE_TYPE_CHECKS
-		check((GetProperty()->GetPropertyFlags() & ~CPF_ComputedFlags) == PropertyBase.PropertyFlags);
+		check(GetPropertySafe() == nullptr || (GetPropertySafe()->GetPropertyFlags() & ~CPF_ComputedFlags) == PropertyBase.PropertyFlags);
 #endif
 		return PropertyBase.PropertyFlags;
 	}
@@ -654,18 +679,24 @@ public:
 	void SetPropertyFlags(EPropertyFlags NewFlags)
 	{
 #if UHT_ENABLE_ENGINE_TYPE_CHECKS
-		check((GetProperty()->GetPropertyFlags() & ~CPF_ComputedFlags) == PropertyBase.PropertyFlags);
+		check(GetPropertySafe() == nullptr || (GetPropertySafe()->GetPropertyFlags() & ~CPF_ComputedFlags) == PropertyBase.PropertyFlags);
 #endif
-		GetProperty()->PropertyFlags |= NewFlags;
+		if (FProperty* LocalProperty = GetPropertySafe())
+		{
+			LocalProperty->PropertyFlags |= NewFlags;
+		}
 		PropertyBase.PropertyFlags |= NewFlags;
 	}
 
 	void ClearPropertyFlags(EPropertyFlags NewFlags)
 	{
 #if UHT_ENABLE_ENGINE_TYPE_CHECKS
-		check((GetProperty()->GetPropertyFlags() & ~CPF_ComputedFlags) == PropertyBase.PropertyFlags);
+		check(GetPropertySafe() == nullptr || (GetPropertySafe()->GetPropertyFlags() & ~CPF_ComputedFlags) == PropertyBase.PropertyFlags);
 #endif
-		GetProperty()->PropertyFlags &= ~NewFlags;
+		if (FProperty* LocalProperty = GetPropertySafe())
+		{
+			LocalProperty->PropertyFlags &= ~NewFlags;
+		}
 		PropertyBase.PropertyFlags &= ~NewFlags;
 	}
 
@@ -684,7 +715,7 @@ public:
 	{
 #if UHT_ENABLE_ENGINE_TYPE_CHECKS
 		check((FlagsToCheck & CPF_ComputedFlags) == 0);
-		check((GetProperty()->GetPropertyFlags() & ~CPF_ComputedFlags) == PropertyBase.PropertyFlags);
+		check(GetPropertySafe() == nullptr || (GetPropertySafe()->GetPropertyFlags() & ~CPF_ComputedFlags) == PropertyBase.PropertyFlags);
 #endif
 		return (PropertyBase.PropertyFlags & FlagsToCheck) != 0;
 	}
@@ -703,7 +734,7 @@ public:
 	{
 #if UHT_ENABLE_ENGINE_TYPE_CHECKS
 		check((FlagsToCheck & CPF_ComputedFlags) == 0);
-		check((GetProperty()->GetPropertyFlags() & ~CPF_ComputedFlags) == PropertyBase.PropertyFlags);
+		check(GetPropertySafe() == nullptr || (GetPropertySafe()->GetPropertyFlags() & ~CPF_ComputedFlags) == PropertyBase.PropertyFlags);
 #endif
 		return (PropertyBase.PropertyFlags & FlagsToCheck) == FlagsToCheck;
 	}
@@ -712,7 +743,7 @@ public:
 	{
 #if UHT_ENABLE_ENGINE_TYPE_CHECKS
 		check((FlagsToCheck & CPF_ComputedFlags) == 0);
-		check((GetProperty()->GetPropertyFlags() & ~CPF_ComputedFlags) == PropertyBase.PropertyFlags);
+		check(GetPropertySafe() == nullptr || (GetPropertySafe()->GetPropertyFlags() & ~CPF_ComputedFlags) == PropertyBase.PropertyFlags);
 #endif
 		return (PropertyBase.PropertyFlags & FlagsToCheck) == ExpectedFlags;
 	}
@@ -742,7 +773,7 @@ public:
 	{
 		bool bResult = HasAnyPropertyFlags(CPF_DevelopmentAssets);
 #if UHT_ENABLE_ENGINE_TYPE_CHECKS
-		check(bResult == GetProperty()->IsEditorOnlyProperty());
+		check(GetPropertySafe() == nullptr || bResult == GetPropertySafe()->IsEditorOnlyProperty());
 #endif
 		return bResult;
 	}
@@ -757,7 +788,7 @@ public:
 	{
 		bool bResult = HasAnyPropertyFlags(CPF_ContainsInstancedReference | CPF_InstancedReference);
 #if UHT_ENABLE_ENGINE_TYPE_CHECKS
-		check(bResult == GetProperty()->ContainsInstancedObjectProperty());
+		check(GetPropertySafe() == nullptr || bResult == GetPropertySafe()->ContainsInstancedObjectProperty());
 #endif
 		return bResult;
 	}
@@ -765,10 +796,7 @@ public:
 	/**
 	 * Test to see if two properties share the same types.  This method utilizes the underlying UProperty SameType implementation.
 	 */
-	bool SameType(const FUnrealPropertyDefinitionInfo& Other) const
-	{
-		return GetProperty()->SameType(Other.GetProperty());
-	}
+	bool SameType(const FUnrealPropertyDefinitionInfo& Other) const;
 
 	/**
 	 * Determines whether this property's type is compatible with another property's type.
@@ -926,6 +954,14 @@ public:
 		return VariableCategory;
 	}
 
+	/**
+	 * Get the access specifier of this property
+	 */
+	EAccessSpecifier GetAccessSpecifier() const
+	{
+		return AccessSpecifier;
+	}
+
 	/**
 	 * Return true if the property is unsized
 	 */
@@ -970,17 +1006,12 @@ public:
 		return PropertyBase;
 	}
 
-	/**
-	 * Add meta data for the given definition
-	 */
-	virtual void AddMetaData(TMap<FName, FString>&& InMetaData) override;
-
 	/**
 	 * Return true if we have a key property def
 	 */
 	bool HasKeyPropDef() const
 	{
-		return KeyPropDef != nullptr;
+		return KeyPropDef.IsValid();
 	}
 
 	/**
@@ -993,20 +1024,29 @@ public:
 	}
 
 	/**
-	 * Return true if we have a value property def
+	 * Get the associated key property definition (valid for maps)
 	 */
-	bool HasValuePropDef() const
+	TSharedRef<FUnrealPropertyDefinitionInfo> GetKeyPropDefRef() const
 	{
-		return ValuePropDef != nullptr;
+		check(KeyPropDef);
+		return KeyPropDef.ToSharedRef();
 	}
 
 	/**
 	 * Sets the associated key property definition (valid for maps)
 	 */
-	void SetKeyPropDef(FUnrealPropertyDefinitionInfo& InKeyPropDef)
+	void SetKeyPropDef(TSharedRef<FUnrealPropertyDefinitionInfo> InKeyPropDef)
 	{
-		check(KeyPropDef == nullptr);
-		KeyPropDef = &InKeyPropDef;
+		check(!KeyPropDef);
+		KeyPropDef = InKeyPropDef;
+	}
+
+	/**
+	 * Return true if we have a value property def
+	 */
+	bool HasValuePropDef() const
+	{
+		return ValuePropDef.IsValid();
 	}
 
 	/**
@@ -1018,13 +1058,22 @@ public:
 		return *ValuePropDef;
 	}
 
+	/**
+	 * Get the associated value property definition (valid for maps, sets, and dynamic arrays)
+	 */
+	TSharedRef<FUnrealPropertyDefinitionInfo> GetValuePropDefRef() const
+	{
+		check(ValuePropDef);
+		return ValuePropDef.ToSharedRef();
+	}
+
 	/**
 	 * Sets the associated value property definition (valid for maps, sets, and dynamic arrays)
 	 */
-	void SetValuePropDef(FUnrealPropertyDefinitionInfo& InValuePropDef)
+	void SetValuePropDef(TSharedRef<FUnrealPropertyDefinitionInfo> InValuePropDef)
 	{
-		check(ValuePropDef == nullptr);
-		ValuePropDef = &InValuePropDef;
+		check(!ValuePropDef);
+		ValuePropDef = InValuePropDef;
 	}
 
 	/**
@@ -1059,7 +1108,7 @@ public:
 	FName GetRepNotifyFunc() const
 	{
 #if UHT_ENABLE_ENGINE_TYPE_CHECKS
-		check(PropertyBase.RepNotifyName == GetProperty()->RepNotifyFunc);
+		check(GetPropertySafe() == nullptr || PropertyBase.RepNotifyName == GetPropertySafe()->RepNotifyFunc);
 #endif
 		return PropertyBase.RepNotifyName;
 	}
@@ -1069,30 +1118,48 @@ public:
 	 */
 	void SetDelegateFunctionSignature(FUnrealFunctionDefinitionInfo& DelegateFunctionDef);
 
-protected:
+	/**
+	 * Add meta data for the given definition
+	 */
+	virtual void AddMetaData(TMap<FName, FString>&& InMetaData) override
+	{
+		check(Property == nullptr);
+		GetMetaDataMap()->Append(MoveTemp(InMetaData));
+	}
+
 	virtual TMap<FName, FString>* GetMetaDataMap() const override
 	{
-		return const_cast<TMap<FName, FString>*>(GetProperty()->GetMetaDataMap());
+		return const_cast<TMap<FName, FString>*>(&PropertyBase.MetaData);
 	}
 
+protected:
 #if UHT_ENABLE_ENGINE_TYPE_CHECKS
 	virtual void CheckFindMetaData(const FName& Key, const FString* ValuePtr) const
 	{
-		const FString* CheckPtr = GetProperty()->FindMetaData(Key);
-		check((CheckPtr == nullptr && ValuePtr == nullptr) || (CheckPtr != nullptr && ValuePtr != nullptr && *CheckPtr == *ValuePtr));
+		if (FProperty* LocalProperty = GetPropertySafe())
+		{
+			const FString* CheckPtr = LocalProperty->FindMetaData(Key);
+			check((CheckPtr == nullptr && ValuePtr == nullptr) || (CheckPtr != nullptr && ValuePtr != nullptr && *CheckPtr == *ValuePtr));
+		}
 	}
 #endif
 
+	/**
+	 * Perform any post parsing finalization and validation
+	 */
+	virtual void PostParseFinalizeInternal(bool bCreateEngineTypes) override;
+
 private:
 	FName Name;
 	FPropertyBase PropertyBase;
 	FString ArrayDimensions;
 	FString TypePackageName;
-	FUnrealPropertyDefinitionInfo* KeyPropDef = nullptr;
-	FUnrealPropertyDefinitionInfo* ValuePropDef = nullptr;
+	TSharedPtr<FUnrealPropertyDefinitionInfo> KeyPropDef;
+	TSharedPtr<FUnrealPropertyDefinitionInfo> ValuePropDef;
 	FProperty* Property = nullptr;
 	int32 ParsePosition;
 	EVariableCategory VariableCategory = EVariableCategory::Member;
+	EAccessSpecifier AccessSpecifier = ACCESS_Public;
 	EAllocatorType AllocatorType = EAllocatorType::Default;
 	bool bUnsized = false;
 };
@@ -1152,15 +1219,8 @@ public:
 	 * Returns the fully qualified pathname for this object, in the format:
 	 * 'Outermost.[Outer:]Name'
 	 */
-	virtual FString GetPathName() const override
-	{
-		return GetObject()->GetPathName(nullptr);
-	}
-	
-	virtual void GetPathName(FStringBuilderBase& ResultString) const override
-	{
-		return GetObject()->GetPathName(nullptr, ResultString);
-	}
+	virtual FString GetPathName(FUnrealObjectDefinitionInfo* StopOuter = nullptr) const override;
+	virtual void GetPathName(FUnrealObjectDefinitionInfo* StopOuter, FStringBuilderBase& ResultString) const override;
 
 	/**
 	 * Return the package associated with this object
@@ -1312,11 +1372,6 @@ public:
 		return ShortUpperName;
 	}
 
-	/**
-	 * Perform any post parsing finalization and validation
-	 */
-	virtual void PostParseFinalize() override;
-
 	/**
 	 * Add a unique cross module reference for this field
 	 */
@@ -1346,6 +1401,12 @@ public:
 		return ExternDecl;
 	}
 
+protected:
+	/**
+	 * Perform any post parsing finalization and validation
+	 */
+	virtual void PostParseFinalizeInternal(bool bCreateEngineTypes) override;
+
 private:
 	const FManifestModule& Module;
 	TArray<TSharedRef<FUnrealSourceFile>> AllSourceFiles;
@@ -1436,11 +1497,6 @@ public:
 		GetField()->SetMetaData(Key, InValue);
 	}
 
-	/**
-	 * Perform any post parsing finalization and validation
-	 */
-	virtual void PostParseFinalize() override;
-
 	/**
 	 * Add a unique cross module reference for this field
 	 */
@@ -1498,6 +1554,12 @@ public:
 	 */
 	FUnrealClassDefinitionInfo* GetOwnerClass() const;
 
+protected:
+	/**
+	 * Perform any post parsing finalization and validation
+	 */
+	virtual void PostParseFinalizeInternal(bool bCreateEngineTypes) override;
+
 private:
 	FString SingletonName[2];
 	FString SingletonNameChopped[2];
@@ -1782,19 +1844,24 @@ public:
 		return FString::Printf(TEXT("%s%s"), (bForceInterface ? TEXT("I") : GetPrefixCPP()), *GetName());
 	}
 
+	/**
+	 * Create the underlying properties for the structure
+	 */
+	void CreatePropertyEngineTypes();
+
 	/**
 	 * Add a new property to the structure
 	 */
-	virtual void AddProperty(FUnrealPropertyDefinitionInfo& PropertyDef);
+	virtual void AddProperty(TSharedRef<FUnrealPropertyDefinitionInfo> PropertyDef);
 
 	/**
 	 * Get the collection of properties
 	 */
-	const TArray<FUnrealPropertyDefinitionInfo*>& GetProperties() const
+	const TArray<TSharedRef<FUnrealPropertyDefinitionInfo>>& GetProperties() const
 	{
 		return Properties;
 	}
-	TArray<FUnrealPropertyDefinitionInfo*>& GetProperties()
+	TArray<TSharedRef<FUnrealPropertyDefinitionInfo>>& GetProperties()
 	{
 		return Properties;
 	}
@@ -1802,16 +1869,16 @@ public:
 	/**
 	 * Add a new function to the structure
 	 */
-	virtual void AddFunction(FUnrealFunctionDefinitionInfo& FunctionDef);
+	virtual void AddFunction(TSharedRef<FUnrealFunctionDefinitionInfo> FunctionDef);
 
 	/**
 	 * Get the collection of functions
 	 */
-	const TArray<FUnrealFunctionDefinitionInfo*>& GetFunctions() const
+	const TArray< TSharedRef<FUnrealFunctionDefinitionInfo>>& GetFunctions() const
 	{
 		return Functions;
 	}
-	TArray<FUnrealFunctionDefinitionInfo*>& GetFunctions()
+	TArray<TSharedRef<FUnrealFunctionDefinitionInfo>>& GetFunctions()
 	{
 		return Functions;
 	}
@@ -1960,14 +2027,20 @@ public:
 		return RigVMInfo;
 	}
 
+protected:
+	/**
+	 * Perform any post parsing finalization and validation
+	 */
+	virtual void PostParseFinalizeInternal(bool bCreateEngineTypes) override;
+
 private:
 	TSharedPtr<FScope> StructScope;
 
 	/** Properties of the structure */
-	TArray<FUnrealPropertyDefinitionInfo*> Properties;
+	TArray<TSharedRef<FUnrealPropertyDefinitionInfo>> Properties;
 
 	/** Functions of the structure */
-	TArray<FUnrealFunctionDefinitionInfo*> Functions;
+	TArray<TSharedRef<FUnrealFunctionDefinitionInfo>> Functions;
 
 	FStructMetaData StructMetaData;
 	FBaseStructInfo SuperStructInfo;
@@ -2168,7 +2241,7 @@ public:
 	//@{
 	FFuncInfo& GetFunctionData() { return FunctionData; }
 	const FFuncInfo& GetFunctionData() const { return FunctionData; }
-	FUnrealPropertyDefinitionInfo* GetReturn() const { return ReturnProperty; }
+	FUnrealPropertyDefinitionInfo* GetReturn() const { return ReturnProperty.Get(); }
 	//@}
 
 	/**
@@ -2176,7 +2249,7 @@ public:
 	 * function parameter, local property, or return value, and adds it to the appropriate
 	 * list
 	 */
-	virtual void AddProperty(FUnrealPropertyDefinitionInfo& PropertyDef) override;
+	virtual void AddProperty(TSharedRef<FUnrealPropertyDefinitionInfo> PropertyDef) override;
 
 	/**
 	 * Sets the specified function export flags
@@ -2199,13 +2272,19 @@ public:
 	 */
 	FUnrealFunctionDefinitionInfo* GetSuperFunction() const;
 
+protected:
+	/**
+	 * Perform any post parsing finalization and validation
+	 */
+	virtual void PostParseFinalizeInternal(bool bCreateEngineTypes) override;
+
 private:
 
 	/** info about the function associated with this FFunctionData */
 	FFuncInfo FunctionData;
 
 	/** return value for this function */
-	FUnrealPropertyDefinitionInfo* ReturnProperty = nullptr;
+	TSharedPtr<FUnrealPropertyDefinitionInfo> ReturnProperty;
 };
 
 /**
@@ -2261,11 +2340,6 @@ public:
 
 	virtual uint32 GetHash(bool bIncludeNoExport = true) const override;
 
-	/**
-	 * Perform any post parsing finalization and validation
-	 */
-	virtual void PostParseFinalize() override;
-
 	/**
 	 * Return the Engine instance associated with the compiler instance
 	 */
@@ -2504,6 +2578,12 @@ public:
 
 	FString GetNameWithPrefix(EEnforceInterfacePrefix EnforceInterfacePrefix = EEnforceInterfacePrefix::None) const;
 
+protected:
+	/**
+	 * Perform any post parsing finalization and validation
+	 */
+	virtual void PostParseFinalizeInternal(bool bCreateEngineTypes) override;
+
 public:
 	TMap<FName, FString> MetaData;
 
@@ -2673,16 +2753,16 @@ const To& UHTCastChecked(const From& Src)
 }
 
 template <class T>
-const TArray<T*>& GetFieldsFromDef(const FUnrealStructDefinitionInfo& InStructDef);
+const TArray<TSharedRef<T>>& GetFieldsFromDef(const FUnrealStructDefinitionInfo& InStructDef);
 
 template <>
-inline const TArray<FUnrealPropertyDefinitionInfo*>& GetFieldsFromDef(const FUnrealStructDefinitionInfo& InStructDef)
+inline const TArray<TSharedRef<FUnrealPropertyDefinitionInfo>>& GetFieldsFromDef(const FUnrealStructDefinitionInfo& InStructDef)
 {
 	return InStructDef.GetProperties();
 }
 
 template <>
-inline const TArray<FUnrealFunctionDefinitionInfo*>& GetFieldsFromDef(const FUnrealStructDefinitionInfo& InStructDef)
+inline const TArray<TSharedRef<FUnrealFunctionDefinitionInfo>>& GetFieldsFromDef(const FUnrealStructDefinitionInfo& InStructDef)
 {
 	return InStructDef.GetFunctions();
 }
@@ -2697,8 +2777,8 @@ private:
 	/** The object being searched for the specified field */
 	const FUnrealStructDefinitionInfo* StructDef;
 	/** The current location in the list of fields being iterated */
-	T* const* Field;
-	T* const* End;
+	TSharedRef<T> const* Field;
+	TSharedRef<T> const* End;
 	/** Whether to include the super class or not */
 	const bool bIncludeSuper;
 
@@ -2736,29 +2816,29 @@ public:
 	inline T* operator*()
 	{
 		checkSlow(Field != End);
-		return *Field;
+		return &**Field;
 	}
 	inline const T* operator*() const
 	{
 		checkSlow(Field != End);
-		return *Field;
+		return &**Field;
 	}
 	inline T* operator->()
 	{
 		checkSlow(Field != End);
-		return *Field;
+		return &**Field;
 	}
 	inline const FUnrealStructDefinitionInfo* GetStructDef()
 	{
 		return StructDef;
 	}
 protected:
-	inline T* const* UpdateRange(const FUnrealStructDefinitionInfo* InStructDef)
+	inline TSharedRef<T> const* UpdateRange(const FUnrealStructDefinitionInfo* InStructDef)
 	{
 		if (InStructDef)
 		{
 			auto& Array = GetFieldsFromDef<T>(*InStructDef);
-			T* const* Out = Array.GetData();
+			TSharedRef<T> const * Out = Array.GetData();
 			End = Out + Array.Num();
 			return Out;
 		}
@@ -2771,7 +2851,7 @@ protected:
 
 	inline void IterateToNext()
 	{
-		T* const* CurrentField = Field;
+		TSharedRef<T> const* CurrentField = Field;
 		const FUnrealStructDefinitionInfo* CurrentStructDef = StructDef;
 
 		while (CurrentStructDef)
