Functions are no longer required to be created at parse time.

#rnx

diff --git a/Engine/Source/Programs/UnrealHeaderTool/Private/CodeGenerator.cpp b/Engine/Source/Programs/UnrealHeaderTool/Private/CodeGenerator.cpp
--- a/Engine/Source/Programs/UnrealHeaderTool/Private/CodeGenerator.cpp
+++ b/Engine/Source/Programs/UnrealHeaderTool/Private/CodeGenerator.cpp
@@ -77,7 +77,7 @@ static TArray<FString> ChangeMessages;
 static bool bWriteContents = false;
 static bool bVerifyContents = false;
 
-UClass* ProcessParsedClass(bool bClassIsAnInterface, const FString& ClassName, const FString& BaseClassName, UObject* InParent, EObjectFlags Flags);
+UClass* ProcessParsedClass(FUnrealClassDefinitionInfo& ClassDef);
 UEnum* ProcessParsedEnum(const FString& EnumName, UObject* InParent, EObjectFlags Flags);
 UScriptStruct* ProcessParsedStruct(const FString& StructName, UObject* InParent, EObjectFlags Flags);
 
@@ -90,6 +90,7 @@ bool HasIdentifierExactMatch(const TCHAR* StringBegin, const TCHAR* StringEnd, c
 FUnrealClassDefinitionInfo* GUObjectDef = nullptr;
 FUnrealClassDefinitionInfo* GUClassDef = nullptr;
 FUnrealClassDefinitionInfo* GUInterfaceDef = nullptr;
+FUnrealClassDefinitionInfo* GUPackageDef = nullptr;
 
 namespace
 {
@@ -1822,6 +1823,12 @@ static TArray<FUnrealScriptStructDefinitionInfo*> FindNoExportStructs(FUnrealStr
 	return Result;
 }
 
+bool IsDelegateFunction(FUnrealFieldDefinitionInfo& FieldDef)
+{
+	FUnrealFunctionDefinitionInfo* FunctionDef = UHTCast<FUnrealFunctionDefinitionInfo>(FieldDef);
+	return FunctionDef != nullptr && FunctionDef->IsDelegateFunction();
+}
+
 void FNativeClassHeaderGenerator::ExportGeneratedPackageInitCode(FOutputDevice& Out, const TCHAR* InDeclarations, uint32 Hash)
 {
 	UPackage* Package = PackageDef.GetPackage();
@@ -1835,8 +1842,8 @@ void FNativeClassHeaderGenerator::ExportGeneratedPackageInitCode(FOutputDevice&
 
 	Algo::Sort(Singletons, [](FUnrealFieldDefinitionInfo* A, FUnrealFieldDefinitionInfo* B)
 	{
-		bool bADel = A->IsADelegateFunction();
-		bool bBDel = B->IsADelegateFunction();
+		bool bADel = IsDelegateFunction(*A);
+		bool bBDel = IsDelegateFunction(*B);
 		if (bADel != bBDel)
 		{
 			return !bADel;
@@ -1920,13 +1927,13 @@ void FNativeClassHeaderGenerator::ExportNativeGeneratedInitCode(FOutputDevice& O
 		if (bAlreadyIncluded)
 		{
 			// In a dynamic class the same function signature may be used for a Multi- and a Single-cast delegate.
-			if (!LocalFuncDef->IsADelegateFunction() || !bIsDynamic)
+			if (!LocalFuncDef->IsDelegateFunction() || !bIsDynamic)
 			{
 				FError::Throwf(TEXT("The same function linked twice. Function: %s Class: %s"), *LocalFuncDef->GetName(), *ClassDef.GetName());
 			}
 			continue;
 		}
-		if (!LocalFuncDef->IsADelegateFunction())
+		if (!LocalFuncDef->IsDelegateFunction())
 		{
 			bAllEditorOnlyFunctions &= LocalFuncDef->HasAnyFunctionFlags(FUNC_EditorOnly);
 		}
@@ -2018,7 +2025,7 @@ void FNativeClassHeaderGenerator::ExportNativeGeneratedInitCode(FOutputDevice& O
 			{
 				const bool bIsEditorOnlyFunction = FunctionDef->HasAnyFunctionFlags(FUNC_EditorOnly);
 
-				if (!FunctionDef->IsADelegateFunction())
+				if (!FunctionDef->IsDelegateFunction())
 				{
 					ExportFunction(Out, OutReferenceGatherers, SourceFile, *FunctionDef, bIsNoExport);
 				}
@@ -6478,7 +6485,7 @@ void DefineTypes(TArray<FUnrealPackageDefinitionInfo*>& PackageDefs)
 				for (TSharedRef<FUnrealTypeDefinitionInfo>& TypeDef : SourceFile.GetDefinedClasses())
 				{
 					FUnrealClassDefinitionInfo& ClassDef = TypeDef->AsClassChecked();
-					UClass* ResultClass = ProcessParsedClass(ClassDef.IsInterface(), ClassDef.GetNameCPP(), ClassDef.GetSuperStructInfo().Name, Package, RF_Public | RF_Standalone);
+					UClass* ResultClass = ProcessParsedClass(ClassDef);
 					GTypeDefinitionInfoMap.Add(ResultClass, TypeDef);
 					ClassDef.SetObject(ResultClass);
 					AllClasses.Add(ResultClass);
@@ -6521,6 +6528,7 @@ void ResolveParents(TArray<FUnrealPackageDefinitionInfo*>& PackageDefs)
 	GUObjectDef = &GTypeDefinitionInfoMap.FindChecked<FUnrealClassDefinitionInfo>(UObject::StaticClass());
 	GUClassDef = &GTypeDefinitionInfoMap.FindChecked<FUnrealClassDefinitionInfo>(UClass::StaticClass());
 	GUInterfaceDef = &GTypeDefinitionInfoMap.FindChecked<FUnrealClassDefinitionInfo>(UInterface::StaticClass());
+	GUPackageDef = &GTypeDefinitionInfoMap.FindChecked<FUnrealClassDefinitionInfo>(UPackage::StaticClass());
 
 	for (FUnrealPackageDefinitionInfo* PackageDef : PackageDefs)
 	{
@@ -6640,11 +6648,19 @@ void ParseSourceFiles(TArray<FUnrealSourceFile*>& OrderedSourceFiles)
 
 void PrepareTypesForExport(TArray<FUnrealPackageDefinitionInfo*>& PackageDefs)
 {
+	for (FUnrealPackageDefinitionInfo* PackageDef : PackageDefs)
+	{
+		FResults::Try([PackageDef]()
+		{
+			PackageDef->CreateUObjectEngineTypes();
+		});
+	}
+		
 	for (FUnrealPackageDefinitionInfo* PackageDef : PackageDefs)
 	{
 		FResults::Try([PackageDef]()
 			{
-				PackageDef->PostParseFinalize(true);
+				PackageDef->PostParseFinalize();
 			}
 		);
 	}
@@ -6655,7 +6671,7 @@ void PrepareTypesForExport(TArray<FUnrealPackageDefinitionInfo*>& PackageDefs)
 		{ 
 			if (!TypeDef.HasSource())
 			{
-				TypeDef.PostParseFinalize(true);
+				TypeDef.PostParseFinalize();
 			}
 		});
 }
@@ -6817,7 +6833,7 @@ ECompilationResult::Type UnrealHeaderTool_Main(const FString& ModuleInfoFilename
 	// Use -WRITEREF and -VERIFYREF to detect these changes.
 	for (TObjectIterator<UClass> ClassIt; ClassIt; ++ClassIt)
 	{
-		FUnrealClassDefinitionInfo* ClassDef = new FUnrealClassDefinitionInfo(FNameLookupCPP::GetNameCPP(*ClassIt));
+		FUnrealClassDefinitionInfo* ClassDef = new FUnrealClassDefinitionInfo(*ClassIt);
 		ClassDef->SetObject(*ClassIt);
 		GTypeDefinitionInfoMap.Add(*ClassIt, MakeShareable<FUnrealTypeDefinitionInfo>(ClassDef));
 	}
@@ -6918,9 +6934,12 @@ ECompilationResult::Type UnrealHeaderTool_Main(const FString& ModuleInfoFilename
 	return FResults::GetOverallResults();
 }
 
-UClass* ProcessParsedClass(bool bClassIsAnInterface, const FString& ClassName, const FString& BaseClassName, UObject* InParent, EObjectFlags Flags)
+UClass* ProcessParsedClass(FUnrealClassDefinitionInfo& ClassDef)
 {
+	UPackage* Package = ClassDef.GetPackageDef().GetPackage();
+	const FString& ClassName = ClassDef.GetNameCPP();
 	FString ClassNameStripped = GetClassNameWithPrefixRemoved(*ClassName);
+	const FString& BaseClassName = ClassDef.GetSuperStructInfo().Name;
 
 	// All classes must start with a valid unreal prefix
 	if (!FHeaderParser::ClassNameHasValidPrefix(ClassName, ClassNameStripped))
@@ -6956,7 +6975,7 @@ UClass* ProcessParsedClass(bool bClassIsAnInterface, const FString& ClassName, c
 		FError::Throwf(TEXT("Class '%s' cannot inherit from itself"), *ClassName);
 	}
 
-	UClass* ResultClass = FEngineAPI::FindObject<UClass>(InParent, *ClassNameStripped);
+	UClass* ResultClass = FEngineAPI::FindObject<UClass>(Package, *ClassNameStripped);
 
 	// if we aren't generating headers, then we shouldn't set misaligned object, since it won't get cleared
 
@@ -6975,12 +6994,12 @@ UClass* ProcessParsedClass(bool bClassIsAnInterface, const FString& ClassName, c
 		}
 
 		// Create new class.
-		ResultClass = new(EC_InternalUseOnlyConstructor, InParent, *ClassNameStripped, Flags) UClass(FObjectInitializer(), nullptr);
+		ResultClass = new(EC_InternalUseOnlyConstructor, Package, *ClassNameStripped, RF_Public | RF_Standalone) UClass(FObjectInitializer(), nullptr);
 
 		// add CLASS_Interface flag if the class is an interface
 		// NOTE: at this pre-parsing/importing stage, we cannot know if our super class is an interface or not,
 		// we leave the validation to the main header parser
-		if (bClassIsAnInterface)
+		if (ClassDef.IsInterface())
 		{
 			ResultClass->ClassFlags |= CLASS_Interface;
 		}
@@ -6990,6 +7009,10 @@ UClass* ProcessParsedClass(bool bClassIsAnInterface, const FString& ClassName, c
 			UE_LOG(LogCompile, Log, TEXT("Imported: %s"), *ResultClass->GetFullName());
 		}
 	}
+	else
+	{
+		ClassDef.SetInternalFlags(ResultClass->GetInternalFlags());
+	}
 
 	return ResultClass;
 }
diff --git a/Engine/Source/Programs/UnrealHeaderTool/Private/HeaderParser.cpp b/Engine/Source/Programs/UnrealHeaderTool/Private/HeaderParser.cpp
--- a/Engine/Source/Programs/UnrealHeaderTool/Private/HeaderParser.cpp
+++ b/Engine/Source/Programs/UnrealHeaderTool/Private/HeaderParser.cpp
@@ -867,7 +867,7 @@ FUnrealFunctionDefinitionInfo* FHeaderParser::FindFunction
 				{
 					if (Thing)
 					{
-						FError::Throwf(TEXT("%s: expecting @todo: FProp, got %s"), Thing, /**FieldClass->GetName(),*/ *PropertyDef->GetEngineClassName());
+						FError::Throwf(TEXT("%s: expecting function or delegate, got a property"), Thing);
 					}
 					return nullptr;
 				}
@@ -904,7 +904,7 @@ FUnrealPropertyDefinitionInfo* FHeaderParser::FindProperty(const FUnrealStructDe
 				{
 					if (Thing)
 					{
-						FError::Throwf(TEXT("%s: expecting a property, got %s"), Thing, *FunctionDef->GetFunction()->GetClass()->GetName());
+						FError::Throwf(TEXT("%s: expecting a property, got a function or delegate"), Thing);
 					}
 					return nullptr;
 				}
@@ -5334,9 +5334,6 @@ FUnrealClassDefinitionInfo& FHeaderParser::ParseClassNameDeclaration(FString& De
  */
 void PostParsingClassSetup(FUnrealClassDefinitionInfo& ClassDef)
 {
-	// Cleanup after first pass.
-	FHeaderParser::ComputeFunctionParametersSize(ClassDef);
-
 	// Set all optimization ClassFlags based on property types
 	auto HasAllOptimizationClassFlags = [&ClassDef]()
 	{
@@ -5445,7 +5442,7 @@ FUnrealClassDefinitionInfo& FHeaderParser::CompileClassDeclaration()
 	StructData.SetPrologLine(PrologFinishLine);
 
 	ClassDef.MergeAndValidateClassFlags(DeclaredClassName, PrevClassFlags);
-	Class->SetInternalFlags(EInternalObjectFlags::Native);
+	ClassDef.SetInternalFlags(EInternalObjectFlags::Native);
 
 	// Class metadata
 	MetaData.Append(ClassDef.MetaData);
@@ -5633,7 +5630,7 @@ void FHeaderParser::CompileInterfaceDeclaration()
 		FError::Throwf(TEXT("Native classes cannot extend non-native classes") );
 	}
 
-	InterfaceClass->SetInternalFlags(EInternalObjectFlags::Native);
+	InterfaceClassDef->SetInternalFlags(EInternalObjectFlags::Native);
 	InterfaceClassDef->SetClassFlags(CLASS_Native);
 
 	// Process all of the interface specifiers
@@ -5947,8 +5944,6 @@ void FHeaderParser::ParseParameterList(FUnrealFunctionDefinitionInfo& FunctionDe
 
 		FUnrealPropertyDefinitionInfo& PropDef = GetVarNameAndDim(FunctionDef, Property, VariableCategory);
 
-		FunctionDef.GetFunction()->NumParms++;
-
 		if( AdvancedDisplay.CanMarkMore() && AdvancedDisplay.ShouldMarkParameter(PropDef.GetName()) )
 		{
 			PropDef.SetPropertyFlags(CPF_AdvancedDisplay);
@@ -6229,14 +6224,13 @@ FUnrealFunctionDefinitionInfo& FHeaderParser::CompileDelegateDeclaration(const T
 		FCString::Strcpy( FuncInfo.Function.Identifier, *Name );
 	}
 
-	FUnrealFunctionDefinitionInfo& DelegateSignatureFunctionDef = bIsSparse ? CreateDelegateFunction<USparseDelegateFunction>(MoveTemp(FuncInfo)) : CreateDelegateFunction<UDelegateFunction>(MoveTemp(FuncInfo));
-	UDelegateFunction* DelegateSignatureFunction = CastChecked<UDelegateFunction>(DelegateSignatureFunctionDef.GetFunction());
-	FFuncInfo& FuncDefFuncInfo = DelegateSignatureFunctionDef.GetFunctionData();
+	FUnrealFunctionDefinitionInfo& DelegateSignatureFunctionDef = bIsSparse ? CreateDelegateFunction<USparseDelegateFunction>(MoveTemp(FuncInfo), EFunctionType::SparseDelegate) : CreateDelegateFunction<UDelegateFunction>(MoveTemp(FuncInfo), EFunctionType::Delegate);
+	const FFuncInfo& FuncDefFuncInfo = DelegateSignatureFunctionDef.GetFunctionData();
 
 	// determine whether this function should be 'const'
 	if (bDeclaredConst)
 	{
-		DelegateSignatureFunction->FunctionFlags |= FUNC_Const;
+		DelegateSignatureFunctionDef.SetFunctionFlags(FUNC_Const);
 	}
 
 	if (bIsSparse)
@@ -6257,9 +6251,8 @@ FUnrealFunctionDefinitionInfo& FHeaderParser::CompileDelegateDeclaration(const T
 			FError::Throwf(TEXT("Missing Delegate Name."));
 		}
 
-		USparseDelegateFunction* SDF = CastChecked<USparseDelegateFunction>(DelegateSignatureFunction);
-		SDF->OwningClassName = *GetClassNameWithoutPrefix(FString(OwningClass.Identifier));
-		SDF->DelegateName = DelegateName.Identifier;
+		DelegateSignatureFunctionDef.SetSparseOwningClassName(*GetClassNameWithoutPrefix(FString(OwningClass.Identifier)));
+		DelegateSignatureFunctionDef.SetSparseDelegateName(DelegateName.Identifier);
 	}
 
 	DelegateSignatureFunctionDef.SetLineNumber(InputLine);
@@ -6273,9 +6266,9 @@ FUnrealFunctionDefinitionInfo& FHeaderParser::CompileDelegateDeclaration(const T
 
 		// Check the expected versus actual number of parameters
 		int32 ParamCount = UE_PTRDIFF_TO_INT32(FoundParamCount - UHTConfig.DelegateParameterCountStrings.GetData()) + 1;
-		if (DelegateSignatureFunction->NumParms != ParamCount)
+		if (DelegateSignatureFunctionDef.GetProperties().Num() != ParamCount)
 		{
-			FError::Throwf(TEXT("Expected %d parameters but found %d parameters"), ParamCount, DelegateSignatureFunction->NumParms);
+			FError::Throwf(TEXT("Expected %d parameters but found %d parameters"), ParamCount, DelegateSignatureFunctionDef.GetProperties().Num());
 		}
 	}
 	else
@@ -6292,12 +6285,10 @@ FUnrealFunctionDefinitionInfo& FHeaderParser::CompileDelegateDeclaration(const T
 	{
 		ReturnType.PropertyFlags |= CPF_Parm | CPF_OutParm | CPF_ReturnParm;
 		FUnrealPropertyDefinitionInfo& PropDef = GetVarNameAndDim(DelegateSignatureFunctionDef, ReturnType, EVariableCategory::Return);
-
-		DelegateSignatureFunction->NumParms++;
 	}
 
 	// Try parsing metadata for the function
-	ParseFieldMetaData(MetaData, *(DelegateSignatureFunction->GetName()));
+	ParseFieldMetaData(MetaData, *DelegateSignatureFunctionDef.GetName());
 
 	AddFormattedPrevCommentAsTooltipMetaData(MetaData);
 
@@ -6314,7 +6305,7 @@ FUnrealFunctionDefinitionInfo& FHeaderParser::CompileDelegateDeclaration(const T
 	while (FunctionIterator.MoveNext())
 	{
 		FUnrealFunctionDefinitionInfo* TestFuncDef = *FunctionIterator;
-		if ((TestFuncDef->GetFName() == DelegateSignatureFunctionDef.GetFName()) && (TestFuncDef->GetFunction() != DelegateSignatureFunction))
+		if (TestFuncDef != &DelegateSignatureFunctionDef && TestFuncDef->GetFName() == DelegateSignatureFunctionDef.GetFName())
 		{
 			FError::Throwf(TEXT("Can't override delegate signature function '%s'"), FuncDefFuncInfo.Function.Identifier);
 		}
@@ -6632,8 +6623,8 @@ FUnrealFunctionDefinitionInfo& FHeaderParser::CompileFunctionDeclaration()
 	}
 
 	FUnrealFunctionDefinitionInfo& FuncDef = CreateFunction(MoveTemp(FuncInfo));
-	UFunction* TopFunction = FuncDef.GetFunction();
-	FFuncInfo& FuncDefFuncInfo = FuncDef.GetFunctionData();
+	const FFuncInfo& FuncDefFuncInfo = FuncDef.GetFunctionData();
+	FUnrealFunctionDefinitionInfo* SuperFuncDef = FuncDef.GetSuperFunction();
 
 	// Get parameter list.
 	ParseParameterList(FuncDef, false, &MetaData);
@@ -6643,8 +6634,6 @@ FUnrealFunctionDefinitionInfo& FHeaderParser::CompileFunctionDeclaration()
 	{
 		ReturnType.PropertyFlags |= CPF_Parm | CPF_OutParm | CPF_ReturnParm;
 		FUnrealPropertyDefinitionInfo& PropDef = GetVarNameAndDim(FuncDef, ReturnType, EVariableCategory::Return);
-
-		TopFunction->NumParms++;
 	}
 
 	// determine if there are any outputs for this function
@@ -6676,7 +6665,7 @@ FUnrealFunctionDefinitionInfo& FHeaderParser::CompileFunctionDeclaration()
 		}
 	}
 
-	if (!bHasAnyOutputs && (FuncDefFuncInfo.FunctionFlags & (FUNC_BlueprintPure)))
+	if (!bHasAnyOutputs && FuncDef.HasAnyFunctionFlags(FUNC_BlueprintPure))
 	{
 		// This bad behavior would be treated as a warning in the Blueprint editor, so when converted assets generates these bad functions
 		// we don't want to prevent compilation:
@@ -6690,13 +6679,13 @@ FUnrealFunctionDefinitionInfo& FHeaderParser::CompileFunctionDeclaration()
 	// determine whether this function should be 'const'
 	if ( MatchIdentifier(TEXT("const"), ESearchCase::CaseSensitive) )
 	{
-		if( (TopFunction->FunctionFlags & (FUNC_Native)) == 0 )
+		if (FuncDef.HasAnyFunctionFlags(FUNC_Native))
 		{
 			// @TODO: UCREMOVAL Reconsider?
 			//FError::Throwf(TEXT("'const' may only be used for native functions"));
 		}
 
-		FuncDefFuncInfo.FunctionFlags |= FUNC_Const;
+		FuncDef.SetFunctionFlags(FUNC_Const);
 
 		// @todo: the presence of const and one or more outputs does not guarantee that there are
 		// no side effects. On GCC and clang we could use __attribure__((pure)) or __attribute__((const))
@@ -6704,14 +6693,14 @@ FUnrealFunctionDefinitionInfo& FHeaderParser::CompileFunctionDeclaration()
 		// const identifier to determine purity is not desirable. We should remove the following logic:
 
 		// If its a const BlueprintCallable function with some sort of output and is not being marked as an BlueprintPure=false function, mark it as BlueprintPure as well
-		if ( bHasAnyOutputs && ((FuncDefFuncInfo.FunctionFlags & FUNC_BlueprintCallable) != 0) && !FuncDefFuncInfo.bForceBlueprintImpure)
+		if (bHasAnyOutputs && FuncDef.HasAnyFunctionFlags(FUNC_BlueprintCallable) && !FuncDefFuncInfo.bForceBlueprintImpure)
 		{
-			FuncDefFuncInfo.FunctionFlags |= FUNC_BlueprintPure;
+			FuncDef.SetFunctionFlags(FUNC_BlueprintPure);
 		}
 	}
 
 	// Try parsing metadata for the function
-	ParseFieldMetaData(MetaData, *(TopFunction->GetName()));
+	ParseFieldMetaData(MetaData, *FuncDef.GetName());
 
 	AddFormattedPrevCommentAsTooltipMetaData(MetaData);
 
@@ -6742,52 +6731,48 @@ FUnrealFunctionDefinitionInfo& FHeaderParser::CompileFunctionDeclaration()
 	if (bFoundFinal)
 	{
 		// This is a final (prebinding, non-overridable) function
-		FuncDefFuncInfo.FunctionFlags |= FUNC_Final;
-		FuncDefFuncInfo.FunctionExportFlags |= FUNCEXPORT_Final;
+		FuncDef.SetFunctionFlags(FUNC_Final);
+		FuncDef.SetFunctionExportFlags(FUNCEXPORT_Final);
 		if (GetCurrentClass()->HasAnyClassFlags(CLASS_Interface))
 		{
 			FError::Throwf(TEXT("Interface functions cannot be declared 'final'"));
 		}
-		else if (FuncDefFuncInfo.FunctionFlags & FUNC_BlueprintEvent)
+		else if (FuncDef.HasAnyFunctionFlags(FUNC_BlueprintEvent))
 		{
 			FError::Throwf(TEXT("Blueprint events cannot be declared 'final'"));
 		}
 	}
 
-	// Make sure any new flags made it to the function
-	//@TODO: UCREMOVAL: Ideally the flags didn't get copied midway thru parsing the function declaration, and we could avoid this
-	TopFunction->FunctionFlags |= FuncDefFuncInfo.FunctionFlags;
-
 	// Make sure that the replication flags set on an overridden function match the parent function
-	if (UFunction* SuperFunc = TopFunction->GetSuperFunction())
+	if (SuperFuncDef != nullptr)
 	{
-		if ((TopFunction->FunctionFlags & FUNC_NetFuncFlags) != (SuperFunc->FunctionFlags & FUNC_NetFuncFlags))
+		if ((FuncDef.GetFunctionFlags() & FUNC_NetFuncFlags) != (SuperFuncDef->GetFunctionFlags() & FUNC_NetFuncFlags))
 		{
-			FError::Throwf(TEXT("Overridden function '%s': Cannot specify different replication flags when overriding a function."), *TopFunction->GetName());
+			FError::Throwf(TEXT("Overridden function '%s': Cannot specify different replication flags when overriding a function."), *FuncDef.GetName());
 		}
 	}
 
 	// if this function is an RPC in state scope, verify that it is an override
 	// this is required because the networking code only checks the class for RPCs when initializing network data, not any states within it
-	if ((TopFunction->FunctionFlags & FUNC_Net) && (TopFunction->GetSuperFunction() == NULL) && Cast<UClass>(TopFunction->GetOuter()) == NULL)
+	if (FuncDef.HasAnyFunctionFlags(FUNC_Net) && SuperFuncDef == nullptr && UHTCast<FUnrealClassDefinitionInfo>(FuncDef.GetOuter()) == nullptr)
 	{
-		FError::Throwf(TEXT("Function '%s': Base implementation of RPCs cannot be in a state. Add a stub outside state scope."), *TopFunction->GetName());
+		FError::Throwf(TEXT("Function '%s': Base implementation of RPCs cannot be in a state. Add a stub outside state scope."), *FuncDef.GetName());
 	}
 
-	if (TopFunction->FunctionFlags & (FUNC_BlueprintCallable | FUNC_BlueprintEvent))
+	if (FuncDef.HasAnyFunctionFlags(FUNC_BlueprintCallable | FUNC_BlueprintEvent))
 	{
 		for (TSharedRef<FUnrealPropertyDefinitionInfo> PropertyDef : FuncDef.GetProperties())
 		{
 			if (PropertyDef->IsStaticArray())
 			{
-				FError::Throwf(TEXT("Static array cannot be exposed to blueprint. Function: %s Parameter %s\n"), *TopFunction->GetName(), *PropertyDef->GetName());
+				FError::Throwf(TEXT("Static array cannot be exposed to blueprint. Function: %s Parameter %s\n"), *FuncDef.GetName(), *PropertyDef->GetName());
 			}
 
 			if (!FPropertyTraits::IsSupportedByBlueprint(*PropertyDef, false))
 			{
 				FString ExtendedCPPType;
 				FString CPPType = PropertyDef->GetCPPType(&ExtendedCPPType);
-				UE_LOG_ERROR_UHT(TEXT("Type '%s%s' is not supported by blueprint. %s.%s"), *CPPType, *ExtendedCPPType, *TopFunction->GetName(), *PropertyDef->GetName());
+				UE_LOG_ERROR_UHT(TEXT("Type '%s%s' is not supported by blueprint. %s.%s"), *CPPType, *ExtendedCPPType, *FuncDef.GetName(), *PropertyDef->GetName());
 			}
 		}
 	}
@@ -7315,58 +7300,6 @@ bool FHeaderParser::CompileStatement(TArray<FUnrealFunctionDefinitionInfo*>& Del
 	return true;
 }
 
-//
-// Compute the function parameter size and save the return offset
-//
-//@TODO: UCREMOVAL: Need to rename ComputeFunctionParametersSize to reflect the additional work it's doing
-void FHeaderParser::ComputeFunctionParametersSize(FUnrealClassDefinitionInfo& ClassDef)
-{
-	// Recurse with all child states in this class.
-	for (TSharedRef<FUnrealFunctionDefinitionInfo> FunctionDef : ClassDef.GetFunctions())
-	{
-		UFunction* ThisFunction = FunctionDef->GetFunction();
-
-		// Fix up any structs that were used as a parameter in a delegate before being defined
-		if (FunctionDef->HasAnyFunctionFlags(FUNC_Delegate))
-		{
-			for (TSharedRef<FUnrealPropertyDefinitionInfo> PropertyDef : FunctionDef->GetProperties())
-			{
-				if (PropertyDef->IsStructOrStructStaticArray())
-				{
-					FUnrealScriptStructDefinitionInfo& StructDef = UHTCastChecked<FUnrealScriptStructDefinitionInfo>(PropertyDef->GetPropertyBase().ClassDef);
-					if (StructDef.HasAnyStructFlags(STRUCT_HasInstancedReference))
-					{
-						PropertyDef->SetPropertyFlags(CPF_ContainsInstancedReference);
-					}
-				}
-			}
-			ThisFunction->StaticLink(true);
-		}
-
-		// Compute the function parameter size, propagate some flags to the outer function, and save the return offset
-		// Must be done in a second phase, as StaticLink resets various fields again!
-		ThisFunction->ParmsSize = 0;
-		for (TSharedRef<FUnrealPropertyDefinitionInfo> PropertyDef : FunctionDef->GetProperties())
-		{
-			if (PropertyDef->HasSpecificPropertyFlags(CPF_ReturnParm | CPF_OutParm, CPF_OutParm))
-			{
-				ThisFunction->FunctionFlags |= FUNC_HasOutParms;
-				FunctionDef->GetFunctionData().FunctionFlags |= FUNC_HasOutParms;
-			}
-				
-			if (PropertyDef->IsStructOrStructStaticArray())
-			{
-				FUnrealScriptStructDefinitionInfo& StructDef = UHTCastChecked<FUnrealScriptStructDefinitionInfo>(PropertyDef->GetPropertyBase().ClassDef);
-				if (StructDef.HasDefaults())
-				{
-					ThisFunction->FunctionFlags |= FUNC_HasDefaults;
-					FunctionDef->GetFunctionData().FunctionFlags |= FUNC_HasDefaults;
-				}
-			}
-		}
-	}
-}
-
 /*-----------------------------------------------------------------------------
 	Code skipping.
 -----------------------------------------------------------------------------*/
@@ -7416,20 +7349,6 @@ void FHeaderParser::SkipStatements( int32 NestCount, const TCHAR* ErrorTag  )
 	Main script compiling routine.
 -----------------------------------------------------------------------------*/
 
-//
-// Finalize any script-exposed functions in the specified class
-//
-void FHeaderParser::FinalizeScriptExposedFunctions(FUnrealClassDefinitionInfo& ClassDef)
-{
-	// Finalize all of the children introduced in this class
-	UClass* Class = ClassDef.GetClass();
-	for (TSharedRef<FUnrealFunctionDefinitionInfo> FunctionDef : ClassDef.GetFunctions())
-	{
-		UFunction* Function = FunctionDef->GetFunction();
-		Class->AddFunctionToFunctionMap(Function, Function->GetFName());
-	}
-}
-
 //
 // Parses the header associated with the specified class.
 // Returns result enumeration.
@@ -7551,9 +7470,6 @@ ECompilationResult::Type FHeaderParser::ParseHeader()
 			UClass* Class = ClassDef.GetClass();
 			PostParsingClassSetup(ClassDef);
 
-			// Finalize functions
-			FinalizeScriptExposedFunctions(ClassDef);
-
 			bNoExportClassesOnly = bNoExportClassesOnly && Class->HasAnyClassFlags(CLASS_NoExport);
 		}
 
@@ -8385,7 +8301,7 @@ TSharedRef<FUnrealTypeDefinitionInfo> FHeaderPreParser::ParseClassDeclaration(FU
 
 
 	// Create the definition
-	TSharedRef<FUnrealClassDefinitionInfo> ClassDef = MakeShareable(new FUnrealClassDefinitionInfo(SourceFile, InLineNumber, MoveTemp(ClassName), bClassIsAnInterface));
+	TSharedRef<FUnrealClassDefinitionInfo> ClassDef = MakeShareable(new FUnrealClassDefinitionInfo(SourceFile, InLineNumber, MoveTemp(ClassName), FName(ClassNameWithoutPrefixStr, FNAME_Add), bClassIsAnInterface));
 
 	// Parse the inheritance list
 	ParseInheritance(TEXT("class"), [this, &SourceFile, &ClassNameWithoutPrefixStr, &ClassDef = *ClassDef](const TCHAR* ClassName, bool bIsSuperClass)
@@ -8470,7 +8386,7 @@ TSharedRef<FUnrealTypeDefinitionInfo> FHeaderPreParser::ParseEnumDeclaration(FUn
 		FError::Throwf(TEXT("Missing enumeration name"));
 	}
 
-	TSharedRef<FUnrealEnumDefinitionInfo> EnumDef = MakeShareable(new FUnrealEnumDefinitionInfo(SourceFile, InLineNumber, EnumToken.Identifier));
+	TSharedRef<FUnrealEnumDefinitionInfo> EnumDef = MakeShareable(new FUnrealEnumDefinitionInfo(SourceFile, InLineNumber, EnumToken.Identifier, FName(EnumToken.Identifier, FNAME_Add)));
 	return EnumDef;
 }
 
@@ -8511,7 +8427,7 @@ TSharedRef<FUnrealTypeDefinitionInfo> FHeaderPreParser::ParseStructDeclaration(F
 
 
 	// Create the structure definition
-	TSharedRef<FUnrealScriptStructDefinitionInfo> StructDef = MakeShareable(new FUnrealScriptStructDefinitionInfo(SourceFile, InLineNumber, MoveTemp(StructNameInScript)));
+	TSharedRef<FUnrealScriptStructDefinitionInfo> StructDef = MakeShareable(new FUnrealScriptStructDefinitionInfo(SourceFile, InLineNumber, MoveTemp(StructNameInScript), FName(StructNameStripped, FNAME_Add)));
 
 	// Parse the inheritance list
 	ParseInheritance(TEXT("struct"), [this, &SourceFile, &StructNameStripped, &StructDef = *StructDef](const TCHAR* StructName, bool bIsSuperClass)
@@ -8989,8 +8905,7 @@ void FHeaderParser::CheckDocumentationPolicyForStruct(FUnrealStructDefinitionInf
 			TMap<FString, FName> ToolTipToFunc;
 			for (TSharedRef<FUnrealFunctionDefinitionInfo> FunctionDef : ClassDef->GetFunctions())
 			{
-				UFunction* Func = FunctionDef->GetFunction();
-				FString ToolTip = Func->GetToolTipText().ToString();
+				FString ToolTip = FunctionDef->GetToolTipText().ToString();
 				if (ToolTip.IsEmpty())
 				{
 					UE_LOG_ERROR_UHT(TEXT("Function '%s::%s' does not provide a tooltip / comment (DocumentationPolicy)."), *ClassDef->GetName(), *FunctionDef->GetName());
@@ -9170,7 +9085,7 @@ void FHeaderParser::ConditionalLogPointerUsage(EPointerMemberBehavior PointerMem
 }
 
 template <class TFunctionType>
-FUnrealFunctionDefinitionInfo& CreateFunctionImpl(FUnrealSourceFile& SourceFile, int32 InputLine, FFuncInfo&& FuncInfo, FUnrealObjectDefinitionInfo& OuterDef, FScope* CurrentScope)
+FUnrealFunctionDefinitionInfo& CreateFunctionImpl(FUnrealSourceFile& SourceFile, int32 InputLine, FFuncInfo&& FuncInfo, EFunctionType InFunctionType, FUnrealObjectDefinitionInfo& OuterDef, FScope* CurrentScope)
 {
 	// Allocate local property frame, push nesting level and verify
 	// uniqueness at this scope level.
@@ -9179,22 +9094,16 @@ FUnrealFunctionDefinitionInfo& CreateFunctionImpl(FUnrealSourceFile& SourceFile,
 		while (TypeIterator.MoveNext())
 		{
 			FUnrealFieldDefinitionInfo* Type = *TypeIterator;
-			if (Type->GetField()->GetFName() == FuncInfo.Function.Identifier)
+			if (Type->GetFName() == FuncInfo.Function.Identifier)
 			{
-				FError::Throwf(TEXT("'%s' conflicts with '%s'"), FuncInfo.Function.Identifier, *Type->GetField()->GetFullName());
+				FError::Throwf(TEXT("'%s' conflicts with '%s'"), FuncInfo.Function.Identifier, *Type->GetFullName());
 			}
 		}
 	}
 
-	TFunctionType* Function = new(EC_InternalUseOnlyConstructor, OuterDef.GetObject(), FuncInfo.Function.Identifier, RF_Public) TFunctionType(FObjectInitializer(), nullptr);
-	Function->ReturnValueOffset = MAX_uint16;
-	Function->FirstPropertyToInit = nullptr;
-	Function->FunctionFlags |= FuncInfo.FunctionFlags;
-
-	TSharedRef<FUnrealFunctionDefinitionInfo> FuncDefRef = MakeShared<FUnrealFunctionDefinitionInfo>(SourceFile, InputLine, FString(FuncInfo.Function.Identifier), OuterDef, MoveTemp(FuncInfo));
+	TSharedRef<FUnrealFunctionDefinitionInfo> FuncDefRef = MakeShared<FUnrealFunctionDefinitionInfo>(SourceFile, InputLine, FString(FuncInfo.Function.Identifier), FName(FuncInfo.Function.Identifier, FNAME_Add), OuterDef, MoveTemp(FuncInfo), InFunctionType);
 	FUnrealFunctionDefinitionInfo& FuncDef = *FuncDefRef;
-	FuncDef.SetObject(Function);
-	GTypeDefinitionInfoMap.Add(Function, MoveTemp(FuncDefRef));
+
 	FuncDef.GetFunctionData().SetFunctionNames(FuncDef);
 
 	CurrentScope->AddType(FuncDef);
@@ -9202,11 +9111,7 @@ FUnrealFunctionDefinitionInfo& CreateFunctionImpl(FUnrealSourceFile& SourceFile,
 	if (!CurrentScope->IsFileScope())
 	{
 		FUnrealStructDefinitionInfo& StructDef = ((FStructScope*)CurrentScope)->GetStructDef();
-		UStruct* Struct = StructDef.GetStruct();
-
 		StructDef.AddFunction(FuncDefRef);
-		Function->Next = Struct->Children;
-		Struct->Children = Function;
 	}
 	else
 	{
@@ -9218,12 +9123,12 @@ FUnrealFunctionDefinitionInfo& CreateFunctionImpl(FUnrealSourceFile& SourceFile,
 
 FUnrealFunctionDefinitionInfo& FHeaderParser::CreateFunction(FFuncInfo&& FuncInfo) const
 {
-	return CreateFunctionImpl<UFunction>(SourceFile, InputLine, MoveTemp(FuncInfo), GetCurrentClassDef(), GetCurrentScope());
+	return CreateFunctionImpl<UFunction>(SourceFile, InputLine, MoveTemp(FuncInfo), EFunctionType::Function, GetCurrentClassDef(), GetCurrentScope());
 }
 
 template<class T>
-FUnrealFunctionDefinitionInfo& FHeaderParser::CreateDelegateFunction(FFuncInfo&& FuncInfo) const
+FUnrealFunctionDefinitionInfo& FHeaderParser::CreateDelegateFunction(FFuncInfo&& FuncInfo, EFunctionType InFunctionType) const
 {
 	FUnrealObjectDefinitionInfo& OuterDef = IsInAClass() ? static_cast<FUnrealObjectDefinitionInfo&>(GetCurrentClassDef()) : SourceFile.GetPackageDef();
-	return CreateFunctionImpl<T>(SourceFile, InputLine, MoveTemp(FuncInfo), OuterDef, GetCurrentScope());
+	return CreateFunctionImpl<T>(SourceFile, InputLine, MoveTemp(FuncInfo), InFunctionType, OuterDef, GetCurrentScope());
 }
diff --git a/Engine/Source/Programs/UnrealHeaderTool/Private/HeaderParser.h b/Engine/Source/Programs/UnrealHeaderTool/Private/HeaderParser.h
--- a/Engine/Source/Programs/UnrealHeaderTool/Private/HeaderParser.h
+++ b/Engine/Source/Programs/UnrealHeaderTool/Private/HeaderParser.h
@@ -132,9 +132,6 @@ struct FDocumentationPolicy
 class FHeaderParser : public FBaseParser, public FContextSupplier
 {
 public:
-	// Compute the function parameter size and save the return offset
-	static void ComputeFunctionParametersSize(FUnrealClassDefinitionInfo& ClassDef);
-
 	// Performs a preliminary parse of the text in the specified buffer, pulling out:
 	//   Class name and parent class name
 	//   Is it an interface
@@ -400,7 +397,6 @@ protected:
 	 */
 	ECompilationResult::Type ParseHeader();
 	void CompileDirective();
-	void FinalizeScriptExposedFunctions(FUnrealClassDefinitionInfo& ClassDef);
 	FUnrealEnumDefinitionInfo& CompileEnum();
 	FUnrealScriptStructDefinitionInfo& CompileStructDeclaration();
 	bool CompileDeclaration(TArray<FUnrealFunctionDefinitionInfo*>& DelegatesToFixup, FToken& Token);
@@ -435,7 +431,7 @@ protected:
 	 * Create new delegate function object based on given info structure.
 	 */
 	template<typename T>
-	FUnrealFunctionDefinitionInfo& CreateDelegateFunction(FFuncInfo&& FuncInfo) const;
+	FUnrealFunctionDefinitionInfo& CreateDelegateFunction(FFuncInfo&& FuncInfo, EFunctionType InFunctionType) const;
 
 	FUnrealClassDefinitionInfo& CompileClassDeclaration();
 	FUnrealFunctionDefinitionInfo& CompileDelegateDeclaration(const TCHAR* DelegateIdentifier, EDelegateSpecifierAction::Type SpecifierAction = EDelegateSpecifierAction::DontParse);
diff --git a/Engine/Source/Programs/UnrealHeaderTool/Private/ParserHelper.cpp b/Engine/Source/Programs/UnrealHeaderTool/Private/ParserHelper.cpp
--- a/Engine/Source/Programs/UnrealHeaderTool/Private/ParserHelper.cpp
+++ b/Engine/Source/Programs/UnrealHeaderTool/Private/ParserHelper.cpp
@@ -323,8 +323,7 @@ FToken& FToken::operator=(FToken&& Other)
 
 void FFuncInfo::SetFunctionNames(FUnrealFunctionDefinitionInfo& FunctionDef)
 {
-	UFunction* FunctionReference = FunctionDef.GetFunction();
-	FString FunctionName = FunctionReference->GetName();
+	FString FunctionName = FunctionDef.GetName();
 	if (FunctionDef.HasAnyFunctionFlags(FUNC_Delegate))
 	{
 		FunctionName.LeftChopInline(FString(HEADER_GENERATED_DELEGATE_SIGNATURE_SUFFIX).Len(), false);
@@ -346,13 +345,13 @@ void FFuncInfo::SetFunctionNames(FUnrealFunctionDefinitionInfo& FunctionDef)
 		if (FunctionDef.HasAllFunctionFlags(FUNC_NetResponse))
 		{
 			// Response function implemented by programmer and called directly from thunk
-			CppImplName = FunctionReference->GetName();
+			CppImplName = FunctionDef.GetName();
 		}
 		else
 		{
 			if (CppImplName.IsEmpty())
 			{
-				CppImplName = FunctionReference->GetName() + TEXT("_Implementation");
+				CppImplName = FunctionDef.GetName() + TEXT("_Implementation");
 			}
 			else if (CppImplName == FunctionName)
 			{
@@ -361,7 +360,7 @@ void FFuncInfo::SetFunctionNames(FUnrealFunctionDefinitionInfo& FunctionDef)
 
 			if (CppValidationImplName.IsEmpty() && FunctionDef.HasAllFunctionFlags(FUNC_NetValidate))
 			{
-				CppValidationImplName = FunctionReference->GetName() + TEXT("_Validate");
+				CppValidationImplName = FunctionDef.GetName() + TEXT("_Validate");
 			}
 			else if (CppValidationImplName == FunctionName)
 			{
@@ -378,7 +377,7 @@ void FFuncInfo::SetFunctionNames(FUnrealFunctionDefinitionInfo& FunctionDef)
 	if (FunctionDef.HasAllFunctionFlags(FUNC_BlueprintEvent | FUNC_Native))
 	{
 		MarshallAndCallName = FunctionName;
-		CppImplName = FunctionReference->GetName() + TEXT("_Implementation");
+		CppImplName = FunctionDef.GetName() + TEXT("_Implementation");
 	}
 
 	if (CppImplName.IsEmpty())
diff --git a/Engine/Source/Programs/UnrealHeaderTool/Private/Scope.cpp b/Engine/Source/Programs/UnrealHeaderTool/Private/Scope.cpp
--- a/Engine/Source/Programs/UnrealHeaderTool/Private/Scope.cpp
+++ b/Engine/Source/Programs/UnrealHeaderTool/Private/Scope.cpp
@@ -19,7 +19,7 @@ FScope::FScope()
 
 void FScope::AddType(FUnrealFieldDefinitionInfo& Type)
 {
-	TypeMap.Add(Type.GetField()->GetFName(), &Type);
+	TypeMap.Add(Type.GetFName(), &Type);
 }
 
 /**
@@ -32,34 +32,31 @@ void FScope::AddType(FUnrealFieldDefinitionInfo& Type)
  */
 void DispatchType(FUnrealFieldDefinitionInfo& FieldDef, TArray<FUnrealEnumDefinitionInfo*>& Enums, TArray<FUnrealScriptStructDefinitionInfo*>& Structs, TArray<FUnrealFunctionDefinitionInfo*>& DelegateFunctions)
 {
-	UField* Type = FieldDef.GetField();
-	UClass* TypeClass = Type->GetClass();
-
-	if (TypeClass == UClass::StaticClass() || TypeClass == UStruct::StaticClass())
+	if (UHTCast<FUnrealClassDefinitionInfo>(FieldDef) != nullptr)
 	{
 		// Inner scopes.
 		FieldDef.GetScope()->SplitTypesIntoArrays(Enums, Structs, DelegateFunctions);
 	}
-	else if (TypeClass == UEnum::StaticClass())
+	else if (FUnrealEnumDefinitionInfo* EnumDef = UHTCast<FUnrealEnumDefinitionInfo>(FieldDef))
 	{
-		Enums.Add(&UHTCastChecked<FUnrealEnumDefinitionInfo>(FieldDef));
+		Enums.Add(EnumDef);
 	}
-	else if (TypeClass == UScriptStruct::StaticClass())
+	else if (FUnrealScriptStructDefinitionInfo* ScriptStructDef = UHTCast<FUnrealScriptStructDefinitionInfo>(FieldDef))
 	{
-		Structs.Add(&UHTCastChecked<FUnrealScriptStructDefinitionInfo>(FieldDef));
+		Structs.Add(ScriptStructDef);
 	}
-	else if (TypeClass == UDelegateFunction::StaticClass() || TypeClass == USparseDelegateFunction::StaticClass())
+	else if (FUnrealFunctionDefinitionInfo* FunctionDef = UHTCast<FUnrealFunctionDefinitionInfo>(FieldDef))
 	{
-		bool bAdded = false;
-		UDelegateFunction* Function = (UDelegateFunction*)Type;
-
-		if (Function->GetSuperFunction() == NULL)
+		if (FunctionDef->IsDelegateFunction())
 		{
-			DelegateFunctions.Add(&UHTCastChecked<FUnrealFunctionDefinitionInfo>(FieldDef));
-			bAdded = true;
+			bool bAdded = false;
+			if (FunctionDef->GetSuperFunction() == nullptr)
+			{
+				DelegateFunctions.Add(FunctionDef);
+				bAdded = true;
+			}
+			check(bAdded);
 		}
-
-		check(bAdded);
 	}
 }
 
@@ -80,7 +77,7 @@ FUnrealFieldDefinitionInfo* FScope::FindTypeByName(FName Name)
 		while (TypeIterator.MoveNext())
 		{
 			FUnrealFieldDefinitionInfo* Type = *TypeIterator;
-			if (Type->GetField()->GetFName() == Name)
+			if (Type->GetFName() == Name)
 			{
 				return Type;
 			}
@@ -99,7 +96,7 @@ const FUnrealFieldDefinitionInfo* FScope::FindTypeByName(FName Name) const
 		while (TypeIterator.MoveNext())
 		{
 			FUnrealFieldDefinitionInfo* Type = *TypeIterator;
-			if (Type->GetField()->GetFName() == Name)
+			if (Type->GetFName() == Name)
 			{
 				return Type;
 			}
diff --git a/Engine/Source/Programs/UnrealHeaderTool/Private/Scope.h b/Engine/Source/Programs/UnrealHeaderTool/Private/Scope.h
--- a/Engine/Source/Programs/UnrealHeaderTool/Private/Scope.h
+++ b/Engine/Source/Programs/UnrealHeaderTool/Private/Scope.h
@@ -137,23 +137,6 @@ public:
 	 */
 	virtual FName GetName() const = 0;
 
-	/**
-	 * Gets scope types filtered by type.
-	 *
-	 * @param OutArray (Output parameter) Array to fill with types from this scope.
-	 */
-	template <class TTypeFilter>
-	void GetTypes(TArray<TTypeFilter*> OutArray)
-	{
-		for (const auto& NameTypePair : TypeMap)
-		{
-			if (NameTypePair.Value->GetField()->IsA<TTypeFilter>())
-			{
-				OutArray.Add((TTypeFilter*)NameTypePair.Value);
-			}
-		}
-	}
-
 	/**
 	 * Gets scope's parent.
 	 */
diff --git a/Engine/Source/Programs/UnrealHeaderTool/Private/Specifiers/CheckedMetadataSpecifiers.h b/Engine/Source/Programs/UnrealHeaderTool/Private/Specifiers/CheckedMetadataSpecifiers.h
--- a/Engine/Source/Programs/UnrealHeaderTool/Private/Specifiers/CheckedMetadataSpecifiers.h
+++ b/Engine/Source/Programs/UnrealHeaderTool/Private/Specifiers/CheckedMetadataSpecifiers.h
@@ -4,7 +4,7 @@
 
 #include "CoreMinimal.h"
 
-enum class ECheckedMetadataSpecifier
+enum class ECheckedMetadataSpecifier : int32
 {
 	None = -1
 
diff --git a/Engine/Source/Programs/UnrealHeaderTool/Private/UnrealTypeDefinitionInfo.cpp b/Engine/Source/Programs/UnrealHeaderTool/Private/UnrealTypeDefinitionInfo.cpp
--- a/Engine/Source/Programs/UnrealHeaderTool/Private/UnrealTypeDefinitionInfo.cpp
+++ b/Engine/Source/Programs/UnrealHeaderTool/Private/UnrealTypeDefinitionInfo.cpp
@@ -15,6 +15,7 @@
 
 #include "Algo/FindSortedStringCaseInsensitive.h"
 #include "Math/UnitConversion.h"
+#include "Misc/ConfigCacheIni.h"
 #include "Misc/PackageName.h"
 #include "UObject/ErrorException.h"
 #include "UObject/Interface.h"
@@ -24,6 +25,7 @@
 extern FUnrealClassDefinitionInfo* GUObjectDef;
 extern FUnrealClassDefinitionInfo* GUClassDef;
 extern FUnrealClassDefinitionInfo* GUInterfaceDef;
+extern FUnrealClassDefinitionInfo* GUPackageDef;
 
 namespace
 {
@@ -35,24 +37,24 @@ namespace
 	/**
 	 * As part of the singleton name, collect the parent chain names
 	 */
-	void AddOuterNames(FUHTStringBuilder& Out, UObject* Outer)
+	void AddOuterNames(FUHTStringBuilder& Out, FUnrealObjectDefinitionInfo* Outer)
 	{
 		if (Outer == nullptr)
 		{
 			return;
 		}
 
-		if (Cast<UClass>(Outer) || Cast<UScriptStruct>(Outer))
+		if (UHTCast<FUnrealClassDefinitionInfo>(Outer) != nullptr || UHTCast<FUnrealScriptStructDefinitionInfo>(Outer) != nullptr)
 		{
 			// Structs can also have UPackage outer.
-			if (!Outer->IsA<UClass>() && !Outer->GetOuter()->IsA<UPackage>())
+			if (UHTCast<FUnrealClassDefinitionInfo>(Outer) == nullptr && UHTCast<FUnrealPackageDefinitionInfo>(Outer->GetOuter()) == nullptr)
 			{
 				AddOuterNames(Out, Outer->GetOuter());
 			}
 			Out.Append(TEXT("_"));
-			Out.Append(FNameLookupCPP::GetNameCPP(Cast<UStruct>(Outer)));
+			Out.Append(UHTCastChecked<FUnrealStructDefinitionInfo>(Outer).GetAlternateNameCPP());
 		}
-		else if (UPackage* Package = Cast<UPackage>(Outer))
+		else if (FUnrealPackageDefinitionInfo* PackageDef = UHTCast<FUnrealPackageDefinitionInfo>(Outer))
 		{
 			Out.Append(TEXT("_"));
 			Out.Append(FPackageName::GetShortName(Outer->GetName()));
@@ -68,15 +70,13 @@ namespace
 	/**
 	 * Generates singleton name.
 	 */
-	void GenerateSingletonName(FUHTStringBuilder& Out, UField* Item, bool bRequiresValidObject)
+	void GenerateSingletonName(FUHTStringBuilder& Out, FUnrealFieldDefinitionInfo* Item, bool bRequiresValidObject)
 	{
-		check(Item);
-
-		Out.Append(TEXT("Z_Construct_"));
-		Out.Append(FNameLookupCPP::GetNameCPP(Item->GetClass()));
+		Out.Append(TEXT("Z_Construct_U"));
+		Out.Append(Item->GetEngineClassName());
 		AddOuterNames(Out, Item);
 
-		if (UClass* ItemClass = Cast<UClass>(Item))
+		if (FUnrealClassDefinitionInfo* ItemClass = UHTCast<FUnrealClassDefinitionInfo>(Item))
 		{
 			if (!bRequiresValidObject && !ItemClass->HasAllClassFlags(CLASS_Intrinsic))
 			{
@@ -169,20 +169,82 @@ void FUHTMetaData::RemapMetaData(TMap<FName, FString>& MetaData)
 	}
 }
 
-const FString& FUHTMetaData::GetMetaData(const FName& Key) const
+FString FUHTMetaData::GetMetaDataHelper(const FName& Key, bool bAllowRemap) const
 {
 	// if not found, return a static empty string
-	static FString EmptyString;
+	const FString* Result = FindMetaDataHelper(Key);
+	if (Result == nullptr)
+	{
+		return FString();
+	}
+
+	if (bAllowRemap && Result->StartsWith(TEXT("ini:")))
+	{
+		FString ResultString = *Result;
+		if (!GConfig->GetString(GetMetaDataRemapConfigName(), *Key.ToString(), ResultString, GEngineIni))
+		{
+			// if this fails, then use what's after the ini:
+			ResultString.MidInline(4, MAX_int32, false);
+		}
+		return ResultString;
+	}
+	else
+	{
+		return *Result;
+	}
+}
+
+const FString* FUHTMetaData::FindMetaDataHelper(const FName& Key) const
+{
+	const FString* Result = nullptr;
 
-	// every key needs to be valid and meta data needs to exist
-	if (Key == NAME_None)
+	if (Key != NAME_None)
 	{
-		return EmptyString;
+		if (TMap<FName, FString>* MetaDataMap = GetMetaDataMap())
+		{
+			Result = MetaDataMap->Find(Key);
+		}
 	}
+#if UHT_ENABLE_ENGINE_TYPE_CHECKS
+	CheckFindMetaData(Key, Result);
+#endif
+	return Result;
+}
+
+void FUHTMetaData::SetMetaDataHelper(const FName& Key, const TCHAR* InValue) 
+{
+	if (TMap<FName, FString>* MetaDataMap = GetMetaDataMap())
+	{
+		MetaDataMap->Add(Key, InValue);
+	}
+	else
+	{
+		checkf(false, TEXT("If GetMetaDataMap can return nullptr, this SetMetaDataHelper must be overridden"));
+	}
+}
+
+FName FUHTMetaData::GetMetaDataKey(const TCHAR* Key, int32 NameIndex, EFindName FindName) const
+{
+	if (NameIndex != INDEX_NONE)
+	{
+		return FName(GetMetaDataIndexName(NameIndex) + TEXT(".") + Key, FindName);
+	}
+	else
+	{
+		return FName(Key, FNAME_Find);
+	}
+}
 
-	// find and return either the located string or an empty string
-	const FString* ValuePtr = FindMetaData(Key);
-	return ValuePtr ? *ValuePtr : EmptyString;
+FName FUHTMetaData::GetMetaDataKey(FName Key, int32 NameIndex, EFindName FindName) const
+{
+	if (NameIndex != INDEX_NONE)
+	{
+		return FName(GetMetaDataIndexName(NameIndex) + TEXT(".") + Key.ToString(), FindName);
+	}
+	else
+	{
+		return Key;
+	}
 }
 
 FUnrealPropertyDefinitionInfo* FUnrealTypeDefinitionInfo::AsProperty()
@@ -269,14 +331,38 @@ void FUnrealTypeDefinitionInfo::GetHashTag(FUHTStringBuilder& Out) const
 	}
 }
 
-bool FUnrealTypeDefinitionInfo::IsDynamic(const UField* Field)
+void FUnrealTypeDefinitionInfo::CreateUObjectEngineTypes()
 {
-	return Field->HasMetaData(NAME_ReplaceConverted);
+	switch (CreateUObectEngineTypesState)
+	{
+	case EFinalizeState::None:
+		CreateUObectEngineTypesState = EFinalizeState::InProgress;
+		CreateUObjectEngineTypesInternal();
+		CreateUObectEngineTypesState = EFinalizeState::Finished;
+		break;
+	case EFinalizeState::InProgress:
+		checkf(false, TEXT("Recursive call to CreateUObectEngineTypes detected"))
+		break;
+	case EFinalizeState::Finished:
+		break;
+	}
 }
 
-bool FUnrealTypeDefinitionInfo::IsDynamic(const FField* Field)
+void FUnrealTypeDefinitionInfo::PostParseFinalize()
 {
-	return Field->HasMetaData(NAME_ReplaceConverted);
+	switch (PostParseFinalizeState)
+	{
+	case EFinalizeState::None:
+		PostParseFinalizeState = EFinalizeState::InProgress;
+		PostParseFinalizeInternal();
+		PostParseFinalizeState = EFinalizeState::Finished;
+		break;
+	case EFinalizeState::InProgress:
+		checkf(false, TEXT("Recursive call to PostParseFinalize detected"))
+		break;
+	case EFinalizeState::Finished:
+		break;
+	}
 }
 
 void FUnrealTypeDefinitionInfo::ValidateMetaDataFormat(const FName InKey, const FString& InValue) const
@@ -310,7 +396,6 @@ void FUnrealTypeDefinitionInfo::ValidateMetaDataFormat(const FName InKey, ECheck
 	{
 		if (const FUnrealFunctionDefinitionInfo* FuncDef = UHTCast<FUnrealFunctionDefinitionInfo>(this))
 		{
-			UFunction* Function = FuncDef->GetFunction();
 			if (FuncDef->HasAnyFunctionFlags(FUNC_Static))
 			{
 				// Determine if it's a function library
@@ -461,19 +546,9 @@ void FUnrealTypeDefinitionInfo::ValidateMetaDataFormat(const FName InKey, ECheck
 	}
 }
 
-void FUnrealPropertyDefinitionInfo::PostParseFinalizeInternal(bool bCreateEngineTypes)
+void FUnrealPropertyDefinitionInfo::PostParseFinalizeInternal()
 {
-	// Finalize the sub property defs, but don't create the engine type
-	if (KeyPropDef.IsValid())
-	{
-		KeyPropDef->PostParseFinalize(false);
-	}
-	if (ValuePropDef.IsValid())
-	{
-		ValuePropDef->PostParseFinalize(false);
-	}
-
-	if (bCreateEngineTypes && GetPropertySafe() == nullptr)
+	if (GetPropertySafe() == nullptr)
 	{
 		FPropertyTraits::CreateEngineType(SharedThis(this));
 	}
@@ -526,6 +601,14 @@ FString FUnrealPropertyDefinitionInfo::GetPathName(FUnrealObjectDefinitionInfo*
 {
 	TStringBuilder<256> ResultString;
 	GetPathName(StopOuter, ResultString);
+#if UHT_ENABLE_ENGINE_TYPE_CHECKS
+	if (GetPropertySafe() != nullptr && (StopOuter == nullptr || StopOuter->GetObjectSafe() != nullptr))
+	{
+		TStringBuilder<256> OtherResultString;
+		GetPropertySafe()->GetPathName(StopOuter ? StopOuter->GetObjectSafe() : nullptr, OtherResultString); // Validation
+		check(FCString::Strcmp(OtherResultString.ToString(), ResultString.ToString()) == 0);
+	}
+#endif
 	return FString(FStringView(ResultString));
 }
 
@@ -552,14 +635,6 @@ void FUnrealPropertyDefinitionInfo::GetPathName(FUnrealObjectDefinitionInfo* Sto
 		ResultString << TEXT(".");
 	}
 	GetFName().AppendString(ResultString);
-#if UHT_ENABLE_ENGINE_TYPE_CHECKS
-	if (GetPropertySafe() != nullptr)
-	{
-		TStringBuilder<256> OtherResultString;
-		GetPropertySafe()->GetPathName(StopOuter ? StopOuter->GetObject() : nullptr, OtherResultString); // Validation
-		check(FCString::Strcmp(OtherResultString.ToString(), ResultString.ToString()) == 0);
-	}
-#endif
 }
 
 FString FUnrealPropertyDefinitionInfo::GetFullName() const
@@ -602,70 +677,83 @@ FString FUnrealPropertyDefinitionInfo::GetCPPTypeForwardDeclaration() const
 	return Out;
 }
 
-struct FUHTFieldDisplayNameHelper
+struct FUHTDisplayNameHelper
+{
+	static FString Get(const FUnrealPropertyDefinitionInfo& Property)
+	{
+		// The GetAuthoredNameForField only does something for user defined structures
+		//if (FUnrealStructDefinitionInfo* OwnerStruct = GetOwnerStruct(Property))
+		//{
+		//	OwnerStruct->GetStruct()->GetAuthoredNameForField(Property.GetProperty());
+		//}
+		return Property.GetName();
+	}
+
+	static FString Get(const FUnrealObjectDefinitionInfo& Object)
+	{
+		const FUnrealClassDefinitionInfo* ClassDef = UHTCast<FUnrealClassDefinitionInfo>(Object);
+		if (ClassDef && !ClassDef->HasAnyClassFlags(CLASS_Native))
+		{
+			FString Name = Object.GetName();
+			Name.RemoveFromEnd(TEXT("_C"));
+			Name.RemoveFromStart(TEXT("SKEL_"));
+			return Name;
+		}
+		return Object.GetName();
+	}
+};
+
+FUnrealObjectDefinitionInfo* FUnrealPropertyDefinitionInfo::GetOwnerObject() const
 {
-	static FUnrealObjectDefinitionInfo* GetOwnerObject(const FUnrealPropertyDefinitionInfo& Property)
+	for (FUnrealTypeDefinitionInfo* TypeDef = GetOuter(); TypeDef; TypeDef = TypeDef->GetOuter())
 	{
-		for (FUnrealTypeDefinitionInfo* TypeDef = Property.GetOuter(); TypeDef; TypeDef = TypeDef->GetOuter())
+		if (FUnrealObjectDefinitionInfo* ObjectDef = UHTCast<FUnrealObjectDefinitionInfo>(TypeDef))
 		{
-			if (FUnrealObjectDefinitionInfo* ObjectDef = UHTCast<FUnrealObjectDefinitionInfo>(TypeDef))
-			{
-				return ObjectDef;
-			}
+			return ObjectDef;
 		}
-		return nullptr;
 	}
+	return nullptr;
+}
 
-	static FUnrealStructDefinitionInfo* GetOwnerStruct(const FUnrealPropertyDefinitionInfo& Property)
+FUnrealStructDefinitionInfo* FUnrealPropertyDefinitionInfo::GetOwnerStruct() const
+{
+	for (FUnrealTypeDefinitionInfo* TypeDef = GetOuter(); TypeDef; TypeDef = TypeDef->GetOuter())
 	{
-		for (FUnrealTypeDefinitionInfo* TypeDef = Property.GetOuter(); TypeDef; TypeDef = TypeDef->GetOuter())
+		if (FUnrealStructDefinitionInfo* StructDef = UHTCast<FUnrealStructDefinitionInfo>(TypeDef))
 		{
-			if (FUnrealStructDefinitionInfo* StructDef = UHTCast<FUnrealStructDefinitionInfo>(TypeDef))
-			{
-				return StructDef;
-			}
+			return StructDef;
 		}
-		return nullptr;
 	}
+	return nullptr;
+}
 
-	static FString GetFullGroupName(const FUnrealPropertyDefinitionInfo& Property, bool bStartWithOuter)
+FString FUnrealPropertyDefinitionInfo::GetFullGroupName(bool bStartWithOuter) const
+{
+	if (bStartWithOuter)
 	{
-		if (bStartWithOuter)
+		if (FUnrealTypeDefinitionInfo* Owner = GetOuter())
 		{
-			if (FUnrealTypeDefinitionInfo* Owner = Property.GetOuter())
+			if (FUnrealObjectDefinitionInfo* ObjectOwner = UHTCast<FUnrealObjectDefinitionInfo>(Owner))
 			{
-				if (FUnrealObjectDefinitionInfo* ObjectOwner = UHTCast<FUnrealObjectDefinitionInfo>(Owner))
-				{
-					return ObjectOwner->GetPathName(&ObjectOwner->GetPackageDef());
-				}
-				else
-				{
-					FUnrealPropertyDefinitionInfo& PropertyOwner = UHTCastChecked<FUnrealPropertyDefinitionInfo>(Owner);
-					return PropertyOwner.GetPathName(&GetOwnerObject(PropertyOwner)->GetPackageDef());
-				}
+				return ObjectOwner->GetPathName(&ObjectOwner->GetPackageDef());
 			}
 			else
 			{
-				return FString();
+				FUnrealPropertyDefinitionInfo& PropertyOwner = UHTCastChecked<FUnrealPropertyDefinitionInfo>(Owner);
+				return PropertyOwner.GetPathName(&PropertyOwner.GetOwnerObject()->GetPackageDef());
 			}
 		}
 		else
 		{
-			FUnrealObjectDefinitionInfo* ObjectOuter = GetOwnerObject(Property);
-			return Property.GetPathName(ObjectOuter ? &ObjectOuter->GetPackageDef() : nullptr);
+			return FString();
 		}
 	}
-
-	static FString Get(const FUnrealPropertyDefinitionInfo& Property)
+	else
 	{
-		// The GetAuthoredNameForField only does something for user defined structures
-		//if (FUnrealStructDefinitionInfo* OwnerStruct = GetOwnerStruct(Property))
-		//{
-		//	OwnerStruct->GetStruct()->GetAuthoredNameForField(Property.GetProperty());
-		//}
-		return Property.GetName();
+		FUnrealObjectDefinitionInfo* ObjectOuter = GetOwnerObject();
+		return GetPathName(ObjectOuter ? &ObjectOuter->GetPackageDef() : nullptr);
 	}
-};
+}
 
 FText FUnrealPropertyDefinitionInfo::GetDisplayNameText() const
 {
@@ -674,7 +762,7 @@ FText FUnrealPropertyDefinitionInfo::GetDisplayNameText() const
 	static const FString Namespace = TEXT("UObjectDisplayNames");
 	static const FName NAME_DisplayName(TEXT("DisplayName"));
 
-	const FString Key = FUHTFieldDisplayNameHelper::GetFullGroupName(*this, false);
+	const FString Key = GetFullGroupName(false);
 
 	FString NativeDisplayName;
 	if (const FString* FoundMetaData = FindMetaData(NAME_DisplayName))
@@ -683,7 +771,7 @@ FText FUnrealPropertyDefinitionInfo::GetDisplayNameText() const
 	}
 	else
 	{
-		NativeDisplayName = FName::NameToDisplayString(FUHTFieldDisplayNameHelper::Get(*this), IsBooleanOrBooleanStaticArray());
+		NativeDisplayName = FName::NameToDisplayString(FUHTDisplayNameHelper::Get(*this), IsBooleanOrBooleanStaticArray());
 	}
 
 	if (!(FText::FindText(Namespace, Key, /*OUT*/LocalizedDisplayName, &NativeDisplayName)))
@@ -726,12 +814,12 @@ FText FUnrealPropertyDefinitionInfo::GetToolTipText(bool bShortTooltip) const
 	}
 
 	const FString Namespace = bFoundShortTooltip ? TEXT("UObjectShortTooltips") : TEXT("UObjectToolTips");
-	const FString Key = FUHTFieldDisplayNameHelper::GetFullGroupName(*this, false);
+	const FString Key = GetFullGroupName(false);
 	if (!FText::FindText(Namespace, Key, /*OUT*/LocalizedToolTip, &NativeToolTip))
 	{
 		if (NativeToolTip.IsEmpty())
 		{
-			NativeToolTip = FName::NameToDisplayString(FUHTFieldDisplayNameHelper::Get(*this), IsBooleanOrBooleanStaticArray());
+			NativeToolTip = FName::NameToDisplayString(FUHTDisplayNameHelper::Get(*this), IsBooleanOrBooleanStaticArray());
 		}
 		else
 		{
@@ -784,18 +872,111 @@ FUnrealPackageDefinitionInfo& FUnrealObjectDefinitionInfo::GetPackageDef() const
 	return GTypeDefinitionInfoMap.FindChecked<FUnrealPackageDefinitionInfo>(GetObject()->GetPackage());
 }
 
+FUnrealObjectDefinitionInfo::FUnrealObjectDefinitionInfo(UObject* Object)
+	: FUnrealTypeDefinitionInfo(Object->GetName())
+	, Name(Object->GetFName())
+	, InternalObjectFlags(Object->GetInternalFlags())
+{
+}
+
+FString FUnrealObjectDefinitionInfo::GetFullName() const
+{	
+	FString FullName = GetEngineClassName();
+	FullName += TEXT(" ");
+	FullName += GetPathName();
+#if UHT_ENABLE_ENGINE_TYPE_CHECKS
+	if (GetObjectSafe() != nullptr)
+	{
+		check(FullName == GetObjectSafe()->GetFullName()); // Validation
+	}
+#endif
+	return FullName;
+}
+
 FString FUnrealObjectDefinitionInfo::GetPathName(FUnrealObjectDefinitionInfo* StopOuter) const
 {
-	return GetObject()->GetPathName(StopOuter ? StopOuter->GetObject() : nullptr);
+	TStringBuilder<256> ResultBuilder;
+	GetPathName(StopOuter, ResultBuilder);
+#if UHT_ENABLE_ENGINE_TYPE_CHECKS
+	if (GetObjectSafe() != nullptr && (StopOuter == nullptr || StopOuter->GetObjectSafe() != nullptr))
+	{
+		TStringBuilder<256> OtherResultString;
+		GetObjectSafe()->GetPathName(StopOuter ? StopOuter->GetObjectSafe() : nullptr, OtherResultString); // Validation
+		check(FCString::Strcmp(OtherResultString.ToString(), ResultBuilder.ToString()) == 0);
+	}
+#endif
+	return FString(FStringView(ResultBuilder));
 }
 
 void FUnrealObjectDefinitionInfo::GetPathName(FUnrealObjectDefinitionInfo* StopOuter, FStringBuilderBase& ResultString) const
 {
-	return GetObject()->GetPathName(StopOuter ? StopOuter->GetObject() : nullptr, ResultString);
+	if (this != StopOuter && this != NULL)
+	{
+		FUnrealObjectDefinitionInfo* ObjOuter = GetOuter();
+		if (ObjOuter && ObjOuter != StopOuter)
+		{
+			ObjOuter->GetPathName(StopOuter, ResultString);
+
+			// SUBOBJECT_DELIMITER_CHAR is used to indicate that this object's outer is not a UPackage
+			if (UHTCast<FUnrealPackageDefinitionInfo>(ObjOuter) == nullptr
+				&& UHTCast<FUnrealPackageDefinitionInfo>(ObjOuter->GetOuter()) != nullptr)
+			{
+				ResultString << SUBOBJECT_DELIMITER_CHAR;
+			}
+			else
+			{
+				ResultString << TEXT('.');
+			}
+		}
+		GetFName().AppendString(ResultString);
+	}
+	else
+	{
+		ResultString << TEXT("None");
+	}
+}
+
+FString FUnrealObjectDefinitionInfo::GetFullGroupName(bool bStartWithOuter) const
+{
+	const FUnrealObjectDefinitionInfo* Obj = bStartWithOuter ? GetOuter() : this;
+	return Obj ? Obj->GetPathName(&GetPackageDef()) : TEXT("");
+}
+
+void FUnrealObjectDefinitionInfo::SetMetaDataHelper(const FName& Key, const TCHAR* InValue) 
+{
+	// If we don't have an existing meta data map, then we must be using the meta data off
+	// the object and no values have been set yet.  In that case just set the whole block
+	TMap<FName, FString>* MetaDataMap = GetMetaDataMap();
+	if (MetaDataMap == nullptr)
+	{
+		UMetaData* MetaData = GetUObjectMetaData();
+		MetaData->SetValue(GetObject(), Key, InValue); // This is valid
+	}
+
+	// Otherwise we are either using the objects meta data and it already exists, or
+	// we are using a UHT private meta data block.  If we have both, then first
+	// set the object followed by the private copy.
+	else
+	{
+		TMap<FName, FString>* UObjectMetaDataMap = GetUObjectMetaDataMap();
+		if (MetaDataMap != UObjectMetaDataMap)
+		{
+			if (UObjectMetaDataMap != nullptr)
+			{
+				UObjectMetaDataMap->Add(Key, InValue);
+			}
+			else if (UMetaData* MetaData = GetUObjectMetaData())
+			{
+				MetaData->SetValue(GetObject(), Key, InValue); // This is valid
+			}
+		}
+
+		MetaDataMap->Add(Key, InValue);
+	}
 }
 
 FUnrealPackageDefinitionInfo::FUnrealPackageDefinitionInfo(const FManifestModule& InModule, UPackage* InPackage)
-	: FUnrealObjectDefinitionInfo(FString())
+	: FUnrealObjectDefinitionInfo(InPackage)
 	, Module(InModule)
 	, ShortUpperName(FPackageName::GetShortName(InPackage).ToUpper())
 	, API(FString::Printf(TEXT("%s_API "), *ShortUpperName))
@@ -803,7 +984,18 @@ FUnrealPackageDefinitionInfo::FUnrealPackageDefinitionInfo(const FManifestModule
 	SetObject(InPackage);
 }
 
-void FUnrealPackageDefinitionInfo::PostParseFinalizeInternal(bool bCreateEngineTypes)
+void FUnrealPackageDefinitionInfo::CreateUObjectEngineTypesInternal()
+{
+	for (TSharedRef<FUnrealSourceFile>& LocalSourceFile : GetAllSourceFiles())
+	{
+		for (TSharedRef<FUnrealTypeDefinitionInfo>& TypeDef : LocalSourceFile->GetDefinedTypes())
+		{
+			TypeDef->CreateUObjectEngineTypes();
+		}
+	}
+}
+
+void FUnrealPackageDefinitionInfo::PostParseFinalizeInternal()
 {
 	UPackage* Package = GetPackage();
 
@@ -818,7 +1010,7 @@ void FUnrealPackageDefinitionInfo::PostParseFinalizeInternal(bool bCreateEngineT
 	{
 		for (TSharedRef<FUnrealTypeDefinitionInfo>& TypeDef : LocalSourceFile->GetDefinedTypes())
 		{
-			TypeDef->PostParseFinalize(bCreateEngineTypes);
+			TypeDef->PostParseFinalize();
 		}
 	}
 }
@@ -831,34 +1023,82 @@ void FUnrealPackageDefinitionInfo::AddCrossModuleReference(TSet<FString>* Unique
 	}
 }
 
-void FUnrealFieldDefinitionInfo::PostParseFinalizeInternal(bool bCreateEngineTypes)
+void FUnrealFieldDefinitionInfo::PostParseFinalizeInternal()
 {
-	const TCHAR* TypeStr = GetSimplifiedTypeClass();
-	UField* Field = GetField();
-	FString PackageShortName = FPackageName::GetShortName(Field->GetOutermost()).ToUpper();
+	const FString& ClassName = GetEngineClassName(true);
+	const FString& PackageShortName = GetPackageDef().GetShortUpperName();
 
 	FUHTStringBuilder Out;
-	GenerateSingletonName(Out, Field, false);
-	ExternDecl[0].Appendf(TEXT("\t%s_API %s* %s;\r\n"), *PackageShortName, TypeStr, *Out);
+	GenerateSingletonName(Out, this, false);
+	ExternDecl[0].Appendf(TEXT("\t%s_API U%s* %s;\r\n"), *PackageShortName, *ClassName, *Out);
 	SingletonName[0] = Out;
 	SingletonNameChopped[0] = SingletonName[0].LeftChop(2);
 
 	Out.Reset();
-	GenerateSingletonName(Out, Field, true);
-	ExternDecl[1].Appendf(TEXT("\t%s_API %s* %s;\r\n"), *PackageShortName, TypeStr, *Out);
+	GenerateSingletonName(Out, this, true);
+	ExternDecl[1].Appendf(TEXT("\t%s_API U%s* %s;\r\n"), *PackageShortName, *ClassName, *Out);
 	SingletonName[1] = Out;
 	SingletonNameChopped[1] = SingletonName[1].LeftChop(2);
 
 	TypePackageName = GetTypePackageNameHelper(*this);
 }
 
+FText FUnrealFieldDefinitionInfo::GetToolTipText(bool bShortTooltip) const
+{
+	bool bFoundShortTooltip = false;
+	static const FName NAME_Tooltip(TEXT("Tooltip"));
+	static const FName NAME_ShortTooltip(TEXT("ShortTooltip"));
+	FText LocalizedToolTip;
+	FString NativeToolTip;
+
+	if (bShortTooltip)
+	{
+		NativeToolTip = GetMetaData(NAME_ShortTooltip);
+		if (NativeToolTip.IsEmpty())
+		{
+			NativeToolTip = GetMetaData(NAME_Tooltip);
+		}
+		else
+		{
+			bFoundShortTooltip = true;
+		}
+	}
+	else
+	{
+		NativeToolTip = GetMetaData(NAME_Tooltip);
+	}
+
+	const FString Namespace = bFoundShortTooltip ? TEXT("UObjectShortTooltips") : TEXT("UObjectToolTips");
+	const FString Key = GetFullGroupName(false);
+	if (!FText::FindText(Namespace, Key, /*OUT*/LocalizedToolTip, &NativeToolTip))
+	{
+		if (NativeToolTip.IsEmpty())
+		{
+			NativeToolTip = FName::NameToDisplayString(FUHTDisplayNameHelper::Get(*this), false);
+		}
+		else if (!bShortTooltip && IsNative())
+		{
+			UField::FormatNativeToolTip(NativeToolTip, true);
+		}
+		LocalizedToolTip = FText::FromString(NativeToolTip);
+	}
+
+#if UHT_ENABLE_ENGINE_TYPE_CHECKS
+	if (GetFieldSafe())
+	{
+		check(LocalizedToolTip.ToString() == GetFieldSafe()->GetToolTipText(bShortTooltip).ToString());
+	}
+#endif
+	return LocalizedToolTip;
+}
+
 void FUnrealFieldDefinitionInfo::AddCrossModuleReference(TSet<FString>* UniqueCrossModuleReferences, bool bRequiresValidObject) const
 {
 	// We don't need to export UFunction externs, though we may need the externs for UDelegateFunctions
 	if (UniqueCrossModuleReferences)
 	{
-		UField* Field = GetField();
-		if (!Field->IsA<UFunction>() || IsADelegateFunction())
+		const FUnrealFunctionDefinitionInfo* FunctionDef = UHTCast<FUnrealFunctionDefinitionInfo>(this);
+		if (FunctionDef == nullptr || FunctionDef->IsDelegateFunction())
 		{
 			UniqueCrossModuleReferences->Add(GetExternDecl(bRequiresValidObject));
 		}
@@ -870,32 +1110,42 @@ void FUnrealFieldDefinitionInfo::AddMetaData(TMap<FName, FString>&& InMetaData)
 	// only add if we have some!
 	if (InMetaData.Num())
 	{
-		UField* Field = GetField();
-		check(Field);
 
-		// get (or create) a metadata object for this package
-		UMetaData* MetaData = Field->GetOutermost()->GetMetaData();
-		TMap<FName, FString>* ExistingMetaData = MetaData->GetMapForObject(Field);
-		if (ExistingMetaData && ExistingMetaData->Num())
+		// If we don't have an existing meta data map, then we must be using the meta data off
+		// the object and no values have been set yet.  In that case just set the whole block
+		TMap<FName, FString>* MetaDataMap = GetMetaDataMap();
+		if (MetaDataMap == nullptr)
 		{
-			// Merge the existing metadata
-			TMap<FName, FString> MergedMetaData;
-			MergedMetaData.Reserve(InMetaData.Num() + ExistingMetaData->Num());
-			MergedMetaData.Append(*ExistingMetaData);
-			MergedMetaData.Append(InMetaData);
-			MetaData->SetObjectValues(Field, MoveTemp(MergedMetaData));
+			UMetaData* MetaData = GetUObjectMetaData();
+			MetaData->SetObjectValues(GetObject(), MoveTemp(InMetaData)); // This is valid
 		}
+
+		// Otherwise we are either using the objects meta data and it already exists, or
+		// we are using a UHT private meta data block.  If we have both, then first
+		// set the object followed by the private copy.
 		else
 		{
-			// set the metadata for this field
-			MetaData->SetObjectValues(Field, MoveTemp(InMetaData));
+			TMap<FName, FString>* UObjectMetaDataMap = GetUObjectMetaDataMap();
+			if (MetaDataMap != UObjectMetaDataMap)
+			{
+				if (UObjectMetaDataMap != nullptr)
+				{
+					UObjectMetaDataMap->Append(InMetaData);
+				}
+				else if (UMetaData* MetaData = GetUObjectMetaData())
+				{
+					MetaData->SetObjectValues(GetObject(), InMetaData); // This is valid
+				}
+			}
+
+			MetaDataMap->Append(MoveTemp(InMetaData));
 		}
 	}
 }
 
 bool FUnrealFieldDefinitionInfo::IsDynamic() const
 {
-	return GetField()->HasMetaData(NAME_ReplaceConverted);
+	return HasMetaData(NAME_ReplaceConverted);
 }
 
 bool FUnrealFieldDefinitionInfo::IsOwnedByDynamicType() const
@@ -922,8 +1172,8 @@ FUnrealClassDefinitionInfo* FUnrealFieldDefinitionInfo::GetOwnerClass() const
 	return nullptr;
 }
 
-FUnrealEnumDefinitionInfo::FUnrealEnumDefinitionInfo(FUnrealSourceFile& InSourceFile, int32 InLineNumber, FString&& InNameCPP)
-	: FUnrealFieldDefinitionInfo(InSourceFile, InLineNumber, MoveTemp(InNameCPP), InSourceFile.GetPackageDef())
+FUnrealEnumDefinitionInfo::FUnrealEnumDefinitionInfo(FUnrealSourceFile& InSourceFile, int32 InLineNumber, FString&& InNameCPP, FName InName)
+	: FUnrealFieldDefinitionInfo(InSourceFile, InLineNumber, MoveTemp(InNameCPP), InName, InSourceFile.GetPackageDef())
 { }
 
 void FUnrealStructDefinitionInfo::AddProperty(TSharedRef<FUnrealPropertyDefinitionInfo> PropertyDef)
@@ -949,31 +1199,54 @@ void FUnrealStructDefinitionInfo::AddProperty(TSharedRef<FUnrealPropertyDefiniti
 	}
 }
 
-void FUnrealStructDefinitionInfo::PostParseFinalizeInternal(bool bCreateEngineTypes)
+void FUnrealStructDefinitionInfo::CreateUObjectEngineTypesInternal()
 {
-	FUnrealFieldDefinitionInfo::PostParseFinalizeInternal(bCreateEngineTypes);
+	FUnrealFieldDefinitionInfo::CreateUObjectEngineTypesInternal();
 
 	if (SuperStructInfo.Struct)
 	{
-		SuperStructInfo.Struct->PostParseFinalize(bCreateEngineTypes);
+		SuperStructInfo.Struct->CreateUObjectEngineTypes();
 	}
 
 	for (const FBaseStructInfo& Info : BaseStructInfo)
 	{
 		if (Info.Struct)
 		{
-			Info.Struct->PostParseFinalize(bCreateEngineTypes);
+			Info.Struct->CreateUObjectEngineTypes();
 		}
 	}
 
-	for (TSharedRef<FUnrealPropertyDefinitionInfo> PropertyDef : Properties)
+	for (TSharedRef<FUnrealFunctionDefinitionInfo> FunctionDef : Functions)
 	{
-		PropertyDef->PostParseFinalize(bCreateEngineTypes);
+		FunctionDef->CreateUObjectEngineTypes();
+	}
+}
+
+void FUnrealStructDefinitionInfo::PostParseFinalizeInternal()
+{
+	FUnrealFieldDefinitionInfo::PostParseFinalizeInternal();
+
+	if (SuperStructInfo.Struct)
+	{
+		SuperStructInfo.Struct->PostParseFinalize();
+	}
+
+	for (const FBaseStructInfo& Info : BaseStructInfo)
+	{
+		if (Info.Struct)
+		{
+			Info.Struct->PostParseFinalize();
+		}
 	}
 
 	for (TSharedRef<FUnrealFunctionDefinitionInfo> FunctionDef : Functions)
 	{
-		FunctionDef->PostParseFinalize(bCreateEngineTypes);
+		FunctionDef->PostParseFinalize();
+	}
+
+	for (TSharedRef<FUnrealPropertyDefinitionInfo> PropertyDef : Properties)
+	{
+		PropertyDef->PostParseFinalize();
 	}
 
 	GetStruct()->Bind();
@@ -1012,8 +1285,8 @@ void FUnrealStructDefinitionInfo::AddFunction(TSharedRef<FUnrealFunctionDefiniti
 	}
 }
 
-FUnrealScriptStructDefinitionInfo::FUnrealScriptStructDefinitionInfo(FUnrealSourceFile& InSourceFile, int32 InLineNumber, FString&& InNameCPP)
-	: FUnrealStructDefinitionInfo(InSourceFile, InLineNumber, MoveTemp(InNameCPP), InSourceFile.GetPackageDef())
+FUnrealScriptStructDefinitionInfo::FUnrealScriptStructDefinitionInfo(FUnrealSourceFile& InSourceFile, int32 InLineNumber, FString&& InNameCPP, FName InName)
+	: FUnrealStructDefinitionInfo(InSourceFile, InLineNumber, MoveTemp(InNameCPP), InName, InSourceFile.GetPackageDef())
 { }
 
 uint32 FUnrealScriptStructDefinitionInfo::GetHash(bool bIncludeNoExport) const
@@ -1054,8 +1327,8 @@ void FUnrealStructDefinitionInfo::SetObject(UObject* InObject)
 	}
 }
 
-FUnrealClassDefinitionInfo::FUnrealClassDefinitionInfo(FUnrealSourceFile& InSourceFile, int32 InLineNumber, FString&& InNameCPP, bool bInIsInterface)
-	: FUnrealStructDefinitionInfo(InSourceFile, InLineNumber, MoveTemp(InNameCPP), InSourceFile.GetPackageDef())
+FUnrealClassDefinitionInfo::FUnrealClassDefinitionInfo(FUnrealSourceFile& InSourceFile, int32 InLineNumber, FString&& InNameCPP, FName InName, bool bInIsInterface)
+	: FUnrealStructDefinitionInfo(InSourceFile, InLineNumber, MoveTemp(InNameCPP), InName, InSourceFile.GetPackageDef())
 	, bIsInterface(bInIsInterface)
 {
 	if (bInIsInterface)
@@ -1155,9 +1428,11 @@ uint32 FUnrealClassDefinitionInfo::GetHash(bool bIncludeNoExport) const
 }
 
 
-void FUnrealClassDefinitionInfo::PostParseFinalizeInternal(bool bCreateEngineTypes)
+void FUnrealClassDefinitionInfo::PostParseFinalizeInternal()
 {
-	FUnrealStructDefinitionInfo::PostParseFinalizeInternal(bCreateEngineTypes);
+	UClass* Class = GetClass();
+
+	FUnrealStructDefinitionInfo::PostParseFinalizeInternal();
 
 	if (IsInterface() && GetStructMetaData().ParsedInterface == EParsedInterface::ParsedUInterface)
 	{
@@ -1168,6 +1443,13 @@ void FUnrealClassDefinitionInfo::PostParseFinalizeInternal(bool bCreateEngineTyp
 
 	FHeaderParser::CheckSparseClassData(*this);
 
+	// Finalize all of the children introduced in this class
+	for (TSharedRef<FUnrealFunctionDefinitionInfo> FunctionDef : GetFunctions())
+	{
+		UFunction* Function = FunctionDef->GetFunction();
+		Class->AddFunctionToFunctionMap(Function, Function->GetFName());
+	}
+
 	if (!HasAnyClassFlags(CLASS_Native))
 	{
 		//Class->UnMark(EObjectMark(OBJECTMARK_TagImp | OBJECTMARK_TagExp));
@@ -1182,7 +1464,7 @@ void FUnrealClassDefinitionInfo::PostParseFinalizeInternal(bool bCreateEngineTyp
 	// This needs to be done outside of parallel blocks because it will modify UClass memory.
 	// Later calls to SetUpUhtReplicationData inside parallel blocks should be fine, because
 	// they will see the memory has already been set up, and just return the parent pointer.
-	GetClass()->SetUpUhtReplicationData();
+	Class->SetUpUhtReplicationData();
 }
 
 void FUnrealClassDefinitionInfo::ParseClassProperties(TArray<FPropertySpecifier>&& InClassSpecifiers, const FString& InRequiredAPIMacroIfPresent)
@@ -1617,7 +1899,7 @@ void FUnrealClassDefinitionInfo::MergeAndValidateClassFlags(const FString& Decla
 			FError::Throwf(TEXT("'abstract': NoExport class missing abstract keyword from class declaration (must change C++ version first)"));
 			Class->ClassFlags |= CLASS_Abstract;
 		}
-		else if (Class->IsNative())
+		else if (IsNative())
 		{
 			FError::Throwf(TEXT("'abstract': missing abstract keyword from class declaration - class will no longer be exported as abstract"));
 		}
@@ -1786,6 +2068,11 @@ FString FUnrealClassDefinitionInfo::GetNameWithPrefix(EEnforceInterfacePrefix En
 	return FString::Printf(TEXT("%s%s"), Prefix, *GetName());
 }
 
+FUnrealClassDefinitionInfo* FUnrealClassDefinitionInfo::GetSuperClass() const
+{
+	return UHTCast<FUnrealClassDefinitionInfo>(GetSuperStruct());
+}
+
 void FUnrealFunctionDefinitionInfo::AddProperty(TSharedRef<FUnrealPropertyDefinitionInfo> PropertyDef)
 {
 	check(PropertyDef->HasAnyPropertyFlags(CPF_Parm));
@@ -1803,10 +2090,96 @@ FUnrealFunctionDefinitionInfo* FUnrealFunctionDefinitionInfo::GetSuperFunction()
 	return UHTCast<FUnrealFunctionDefinitionInfo>(GetSuperStruct());
 }
 
-void FUnrealFunctionDefinitionInfo::PostParseFinalizeInternal(bool bCreateEngineTypes)
+void FUnrealFunctionDefinitionInfo::CreateUObjectEngineTypesInternal()
 {
-	FUnrealStructDefinitionInfo::PostParseFinalizeInternal(bCreateEngineTypes);
+	// Invoke the base class creatation
+	FUnrealStructDefinitionInfo::CreateUObjectEngineTypesInternal();
 
-	GetFunction()->Bind();
+	// We have to precreate the function prior to invoking the parent finalize
+	UFunction* Function = nullptr;
+	switch (FunctionType)
+	{
+	case EFunctionType::Function:
+		Function = new(EC_InternalUseOnlyConstructor, GetOuter()->GetObject(), *GetNameCPP(), RF_Public) UFunction(FObjectInitializer(), nullptr);
+		break;
+	case EFunctionType::Delegate:
+		Function = new(EC_InternalUseOnlyConstructor, GetOuter()->GetObject(), *GetNameCPP(), RF_Public) UDelegateFunction(FObjectInitializer(), nullptr);
+		break;
+	case EFunctionType::SparseDelegate:
+	{
+		USparseDelegateFunction* USPF = new(EC_InternalUseOnlyConstructor, GetOuter()->GetObject(), *GetNameCPP(), RF_Public) USparseDelegateFunction(FObjectInitializer(), nullptr);
+		USPF->OwningClassName = SparseOwningClassName;
+		USPF->DelegateName = SparseDelegateName;
+		Function = USPF;
+		break;
+	}
+	}
+	check(Function);
+
+	Function->ReturnValueOffset = MAX_uint16;
+	Function->FirstPropertyToInit = nullptr;
+	Function->FunctionFlags |= FunctionData.FunctionFlags;
+
+	SetObject(Function);
+	GTypeDefinitionInfoMap.Add(Function, SharedThis(this));
+
+	if (FUnrealStructDefinitionInfo* StructDef = UHTCast<FUnrealStructDefinitionInfo>(GetOuter()))
+	{
+		UStruct* Struct = StructDef->GetStruct();
+		Function->Next = Struct->Children;
+		Struct->Children = Function;
+	}
+	Function->NumParms = uint8(GetProperties().Num());
+	Function->Bind();
 }
 
+void FUnrealFunctionDefinitionInfo::PostParseFinalizeInternal()
+{
+	// Invoke the base class finalization
+	FUnrealStructDefinitionInfo::PostParseFinalizeInternal();
+
+	UFunction* Function = GetFunction();
+
+	// The following code is only performed on functions in a class.  
+	if (UHTCast<FUnrealClassDefinitionInfo>(GetOuter()) != nullptr)
+	{
+
+		// Fix up any structs that were used as a parameter in a delegate before being defined
+		if (HasAnyFunctionFlags(FUNC_Delegate))
+		{
+			for (TSharedRef<FUnrealPropertyDefinitionInfo> PropertyDef : GetProperties())
+			{
+				if (PropertyDef->IsStructOrStructStaticArray())
+				{
+					FUnrealScriptStructDefinitionInfo& StructDef = UHTCastChecked<FUnrealScriptStructDefinitionInfo>(PropertyDef->GetPropertyBase().ClassDef);
+					if (StructDef.HasAnyStructFlags(STRUCT_HasInstancedReference))
+					{
+						PropertyDef->SetPropertyFlags(CPF_ContainsInstancedReference);
+					}
+				}
+			}
+		}
+
+		Function->StaticLink(true);
+
+		// Compute the function parameter size, propagate some flags to the outer function, and save the return offset
+		// Must be done in a second phase, as StaticLink resets various fields again!
+		Function->ParmsSize = 0;
+		for (TSharedRef<FUnrealPropertyDefinitionInfo> PropertyDef : GetProperties())
+		{
+			if (PropertyDef->HasSpecificPropertyFlags(CPF_ReturnParm | CPF_OutParm, CPF_OutParm))
+			{
+				SetFunctionFlags(FUNC_HasOutParms);
+			}
+
+			if (PropertyDef->IsStructOrStructStaticArray())
+			{
+				FUnrealScriptStructDefinitionInfo& StructDef = UHTCastChecked<FUnrealScriptStructDefinitionInfo>(PropertyDef->GetPropertyBase().ClassDef);
+				if (StructDef.HasDefaults())
+				{
+					SetFunctionFlags(FUNC_HasDefaults);
+				}
+			}
+		}
+	}
+}
diff --git a/Engine/Source/Programs/UnrealHeaderTool/Private/UnrealTypeDefinitionInfo.h b/Engine/Source/Programs/UnrealHeaderTool/Private/UnrealTypeDefinitionInfo.h
--- a/Engine/Source/Programs/UnrealHeaderTool/Private/UnrealTypeDefinitionInfo.h
+++ b/Engine/Source/Programs/UnrealHeaderTool/Private/UnrealTypeDefinitionInfo.h
@@ -76,6 +76,13 @@ enum class EVariableCategory
 	Member
 };
 
+enum class EFunctionType
+{
+	Function,
+	Delegate,
+	SparseDelegate,
+};
+
 /**
  * UHT Specific implementation of meta data support
  */
@@ -102,54 +109,96 @@ public:
 	}
 
 	/**
-	* Determines if the property has any metadata associated with the key
-	*
-	* @param Key The key to lookup in the metadata
-	* @return true if there is a (possibly blank) value associated with this key
-	*/
-	bool HasMetaData(const TCHAR* Key) const 
+	 * Determines if the property has any metadata associated with the key
+	 *
+	 * @param Key The key to lookup in the metadata
+	 * @param NameIndex if specified, will search for metadata linked to a specified value (in an enum); otherwise, searches for metadata for the object itself
+	 * @return true if there is a (possibly blank) value associated with this key
+	 */
+	bool HasMetaData(const TCHAR* Key, int32 NameIndex = INDEX_NONE) const
 	{
-		return FindMetaData(Key) != nullptr;
+		return FindMetaData(Key, NameIndex) != nullptr;
 	}
-	bool HasMetaData(const FName& Key) const 
+	bool HasMetaData(const FName& Key, int32 NameIndex = INDEX_NONE) const
 	{
-		return FindMetaData(Key) != nullptr;
+		return FindMetaData(Key, NameIndex) != nullptr;
 	}
 
 	/**
-	* Find the metadata value associated with the key
-	*
-	* @param Key The key to lookup in the metadata
-	* @return The value associated with the key if it exists, null otherwise
-	*/
-	const FString* FindMetaData(const TCHAR * Key) const
+	 * Find the metadata value associated with the key
+	 *
+	 * @param Key The key to lookup in the metadata
+	 * @param NameIndex if specified, will search for metadata linked to a specified value (in an enum); otherwise, searches for metadata for the object itself
+	 * @return The value associated with the key if it exists, null otherwise
+	 */
+	const FString* FindMetaData(const TCHAR* Key, int32 NameIndex = INDEX_NONE) const
 	{
-		return FindMetaData(FName(Key, FNAME_Find));
+		TMap<FName, FString>* MetaDataMap = GetMetaDataMap();
+		return MetaDataMap ? MetaDataMap->Find(GetMetaDataKey(Key, NameIndex, FNAME_Find)) : nullptr;
 	}
-	virtual const FString* FindMetaData(const FName & Key) const
+	const FString* FindMetaData(const FName& Key, int32 NameIndex = INDEX_NONE) const
 	{
 		TMap<FName, FString>* MetaDataMap = GetMetaDataMap();
-		return MetaDataMap ? MetaDataMap->Find(Key) : nullptr;
+		return MetaDataMap ? MetaDataMap->Find(GetMetaDataKey(Key, NameIndex, FNAME_Find)) : nullptr;
 	}
 
 	/**
-	* Find the metadata value associated with the key
-	*
-	* @param Key The key to lookup in the metadata
-	* @return The value associated with the key
-	*/
-	const FString& GetMetaData(const TCHAR* Key) const
+	 * Find the metadata value associated with the key
+	 *
+	 * @param Key The key to lookup in the metadata
+	 * @param NameIndex if specified, will search for metadata linked to a specified value (in an enum); otherwise, searches for metadata for the object itself
+	 * @return The value associated with the key
+	 */
+	FString GetMetaData(const TCHAR* Key, int32 NameIndex = INDEX_NONE, bool bAllowRemap = true) const
 	{
-		return GetMetaData(FName(Key, FNAME_Find));
+		return GetMetaDataHelper(GetMetaDataKey(Key, NameIndex, FNAME_Find), bAllowRemap);
+	}
+	FString GetMetaData(const FName& Key, int32  NameIndex = INDEX_NONE, bool bAllowRemap = true) const
+	{
+		return GetMetaDataHelper(GetMetaDataKey(Key, NameIndex, FNAME_Find), bAllowRemap);
+	}
+
+	/**
+	 * Set the metadata value associated with the specified key.
+	 *
+	 * @param	Key			the metadata tag to find the value for
+	 * @param	NameIndex	if specified, will search the metadata linked for that enum value; otherwise, searches the metadata for the enum itself
+	 * @param	InValue		Value of the metadata for the key
+	 *
+	 */
+	void SetMetaData(const TCHAR* Key, const TCHAR* InValue, int32 NameIndex = INDEX_NONE)
+	{
+		SetMetaDataHelper(GetMetaDataKey(Key, NameIndex, FNAME_Add), InValue);
+	}
+	void SetMetaData(const FName& Key, const TCHAR* InValue, int32 NameIndex = INDEX_NONE) 
+	{
+		SetMetaDataHelper(GetMetaDataKey(Key, NameIndex, FNAME_Add), InValue);
 	}
-	virtual const FString& GetMetaData(const FName& Key) const;
 
 	virtual TMap<FName, FString>* GetMetaDataMap() const = 0;
 
 protected:
+	virtual FString GetMetaDataIndexName(int32 NameIndex) const
+	{
+		checkf(false, TEXT("You can not perform indexed meta data searches on this type"));
+		return FString();
+	}
+
+	virtual const TCHAR* GetMetaDataRemapConfigName() const
+	{
+		checkf(false, TEXT("You can not perform remapping of meta data values on this type"));
+		return nullptr;
+	}
 #if UHT_ENABLE_ENGINE_TYPE_CHECKS
 	virtual void CheckFindMetaData(const FName& Key, const FString* ValuePtr) const = 0;
 #endif
+
+private:
+	FName GetMetaDataKey(const TCHAR* Key, int32 NameIndex, EFindName FindName) const;
+	FName GetMetaDataKey(FName Key, int32 NameIndex, EFindName FindName) const;
+	const FString* FindMetaDataHelper(const FName& Key) const;
+	FString GetMetaDataHelper(const FName& Key, bool bAllowRemap) const;
+	virtual void SetMetaDataHelper(const FName& Key, const TCHAR* InValue);
 };
 
 /**
@@ -322,29 +371,14 @@ public:
 	}
 
 	/**
-	 * Perform any post parsing finalization and validation
+	 * Create UObject engine types
 	 */
-	void PostParseFinalize(bool bCreateEngineTypes)
-	{
-		switch (FinalizeState)
-		{
-		case EFinalizeState::None:
-			FinalizeState = EFinalizeState::InProgress;
-			PostParseFinalizeInternal(bCreateEngineTypes);
-			FinalizeState = EFinalizeState::Finished;
-			break;
-		case EFinalizeState::InProgress:
-			check(false);
-			break;
-		case EFinalizeState::Finished:
-			break;
-		}
-	}
+	void CreateUObjectEngineTypes();
 
 	/**
-	 * Return the C++ type name that represents this type (i.e. UClass)
+	 * Perform any post parsing finalization and validation
 	 */
-	virtual const TCHAR* GetSimplifiedTypeClass() const = 0;
+	void PostParseFinalize();
 
 	/**
 	 * Return the compilation scope associated with this object
@@ -376,6 +410,14 @@ public:
 	virtual FString GetPathName(FUnrealObjectDefinitionInfo* StopOuter = nullptr) const = 0;
 	virtual void GetPathName(FUnrealObjectDefinitionInfo* StopOuter, FStringBuilderBase& ResultString) const = 0;
 
+	/**
+	 * Walks up the chain of packages until it reaches the top level, which it ignores.
+	 *
+	 * @param	bStartWithOuter		whether to include this object's name in the returned string
+	 * @return	string containing the path name for this object, minus the outermost-package's name
+	 */
+	virtual FString GetFullGroupName(bool bStartWithOuter) const = 0;
+
 	/**
 	 * Return true if this type has source information
 	 */
@@ -440,12 +482,6 @@ public:
 		FUHTException::Throwf(*this, TEXT("Meta data can not be set for a definition of this type."));
 	}
 
-	/**
-	 * Helper function that checks if the field is a dynamic type (can be constructed post-startup)
-	 */
-	static bool IsDynamic(const UField* Field);
-	static bool IsDynamic(const FField* Field);
-
 	/**
 	 * Helper function that checks if the field is a dynamic type
 	 */
@@ -505,7 +541,10 @@ protected:
 
 	virtual void ValidateMetaDataFormat(const FName InKey, ECheckedMetadataSpecifier InCheckedMetadataSpecifier, const FString& InValue) const;
 
-	virtual void PostParseFinalizeInternal(bool bCreateEngineTypes)
+	virtual void PostParseFinalizeInternal()
+	{}
+
+	virtual void CreateUObjectEngineTypesInternal()
 	{}
 
 private:
@@ -514,7 +553,8 @@ private:
 	FUnrealSourceFile* SourceFile = nullptr;
 	int32 LineNumber = 0;
 	std::atomic<uint32> Hash = 0;
-	EFinalizeState FinalizeState = EFinalizeState::None;
+	EFinalizeState CreateUObectEngineTypesState = EFinalizeState::None;
+	EFinalizeState PostParseFinalizeState = EFinalizeState::None;
 };
 
 /**
@@ -540,14 +580,6 @@ public:
 		return this;
 	}
 
-	/**
-	 * Return the C++ type name that represents this type (i.e. UClass)
-	 */
-	virtual const TCHAR* GetSimplifiedTypeClass() const override
-	{
-		return TEXT("UProperty");
-	}
-
 	/**
 	 * Return the Engine instance associated with the compiler instance
 	 */
@@ -614,10 +646,18 @@ public:
 	virtual FString GetPathName(FUnrealObjectDefinitionInfo* StopOuter = nullptr) const override;
 	virtual void GetPathName(FUnrealObjectDefinitionInfo* StopOuter, FStringBuilderBase& ResultString) const override;
 
+	/**
+	 * Walks up the chain of packages until it reaches the top level, which it ignores.
+	 *
+	 * @param	bStartWithOuter		whether to include this object's name in the returned string
+	 * @return	string containing the path name for this object, minus the outermost-package's name
+	 */
+	virtual FString GetFullGroupName(bool bStartWithOuter) const;
+
 	/**
 	 * Return the engine class name
 	 */
-	FString GetEngineClassName() const;
+	FString GetEngineClassName() const; //@TODO make this const &
 
 	/**
 	 * Return the package associated with this object
@@ -754,7 +794,7 @@ public:
 	virtual TMap<FName, FString> GenerateMetadataMap() const override
 	{
 		TMap<FName, FString> MetaDataMap;
-		const TMap<FName, FString>* FieldMetaDataMap = GetProperty()->GetMetaDataMap();
+		const TMap<FName, FString>* FieldMetaDataMap = GetMetaDataMap();
 		if (FieldMetaDataMap)
 		{
 			MetaDataMap = *FieldMetaDataMap;
@@ -1129,9 +1169,13 @@ public:
 
 	virtual TMap<FName, FString>* GetMetaDataMap() const override
 	{
+		//return const_cast<TMap<FName, FString>*>(GetProperty()->GetMetaDataMap());
 		return const_cast<TMap<FName, FString>*>(&PropertyBase.MetaData);
 	}
 
+	FUnrealObjectDefinitionInfo* GetOwnerObject() const;
+	FUnrealStructDefinitionInfo* GetOwnerStruct() const;
+
 protected:
 #if UHT_ENABLE_ENGINE_TYPE_CHECKS
 	virtual void CheckFindMetaData(const FName& Key, const FString* ValuePtr) const
@@ -1147,7 +1191,7 @@ protected:
 	/**
 	 * Perform any post parsing finalization and validation
 	 */
-	virtual void PostParseFinalizeInternal(bool bCreateEngineTypes) override;
+	virtual void PostParseFinalizeInternal() override;
 
 private:
 	FName Name;
@@ -1169,6 +1213,7 @@ private:
  */
 class FUnrealObjectDefinitionInfo
 	: public FUnrealTypeDefinitionInfo
+	, public FUHTMetaData
 {
 public:
 
@@ -1186,6 +1231,23 @@ public:
 		return Object;
 	}
 
+	/**
+	 * Return the Engine instance associated with the compiler instance without the check that it is set
+	 */
+	UObject* GetObjectSafe() const
+	{
+		return Object;
+	}
+
+	/**
+	 * Return the engine class name
+	 */
+	virtual const FString& GetEngineClassName(bool bRootClassName = false) const
+	{
+		static FString ClassName(TEXT("Object"));
+		return ClassName;
+	}
+
 	/**
 	 * Set the Engine instance associated with the compiler instance
 	 */
@@ -1194,6 +1256,10 @@ public:
 		check(InObject != nullptr);
 		check(Object == nullptr);
 		Object = InObject;
+#if UHT_ENABLE_ENGINE_TYPE_CHECKS
+		check(Object->GetFName() == Name);
+		check(Object->GetInternalFlags() == InternalObjectFlags);
+#endif
 	}
 
 	/**
@@ -1201,19 +1267,23 @@ public:
 	 */
 	virtual FString GetName() const override
 	{
-		return GetObject()->GetName();
+		FString Result = GetFName().ToString();
+#if UHT_ENABLE_ENGINE_TYPE_CHECKS
+		check(GetObjectSafe() == nullptr || GetObjectSafe()->GetName() == Result);
+#endif
+		return Result;
 	}
 
 	/** Returns the logical name of this object */
 	virtual FName GetFName() const override
 	{
-		return GetObject()->GetFName();
+#if UHT_ENABLE_ENGINE_TYPE_CHECKS
+		check(GetObjectSafe() == nullptr || GetObjectSafe()->GetFName() == Name);
+#endif
+		return Name;
 	}
 
-	virtual FString GetFullName() const override
-	{
-		return GetObject()->GetFullName();
-	}
+	virtual FString GetFullName() const override;
 
 	/**
 	 * Returns the fully qualified pathname for this object, in the format:
@@ -1222,6 +1292,14 @@ public:
 	virtual FString GetPathName(FUnrealObjectDefinitionInfo* StopOuter = nullptr) const override;
 	virtual void GetPathName(FUnrealObjectDefinitionInfo* StopOuter, FStringBuilderBase& ResultString) const override;
 
+	/**
+	 * Walks up the chain of packages until it reaches the top level, which it ignores.
+	 *
+	 * @param	bStartWithOuter		whether to include this object's name in the returned string
+	 * @return	string containing the path name for this object, minus the outermost-package's name
+	 */
+	virtual FString GetFullGroupName(bool bStartWithOuter) const override;
+
 	/**
 	 * Return the package associated with this object
 	 */
@@ -1235,19 +1313,75 @@ public:
 		return static_cast<FUnrealObjectDefinitionInfo*>(FUnrealTypeDefinitionInfo::GetOuter());
 	}
 
+	/**
+	 * Return the flags that were parsed as part of the pre-parse phase
+	 */
+	EInternalObjectFlags GetInternalFlags() const
+	{
+#if UHT_ENABLE_ENGINE_TYPE_CHECKS
+		check(GetObjectSafe() == nullptr || GetObjectSafe()->GetInternalFlags() == InternalObjectFlags);
+#endif
+		return InternalObjectFlags;
+	}
+
+	/**
+	 * Sets the given flags
+	 */
+	void SetInternalFlags(EInternalObjectFlags FlagsToSet)
+	{
+		InternalObjectFlags |= FlagsToSet;
+		if (UObject* Obj = GetObjectSafe())
+		{
+			Obj->SetInternalFlags(FlagsToSet);
+#if UHT_ENABLE_ENGINE_TYPE_CHECKS
+			check(Obj->GetInternalFlags() == InternalObjectFlags);
+#endif
+		}
+	}
+
+	/**
+	 * Clears the given flags
+	 */
+	void ClearInternalFlags(EInternalObjectFlags FlagsToClear)
+	{
+		InternalObjectFlags &= ~FlagsToClear;
+		if (UObject* Obj = GetObjectSafe())
+		{
+			Obj->ClearInternalFlags(FlagsToClear);
+#if UHT_ENABLE_ENGINE_TYPE_CHECKS
+			check(Obj->GetInternalFlags() == InternalObjectFlags);
+#endif
+		}
+	}
+
+	/**
+	 * Used to safely check whether the passed in flag is set.
+	 *
+	 * @param	FlagsToCheck		Flag(s) to check for
+	 * @return	true if the passed in flag is set, false otherwise
+	 */
+	bool HasAnyInternalFlags(EInternalObjectFlags FlagsToCheck) const
+	{
+#if UHT_ENABLE_ENGINE_TYPE_CHECKS
+		check(GetObjectSafe() == nullptr || GetObjectSafe()->GetInternalFlags() == InternalObjectFlags);
+#endif
+		return EnumHasAnyFlags(InternalObjectFlags, FlagsToCheck);
+	}
+
+	/**
+	 * Helper method to see if native
+	 */
+	bool IsNative() const
+	{
+		return HasAnyInternalFlags(EInternalObjectFlags::Native);
+	}
+
 	/**
 	 * Return the meta data in map form
 	 */
 	virtual TMap<FName, FString> GenerateMetadataMap() const override
 	{
-		UObject* Obj = GetObject();
-		check(Obj);
-		UPackage* Package = Obj->GetPackage();
-		check(Package);
-		UMetaData* Metadata = Package->GetMetaData();
-		check(Metadata);
-
-		TMap<FName, FString>* PackageMap = Metadata->ObjectMetaDataMap.Find(Obj);
+		TMap<FName, FString>* PackageMap = GetMetaDataMap();
 		TMap<FName, FString> Map;
 		if (PackageMap)
 		{
@@ -1263,32 +1397,63 @@ public:
 		return Map;
 	}
 
-	/**
-	 * Helper methods to remove dependencies on engine types
-	 */
-	bool IsADelegateFunction() const
+	virtual TMap<FName, FString>* GetMetaDataMap() const override
 	{
-		return GetObject()->IsA<UDelegateFunction>();
+		return GetUObjectMetaDataMap();
 	}
 
 protected:
-	// Used only by packages
-	explicit FUnrealObjectDefinitionInfo(FString&& InNameCPP)
-		: FUnrealTypeDefinitionInfo(MoveTemp(InNameCPP))
-	{ }
+	virtual void SetMetaDataHelper(const FName& Key, const TCHAR* InValue) override;
 
-	FUnrealObjectDefinitionInfo(FUnrealSourceFile& InSourceFile, int32 InLineNumber, FString&& InNameCPP, FUnrealObjectDefinitionInfo& InOuter)
+	explicit FUnrealObjectDefinitionInfo(UObject* Object);
+
+	FUnrealObjectDefinitionInfo(FUnrealSourceFile& InSourceFile, int32 InLineNumber, FString&& InNameCPP, FName InName, FUnrealObjectDefinitionInfo& InOuter)
 		: FUnrealTypeDefinitionInfo(InSourceFile, InLineNumber, MoveTemp(InNameCPP), InOuter)
+		, Name(InName)
 	{ }
 
+#if UHT_ENABLE_ENGINE_TYPE_CHECKS
+	virtual void CheckFindMetaData(const FName& Key, const FString* ValuePtr) const override
+	{
+		if (UObject* Obj = GetObjectSafe())
+		{
+			TMap<FName, FString>* PackageMap = GetUObjectMetaDataMap();
+			const FString* CheckPtr = PackageMap ? PackageMap->Find(Key) : nullptr;
+			check((CheckPtr == nullptr && ValuePtr == nullptr) || (CheckPtr != nullptr && ValuePtr != nullptr && *CheckPtr == *ValuePtr));
+		}
+	}
+#endif
+
+	UMetaData* GetUObjectMetaData() const
+	{
+		if (UObject* Obj = GetObjectSafe())
+		{
+			UPackage* Package = Obj->GetPackage();
+			check(Package);
+			UMetaData* Metadata = Package->GetMetaData();
+			check(Metadata);
+			return Metadata;
+		}
+		return nullptr;
+	}
+
+	TMap<FName, FString>* GetUObjectMetaDataMap() const
+	{
+		UMetaData* Metadata = GetUObjectMetaData();
+		return Metadata ? Metadata->ObjectMetaDataMap.Find(GetObject()) : nullptr;
+	}
+
 private:
 	UObject* Object = nullptr;
+	FName Name = NAME_None;
+	EInternalObjectFlags InternalObjectFlags = EInternalObjectFlags::None;
 };
 
 /**
  * Class that stores information about packages.
  */
-class FUnrealPackageDefinitionInfo : public FUnrealObjectDefinitionInfo
+class FUnrealPackageDefinitionInfo 
+	: public FUnrealObjectDefinitionInfo
 {
 public:
 	// Constructor
@@ -1308,11 +1473,20 @@ public:
 	}
 
 	/**
-	 * Return the C++ type name that represents this type (i.e. UClass)
+	 * Return the Engine instance associated with the compiler instance
 	 */
-	virtual const TCHAR* GetSimplifiedTypeClass() const override
+	UPackage* GetPackageSafe() const
 	{
-		return TEXT("UPackage");
+		return static_cast<UPackage*>(GetObjectSafe());
+	}
+
+	/**
+	 * Return the engine class name
+	 */
+	virtual const FString& GetEngineClassName(bool bRootClassName = false) const
+	{
+		static FString ClassName(TEXT("Package"));
+		return ClassName;
 	}
 
 	/**
@@ -1402,10 +1576,15 @@ public:
 	}
 
 protected:
+	/**
+	 * Create UObject engine types
+	 */
+	virtual void CreateUObjectEngineTypesInternal() override;
+
 	/**
 	 * Perform any post parsing finalization and validation
 	 */
-	virtual void PostParseFinalizeInternal(bool bCreateEngineTypes) override;
+	virtual void PostParseFinalizeInternal() override;
 
 private:
 	const FManifestModule& Module;
@@ -1444,58 +1623,30 @@ public:
 	}
 
 	/**
-	 * Determines if the property has any metadata associated with the key
-	 * 
-	 * @param Key The key to lookup in the metadata
-	 * @return true if there is a (possibly blank) value associated with this key
-	 */
-	bool HasMetaData(const TCHAR* Key) const { return FindMetaData(Key) != nullptr; }
-	bool HasMetaData(const FName& Key) const { return FindMetaData(Key) != nullptr; }
-
-	/**
-	 * Find the metadata value associated with the key
-	 *
-	 * @param Key The key to lookup in the metadata
-	 * @return The value associated with the key if exists, null otherwise
+	 * Return the Engine instance associated with the compiler instance without the check that it is set
 	 */
-	const FString* FindMetaData(const TCHAR* Key) const
+	UField* GetFieldSafe() const
 	{
-		return GetField()->FindMetaData(Key);
-	}
-	const FString* FindMetaData(const FName& Key) const
-	{
-		return GetField()->FindMetaData(Key);
+		return static_cast<UField*>(GetObjectSafe());
 	}
 
 	/**
-	 * Find the metadata value associated with the key
-	 * 
-	 * @param Key The key to lookup in the metadata
-	 * @return The value associated with the key
+	 * Return the engine class name
 	 */
-	const FString& GetMetaData(const TCHAR* Key) const
-	{
-		return GetField()->GetMetaData(Key);
-	}
-	const FString& GetMetaData(const FName& Key) const
+	virtual const FString& GetEngineClassName(bool bRootClassName = false) const
 	{
-		return GetField()->GetMetaData(Key);
+		static FString ClassName(TEXT("Field"));
+		return ClassName;
 	}
 
 	/**
-	 * Sets the metadata value associated with the key
-	 * 
-	 * @param Key The key to lookup in the metadata
-	 * @return The value associated with the key
-	 */
-	void SetMetaData(const TCHAR* Key, const TCHAR* InValue)
-	{
-		GetField()->SetMetaData(Key, InValue);
-	}
-	void SetMetaData(const FName& Key, const TCHAR* InValue)
-	{
-		GetField()->SetMetaData(Key, InValue);
-	}
+	* Finds the localized tooltip or native tooltip as a fallback.
+	*
+	* @param bShortTooltip Look for a shorter version of the tooltip (falls back to the long tooltip if none was specified)
+	*
+	* @return The tooltip for this object.
+	*/
+	FText GetToolTipText(bool bShortTooltip = false) const;
 
 	/**
 	 * Add a unique cross module reference for this field
@@ -1558,7 +1709,7 @@ protected:
 	/**
 	 * Perform any post parsing finalization and validation
 	 */
-	virtual void PostParseFinalizeInternal(bool bCreateEngineTypes) override;
+	virtual void PostParseFinalizeInternal() override;
 
 private:
 	FString SingletonName[2];
@@ -1573,7 +1724,7 @@ private:
 class FUnrealEnumDefinitionInfo : public FUnrealFieldDefinitionInfo
 {
 public:
-	FUnrealEnumDefinitionInfo(FUnrealSourceFile& InSourceFile, int32 InLineNumber, FString&& InNameCPP);
+	FUnrealEnumDefinitionInfo(FUnrealSourceFile& InSourceFile, int32 InLineNumber, FString&& InNameCPP, FName InName);
 
 	virtual FUnrealEnumDefinitionInfo* AsEnum() override
 	{
@@ -1581,11 +1732,12 @@ public:
 	}
 
 	/**
-	 * Return the C++ type name that represents this type (i.e. UClass)
+	 * Return the engine class name
 	 */
-	virtual const TCHAR* GetSimplifiedTypeClass() const override
+	virtual const FString& GetEngineClassName(bool bRootClassName = false) const
 	{
-		return TEXT("UEnum");
+		static FString ClassName(TEXT("Enum"));
+		return ClassName;
 	}
 
 	/**
@@ -1699,46 +1851,6 @@ public:
 		return GetEnum()->SetEnums(InNames, InCppForm, InFlags, bAddMaxKeyIfMissing);
 	}
 
-	/**
-	 * Wrapper method for easily determining whether this enum has metadata associated with it.
-	 *
-	 * @param	Key			the metadata tag to check for
-	 * @param	NameIndex	if specified, will search for metadata linked to a specified value in this enum; otherwise, searches for metadata for the enum itself
-	 *
-	 * @return true if the specified key exists in the list of metadata for this enum, even if the value of that key is empty
-	 */
-	bool HasMetaData(const TCHAR* Key, int32 NameIndex = INDEX_NONE) const
-	{
-		return GetEnum()->HasMetaData(Key, NameIndex);
-	}
-
-	/**
-	 * Return the metadata value associated with the specified key.
-	 *
-	 * @param	Key			the metadata tag to find the value for
-	 * @param	NameIndex	if specified, will search the metadata linked for that enum value; otherwise, searches the metadata for the enum itself
-	 * @param	bAllowRemap	if true, the returned value may be remapped from a .ini if the value starts with ini: Pass false when you need the exact string, including any ini:
-	 *
-	 * @return	the value for the key specified, or an empty string if the key wasn't found or had no value.
-	 */
-	FString GetMetaData(const TCHAR* Key, int32 NameIndex = INDEX_NONE, bool bAllowRemap = true) const
-	{
-		return GetEnum()->GetMetaData(Key, NameIndex, bAllowRemap);
-	}
-
-	/**
-	 * Set the metadata value associated with the specified key.
-	 *
-	 * @param	Key			the metadata tag to find the value for
-	 * @param	NameIndex	if specified, will search the metadata linked for that enum value; otherwise, searches the metadata for the enum itself
-	 * @param	InValue		Value of the metadata for the key
-	 *
-	 */
-	void SetMetaData(const TCHAR* Key, const TCHAR* InValue, int32 NameIndex = INDEX_NONE) const
-	{
-		GetEnum()->SetMetaData(Key, InValue, NameIndex);
-	}
-
 	/**
 	 * Returns the underlying enumeration type
 	 */
@@ -1771,6 +1883,19 @@ public:
 		bIsEditorOnly = true;
 	}
 
+protected:
+	virtual FString GetMetaDataIndexName(int32 NameIndex) const
+	{
+		FString EnumName = GetEnum()->GetNameStringByIndex(NameIndex);
+		check(!EnumName.IsEmpty());
+		return EnumName;
+	}
+
+	virtual const TCHAR* GetMetaDataRemapConfigName() const
+	{
+		return TEXT("EnumRemap");
+	}
+
 private:
 	EUnderlyingEnumType UnderlyingType = EUnderlyingEnumType::Unspecified;
 	bool bIsEditorOnly = false;
@@ -1823,6 +1948,15 @@ public:
 		return static_cast<UStruct*>(GetObject());
 	}
 
+	/**
+	 * Return the engine class name
+	 */
+	virtual const FString& GetEngineClassName(bool bRootClassName = false) const
+	{
+		static FString ClassName(TEXT("Struct"));
+		return ClassName;
+	}
+
 	/**
 	 * Return the CPP prefix 
 	 */
@@ -1925,29 +2059,6 @@ public:
 		return BaseStructInfo;
 	}
 
-	/** Try and find boolean metadata with the given key. If not found on this class, work up hierarchy looking for it. */
-	bool GetBoolMetaDataHierarchical(const FName& Key) const
-	{
-		return GetStruct()->GetBoolMetaDataHierarchical(Key);
-	}
-
-	/** Try and find string metadata with the given key. If not found on this class, work up hierarchy looking for it. */
-	bool GetStringMetaDataHierarchical(const FName& Key, FString* OutValue = nullptr) const
-	{
-		return GetStruct()->GetStringMetaDataHierarchical(Key, OutValue);
-	}
-
-	/**
-	 * Determines if the struct or any of its super structs has any metadata associated with the provided key
-	 *
-	 * @param Key The key to lookup in the metadata
-	 * @return pointer to the UStruct that has associated metadata, nullptr if Key is not associated with any UStruct in the hierarchy
-	 */
-	const UStruct* HasMetaDataHierarchical(const FName& Key) const
-	{
-		return GetStruct()->HasMetaDataHierarchical(Key);
-	}
-
 	/**
 	 * Get the contains delegates flag
 	 */
@@ -1979,6 +2090,42 @@ public:
 		return false;
 	}
 
+	const FString* FindMetaDataHierarchical(const FName& Key) const
+	{
+		if (Key == NAME_None)
+		{
+			return nullptr;
+		}
+
+		for (const FUnrealStructDefinitionInfo* StructDef = this; StructDef; StructDef = StructDef->GetSuperStruct())
+		{
+			if (const FString* Found = StructDef->FindMetaData(Key))
+			{
+				return Found;
+			}
+		}
+		return nullptr;
+	}
+
+	/** Try and find boolean metadata with the given key. If not found on this class, work up hierarchy looking for it. */
+	bool GetBoolMetaDataHierarchical(const FName& Key) const
+	{
+		const FString* Found = FindMetaDataHierarchical(Key);
+		// FString == operator does case insensitive comparison
+		return Found != nullptr && *Found == "true";
+	}
+
+	/** Try and find string metadata with the given key. If not found on this class, work up hierarchy looking for it. */
+	bool GetStringMetaDataHierarchical(const FName& Key, FString* OutValue = nullptr) const
+	{
+		const FString* Found = FindMetaDataHierarchical(Key);
+		if (OutValue)
+		{
+			*OutValue = Found ? *Found : FString();
+		}
+		return Found != nullptr;
+	}
+
 	/**
 	 * Get the generated code version
 	 */
@@ -2028,10 +2175,15 @@ public:
 	}
 
 protected:
+	/**
+	 * Create UObject engine types
+	 */
+	virtual void CreateUObjectEngineTypesInternal() override;
+
 	/**
 	 * Perform any post parsing finalization and validation
 	 */
-	virtual void PostParseFinalizeInternal(bool bCreateEngineTypes) override;
+	virtual void PostParseFinalizeInternal() override;
 
 private:
 	TSharedPtr<FScope> StructScope;
@@ -2065,7 +2217,7 @@ private:
 class FUnrealScriptStructDefinitionInfo : public FUnrealStructDefinitionInfo
 {
 public:
-	FUnrealScriptStructDefinitionInfo(FUnrealSourceFile& InSourceFile, int32 InLineNumber, FString&& InNameCPP);
+	FUnrealScriptStructDefinitionInfo(FUnrealSourceFile& InSourceFile, int32 InLineNumber, FString&& InNameCPP, FName InName);
 
 	virtual FUnrealScriptStructDefinitionInfo* AsScriptStruct() override
 	{
@@ -2075,11 +2227,12 @@ public:
 	virtual uint32 GetHash(bool bIncludeNoExport = true) const override;
 
 	/**
-	 * Return the C++ type name that represents this type (i.e. UClass)
+	 * Return the engine class name
 	 */
-	virtual const TCHAR* GetSimplifiedTypeClass() const override
+	virtual const FString& GetEngineClassName(bool bRootClassName = false) const
 	{
-		return TEXT("UScriptStruct");
+		static FString ClassName(TEXT("ScriptStruct"));
+		return ClassName;
 	}
 
 	/**
@@ -2165,9 +2318,10 @@ private:
 class FUnrealFunctionDefinitionInfo : public FUnrealStructDefinitionInfo
 {
 public:
-	FUnrealFunctionDefinitionInfo(FUnrealSourceFile& InSourceFile, int32 InLineNumber, FString&& InNameCPP, FUnrealObjectDefinitionInfo& InOuter, FFuncInfo&& InFuncInfo)
-		: FUnrealStructDefinitionInfo(InSourceFile, InLineNumber, MoveTemp(InNameCPP), InOuter)
+	FUnrealFunctionDefinitionInfo(FUnrealSourceFile& InSourceFile, int32 InLineNumber, FString&& InNameCPP, FName InName, FUnrealObjectDefinitionInfo& InOuter, FFuncInfo&& InFuncInfo, EFunctionType InFunctionType)
+		: FUnrealStructDefinitionInfo(InSourceFile, InLineNumber, MoveTemp(InNameCPP), InName, InOuter)
 		, FunctionData(MoveTemp(InFuncInfo))
+		, FunctionType(InFunctionType)
 	{ 
 	}
 
@@ -2177,11 +2331,31 @@ public:
 	}
 
 	/**
-	 * Return the C++ type name that represents this type (i.e. UClass)
+	 * Return the engine class name
 	 */
-	virtual const TCHAR* GetSimplifiedTypeClass() const override
+	virtual const FString& GetEngineClassName(bool bRootClassName = false) const
 	{
-		return TEXT("UFunction");
+		static FString FunctionClassName(TEXT("Function"));
+		static FString DelegateFunctionClassName(TEXT("DelegateFunction"));
+		static FString SparseDelegateFunctionClassName(TEXT("SparseDelegateFunction"));
+
+		if (bRootClassName)
+		{
+			return FunctionClassName;
+		}
+
+		switch (FunctionType)
+		{
+		case EFunctionType::Function:
+			return FunctionClassName;
+		case EFunctionType::Delegate:
+			return DelegateFunctionClassName;
+		case EFunctionType::SparseDelegate:
+			return SparseDelegateFunctionClassName;
+		default:
+			check(false);
+			return FunctionClassName;
+		}
 	}
 
 	/**
@@ -2189,7 +2363,15 @@ public:
 	 */
 	UFunction* GetFunction() const
 	{
-		return static_cast<UFunction*>(GetField());
+		return static_cast<UFunction*>(GetObject());
+	}
+
+	/**
+	 * Return the Engine instance associated with the compiler instance without the check to see if it is set
+	 */
+	UFunction* GetFunctionSafe() const
+	{
+		return static_cast<UFunction*>(GetObjectSafe());
 	}
 
 	/**
@@ -2197,7 +2379,25 @@ public:
 	 */
 	FORCEINLINE EFunctionFlags GetFunctionFlags() const
 	{
-		return GetFunction()->FunctionFlags;
+#if UHT_ENABLE_ENGINE_TYPE_CHECKS
+		check(GetFunctionSafe() == nullptr || GetFunctionSafe()->FunctionFlags == FunctionData.FunctionFlags);
+#endif
+		return FunctionData.FunctionFlags;
+	}
+
+	/**
+	 * Sets the given flags
+	 */
+	void SetFunctionFlags(EFunctionFlags FlagsToSet)
+	{
+#if UHT_ENABLE_ENGINE_TYPE_CHECKS
+		check(GetFunctionSafe() == nullptr || GetFunctionSafe()->FunctionFlags == FunctionData.FunctionFlags);
+#endif
+		if (GetFunctionSafe())
+		{
+			GetFunctionSafe()->FunctionFlags |= FlagsToSet;
+		}
+		FunctionData.FunctionFlags |= FlagsToSet;
 	}
 
 	/**
@@ -2210,7 +2410,10 @@ public:
 	 */
 	FORCEINLINE bool HasAnyFunctionFlags(EFunctionFlags FlagsToCheck) const
 	{
-		return GetFunction()->HasAnyFunctionFlags(FlagsToCheck);
+#if UHT_ENABLE_ENGINE_TYPE_CHECKS
+		check(GetFunctionSafe() == nullptr || GetFunctionSafe()->FunctionFlags == FunctionData.FunctionFlags);
+#endif
+		return EnumHasAnyFlags(FunctionData.FunctionFlags, FlagsToCheck);
 	}
 
 	/**
@@ -2221,7 +2424,10 @@ public:
 	 */
 	FORCEINLINE bool HasAllFunctionFlags(EFunctionFlags FlagsToCheck) const
 	{
-		return GetFunction()->HasAllFunctionFlags(FlagsToCheck);
+#if UHT_ENABLE_ENGINE_TYPE_CHECKS
+		check(GetFunctionSafe() == nullptr || GetFunctionSafe()->FunctionFlags == FunctionData.FunctionFlags);
+#endif
+		return EnumHasAllFlags(FunctionData.FunctionFlags, FlagsToCheck);
 	}
 
 	/**
@@ -2233,16 +2439,16 @@ public:
 	 */
 	FORCEINLINE bool HasSpecificFunctionFlags(EFunctionFlags FlagsToCheck, EFunctionFlags ExpectedFlags) const
 	{
-		EFunctionFlags Flags = GetFunction()->FunctionFlags;
+#if UHT_ENABLE_ENGINE_TYPE_CHECKS
+		check(GetFunctionSafe() == nullptr || GetFunctionSafe()->FunctionFlags == FunctionData.FunctionFlags);
+#endif
+		EFunctionFlags Flags = FunctionData.FunctionFlags;
 		return (Flags & FlagsToCheck) == ExpectedFlags;
 	}
 
-	/** @name getters */
-	//@{
 	FFuncInfo& GetFunctionData() { return FunctionData; }
 	const FFuncInfo& GetFunctionData() const { return FunctionData; }
 	FUnrealPropertyDefinitionInfo* GetReturn() const { return ReturnProperty.Get(); }
-	//@}
 
 	/**
 	 * Adds a new function property to be tracked.  Determines whether the property is a
@@ -2254,7 +2460,7 @@ public:
 	/**
 	 * Sets the specified function export flags
 	 */
-	void SetFunctionExportFlag(uint32 NewFlags)
+	void SetFunctionExportFlags(uint32 NewFlags)
 	{
 		FunctionData.FunctionExportFlags |= NewFlags;
 	}
@@ -2272,11 +2478,47 @@ public:
 	 */
 	FUnrealFunctionDefinitionInfo* GetSuperFunction() const;
 
+	/**
+	 * Set the name of the sparse owning class name
+	 */
+	void SetSparseOwningClassName(FName InSparseOwningClassName)
+	{
+		SparseOwningClassName = InSparseOwningClassName;
+	}
+
+	/**
+	 * Set the name of the sparse delegate name
+	 */
+	void SetSparseDelegateName(FName InSparseDelegateName)
+	{
+		SparseDelegateName = InSparseDelegateName;
+	}
+
+	EFunctionType GetFunctionType() const
+	{
+		return FunctionType;
+	}
+
+	bool IsDelegateFunction() const
+	{
+		return FunctionType == EFunctionType::Delegate || FunctionType == EFunctionType::SparseDelegate;
+	}
+
+	virtual TMap<FName, FString>* GetMetaDataMap() const override
+	{
+		return const_cast<TMap<FName, FString>*>(&MetaData);
+	}
+
 protected:
+	/**
+	 * Create UObject engine types
+	 */
+	virtual void CreateUObjectEngineTypesInternal() override;
+
 	/**
 	 * Perform any post parsing finalization and validation
 	 */
-	virtual void PostParseFinalizeInternal(bool bCreateEngineTypes) override;
+	virtual void PostParseFinalizeInternal() override;
 
 private:
 
@@ -2285,6 +2527,11 @@ private:
 
 	/** return value for this function */
 	TSharedPtr<FUnrealPropertyDefinitionInfo> ReturnProperty;
+
+	TMap<FName, FString> MetaData;
+	FName SparseOwningClassName = NAME_None;
+	FName SparseDelegateName = NAME_None;
+	EFunctionType FunctionType = EFunctionType::Function;
 };
 
 /**
@@ -2294,10 +2541,10 @@ class FUnrealClassDefinitionInfo
 	: public FUnrealStructDefinitionInfo
 {
 public:
-	FUnrealClassDefinitionInfo(FUnrealSourceFile& InSourceFile, int32 InLineNumber, FString&& InNameCPP, bool bInIsInterface);
+	FUnrealClassDefinitionInfo(FUnrealSourceFile& InSourceFile, int32 InLineNumber, FString&& InNameCPP, FName InName, bool bInIsInterface);
 
-	FUnrealClassDefinitionInfo(FString&& InNameCPP)
-		: FUnrealStructDefinitionInfo(MoveTemp(InNameCPP))
+	explicit FUnrealClassDefinitionInfo(UClass* Class) 
+		: FUnrealStructDefinitionInfo(Class)
 	{ }
 
 	/**
@@ -2331,11 +2578,12 @@ public:
 	}
 
 	/**
-	 * Return the C++ type name that represents this type (i.e. UClass)
+	 * Return the engine class name
 	 */
-	virtual const TCHAR* GetSimplifiedTypeClass() const override
+	virtual const FString& GetEngineClassName(bool bRootClassName = false) const
 	{
-		return TEXT("UClass");
+		static FString ClassName(TEXT("Class"));
+		return ClassName;
 	}
 
 	virtual uint32 GetHash(bool bIncludeNoExport = true) const override;
@@ -2359,14 +2607,7 @@ public:
 		FUnrealStructDefinitionInfo::SetObject(InObject);
 	}
 
-	FUnrealClassDefinitionInfo* GetSuperClass() const
-	{
-		if (FUnrealStructDefinitionInfo* SuperClass = GetSuperStructInfo().Struct)
-		{
-			return &SuperClass->AsClassChecked();
-		}
-		return nullptr;
-	}
+	FUnrealClassDefinitionInfo* GetSuperClass() const;
 
 	/**
 	 * Get the archive type
@@ -2582,7 +2823,7 @@ protected:
 	/**
 	 * Perform any post parsing finalization and validation
 	 */
-	virtual void PostParseFinalizeInternal(bool bCreateEngineTypes) override;
+	virtual void PostParseFinalizeInternal() override;
 
 public:
 	TMap<FName, FString> MetaData;
