Removed the usage of FError::Throwf and the need for HeaderParser to set a logging context.

#rnx

diff --git a/Engine/Source/Programs/UnrealHeaderTool/Private/BaseParser.cpp b/Engine/Source/Programs/UnrealHeaderTool/Private/BaseParser.cpp
--- a/Engine/Source/Programs/UnrealHeaderTool/Private/BaseParser.cpp
+++ b/Engine/Source/Programs/UnrealHeaderTool/Private/BaseParser.cpp
@@ -2,8 +2,8 @@
 
 #include "BaseParser.h"
 #include "UnrealHeaderTool.h"
+#include "UnrealSourceFile.h"
 #include "UObject/NameTypes.h"
-#include "UObject/ErrorException.h"
 
 #include "ParserHelper.h"
 
@@ -60,13 +60,13 @@ namespace
 			ValueActions.Add(FMetadataValueAction(InMapping, TEXT(""), EMetadataValueAction::Remove));
 		}
 
-		void ApplyToMetadata(TMap<FName, FString>& Metadata, const FString* Value = nullptr)
+		void ApplyToMetadata(FBaseParser& Parser, TMap<FName, FString>& Metadata, const FString* Value = nullptr)
 		{
 			for (const FMetadataValueAction& ValueAction : ValueActions)
 			{
 				if (ValueAction.ValueAction == EMetadataValueAction::Add)
 				{
-					FBaseParser::InsertMetaDataPair(Metadata, ValueAction.Mapping, Value ? *Value : ValueAction.DefaultValue);
+					Parser.InsertMetaDataPair(Metadata, ValueAction.Mapping, Value ? *Value : ValueAction.DefaultValue);
 				}
 				else
 				{
@@ -180,13 +180,19 @@ FString FPropertySpecifier::ConvertToString() const
 /////////////////////////////////////////////////////
 // FBaseParser
 
-FBaseParser::FBaseParser()
+FBaseParser::FBaseParser(FUnrealSourceFile& InSourceFile)
 	: UHTConfig(FUHTConfig::Get())
+	, SourceFile(InSourceFile)
 	, StatementsParsed(0)
 	, LinesParsed(0)
 {
 }
 
+FString FBaseParser::GetFilename() const
+{
+	return SourceFile.GetFilename();
+};
+
 void FBaseParser::ResetParser(const TCHAR* SourceBuffer, int32 StartingLineNumber)
 {
 	Input = SourceBuffer;
@@ -251,7 +257,7 @@ Loop:
 			if (!bInsideComment)
 			{
 				ClearComment();
-				FError::Throwf(TEXT("Unexpected '*/' outside of comment") );
+				FUHTException::Throwf(*this, TEXT("Unexpected '*/' outside of comment") );
 			}
 
 			/** Asterisk and slash always end comment. */
@@ -270,7 +276,7 @@ Loop:
 		if (c == 0)
 		{
 			ClearComment();
-			FError::Throwf(TEXT("End of class header encountered inside comment") );
+			FUHTException::Throwf(*this, TEXT("End of class header encountered inside comment") );
 		}
 		goto Loop;
 	}
@@ -395,7 +401,7 @@ bool FBaseParser::GetToken( FToken& Token, bool bNoConsts/*=false*/, ESymbolPars
 			Token.Identifier[Length++] = c;
 			if( Length >= NAME_SIZE )
 			{
-				FError::Throwf(TEXT("Identifer length exceeds maximum of %i"), (int32)NAME_SIZE);
+				FUHTException::Throwf(*this, TEXT("Identifer length exceeds maximum of %i"), (int32)NAME_SIZE);
 				Length = ((int32)NAME_SIZE) - 1;
 				break;
 			}
@@ -447,7 +453,7 @@ bool FBaseParser::GetToken( FToken& Token, bool bNoConsts/*=false*/, ESymbolPars
 			Token.Identifier[Length++] = c;
 			if( Length >= NAME_SIZE )
 			{
-				FError::Throwf(TEXT("Number length exceeds maximum of %i "), (int32)NAME_SIZE );
+				FUHTException::Throwf(*this, TEXT("Number length exceeds maximum of %i "), (int32)NAME_SIZE );
 				Length = ((int32)NAME_SIZE) - 1;
 				break;
 			}
@@ -499,7 +505,7 @@ bool FBaseParser::GetToken( FToken& Token, bool bNoConsts/*=false*/, ESymbolPars
 		c = GetChar(/*bLiteral=*/ true);
 		if (c != '\'')
 		{
-			FError::Throwf(TEXT("Unterminated character constant"));
+			FUHTException::Throwf(*this, TEXT("Unterminated character constant"));
 			UngetChar();
 		}
 
@@ -530,7 +536,7 @@ bool FBaseParser::GetToken( FToken& Token, bool bNoConsts/*=false*/, ESymbolPars
 			Temp[Length++] = c;
 			if( Length >= MAX_STRING_CONST_SIZE )
 			{
-				FError::Throwf(TEXT("String constant exceeds maximum of %i characters"), (int32)MAX_STRING_CONST_SIZE );
+				FUHTException::Throwf(*this, TEXT("String constant exceeds maximum of %i characters"), (int32)MAX_STRING_CONST_SIZE );
 				c = TEXT('\"');
 				Length = ((int32)MAX_STRING_CONST_SIZE) - 1;
 				break;
@@ -541,7 +547,7 @@ bool FBaseParser::GetToken( FToken& Token, bool bNoConsts/*=false*/, ESymbolPars
 
 		if( c != '"' )
 		{
-			FError::Throwf(TEXT("Unterminated string constant: %s"), Temp);
+			FUHTException::Throwf(*this, TEXT("Unterminated string constant: %s"), Temp);
 			UngetChar();
 		}
 
@@ -628,7 +634,7 @@ bool FBaseParser::GetRawTokenRespectingQuotes( FToken& Token, TCHAR StopChar /*
 		Temp[Length++] = c;
 		if( Length >= MAX_STRING_CONST_SIZE )
 		{
-			FError::Throwf(TEXT("Identifier exceeds maximum of %i characters"), (int32)MAX_STRING_CONST_SIZE );
+			FUHTException::Throwf(*this, TEXT("Identifier exceeds maximum of %i characters"), (int32)MAX_STRING_CONST_SIZE );
 			c = GetChar(true);
 			Length = ((int32)MAX_STRING_CONST_SIZE) - 1;
 			break;
@@ -639,7 +645,7 @@ bool FBaseParser::GetRawTokenRespectingQuotes( FToken& Token, TCHAR StopChar /*
 
 	if (bInQuote)
 	{
-		FError::Throwf(TEXT("Unterminated quoted string"));
+		FUHTException::Throwf(*this, TEXT("Unterminated quoted string"));
 	}
 
 	// Get rid of trailing whitespace.
@@ -678,7 +684,7 @@ bool FBaseParser::GetRawToken( FToken& Token, TCHAR StopChar /* = TCHAR('\n') */
 		Temp[Length++] = c;
 		if( Length >= MAX_STRING_CONST_SIZE )
 		{
-			FError::Throwf(TEXT("Identifier exceeds maximum of %i characters"), (int32)MAX_STRING_CONST_SIZE );
+			FUHTException::Throwf(*this, TEXT("Identifier exceeds maximum of %i characters"), (int32)MAX_STRING_CONST_SIZE );
 		}
 		c = GetChar(true);
 	}
@@ -763,7 +769,7 @@ bool FBaseParser::GetConstInt(int32& Result, const TCHAR* Tag)
 
 	if (Tag != NULL)
 	{
-		FError::Throwf(TEXT("%s: Missing constant integer"), Tag );
+		FUHTException::Throwf(*this, TEXT("%s: Missing constant integer"), Tag );
 	}
 
 	return false;
@@ -786,7 +792,7 @@ bool FBaseParser::GetConstInt64(int64& Result, const TCHAR* Tag)
 
 	if (Tag != NULL)
 	{
-		FError::Throwf(TEXT("%s: Missing constant integer"), Tag );
+		FUHTException::Throwf(*this, TEXT("%s: Missing constant integer"), Tag );
 	}
 
 	return false;
@@ -891,11 +897,11 @@ void FBaseParser::MatchSemi()
 		FToken Token;
 		if( GetToken(Token) )
 		{
-			FError::Throwf(TEXT("Missing ';' before '%s'"), Token.Identifier );
+			FUHTException::Throwf(*this, TEXT("Missing ';' before '%s'"), Token.Identifier );
 		}
 		else
 		{
-			FError::Throwf(TEXT("Missing ';'") );
+			FUHTException::Throwf(*this, TEXT("Missing ';'") );
 		}
 	}
 }
@@ -952,7 +958,7 @@ void FBaseParser::RequireSymbol( const TCHAR Match, const TCHAR* Tag, ESymbolPar
 {
 	if (!MatchSymbol(Match, bParseTemplateClosingBracket))
 	{
-		FError::Throwf(TEXT("Missing '%c' in %s"), Match, Tag );
+		FUHTException::Throwf(*this, TEXT("Missing '%c' in %s"), Match, Tag );
 	}
 }
 
@@ -960,7 +966,7 @@ void FBaseParser::RequireSymbol(const TCHAR Match, TFunctionRef<FString()> Error
 {
 	if (!MatchSymbol(Match, bParseTemplateClosingBracket))
 	{
-		FError::Throwf(TEXT("Missing '%c' in %s"), Match, *ErrorGetter());
+		FUHTException::Throwf(*this, TEXT("Missing '%c' in %s"), Match, *ErrorGetter());
 	}
 }
 
@@ -971,7 +977,7 @@ void FBaseParser::RequireConstInt( const TCHAR* Match, const TCHAR* Tag )
 {
 	if (!MatchConstInt(Match))
 	{
-		FError::Throwf(TEXT("Missing integer '%s' in %s"), Match, Tag );
+		FUHTException::Throwf(*this, TEXT("Missing integer '%s' in %s"), Match, Tag );
 	}
 }
 
@@ -979,7 +985,7 @@ void FBaseParser::RequireAnyConstInt( const TCHAR* Tag )
 {
 	if (!MatchAnyConstInt())
 	{
-		FError::Throwf(TEXT("Missing integer in %s"), Tag );
+		FUHTException::Throwf(*this, TEXT("Missing integer in %s"), Tag );
 	}
 }
 
@@ -990,7 +996,7 @@ void FBaseParser::RequireIdentifier( const TCHAR* Match, const ESearchCase::Type
 {
 	if (!MatchIdentifier(Match, SearchCase))
 	{
-		FError::Throwf(TEXT("Missing '%s' in %s"), Match, Tag );
+		FUHTException::Throwf(*this, TEXT("Missing '%s' in %s"), Match, Tag );
 	}
 }
 
@@ -1007,7 +1013,7 @@ FString FBaseParser::ReadNewStyleValue(const TCHAR* TypeOfSpecifier)
 	FToken ValueToken;
 	if (!GetToken(ValueToken, false))
 	{
-		FError::Throwf(TEXT("Expected a value when handling a %s"), TypeOfSpecifier);
+		FUHTException::Throwf(*this, TEXT("Expected a value when handling a %s"), TypeOfSpecifier);
 	}
 
 	FString Result;
@@ -1062,7 +1068,7 @@ void FBaseParser::ParseNameWithPotentialAPIMacroPrefix(FString& DeclaredName, FS
 	// Read an identifier
 	if (!GetIdentifier(NameToken))
 	{
-		FError::Throwf(TEXT("Missing %s name"), FailureMessage);
+		FUHTException::Throwf(*this, TEXT("Missing %s name"), FailureMessage);
 	}
 
 	// Is the identifier the name or an DLL import/export API macro?
@@ -1074,7 +1080,7 @@ void FBaseParser::ParseNameWithPotentialAPIMacroPrefix(FString& DeclaredName, FS
 		// Read the real name
 		if (!GetIdentifier(NameToken))
 		{
-			FError::Throwf(TEXT("Missing %s name"), FailureMessage);
+			FUHTException::Throwf(*this, TEXT("Missing %s name"), FailureMessage);
 		}
 		DeclaredName = NameToken.Identifier;
 	}
@@ -1106,7 +1112,7 @@ void FBaseParser::ReadSpecifierSetInsideMacro(TArray<FPropertySpecifier>& Specif
 		FToken Specifier;
 		if (!GetToken(Specifier))
 		{
-			FError::Throwf(TEXT("Expected %s"), *ErrorMessageGetter());
+			FUHTException::Throwf(*this, TEXT("Expected %s"), *ErrorMessageGetter());
 		}
 
 		if (Specifier.Matches(TEXT("meta"), ESearchCase::IgnoreCase))
@@ -1121,7 +1127,7 @@ void FBaseParser::ReadSpecifierSetInsideMacro(TArray<FPropertySpecifier>& Specif
 				FToken MetaKeyToken;
 				if (!GetIdentifier(MetaKeyToken))
 				{
-					FError::Throwf(TEXT("Expected a metadata key"));
+					FUHTException::Throwf(*this, TEXT("Expected a metadata key"));
 				}
 
 				FString Key = MetaKeyToken.Identifier;
@@ -1146,20 +1152,20 @@ void FBaseParser::ReadSpecifierSetInsideMacro(TArray<FPropertySpecifier>& Specif
 			{
 				if (MetadataKeyword->ValueArgument == EMetadataValueArgument::None)
 				{
-					FError::Throwf(TEXT("Incorrect = after metadata specifier '%s'"), Specifier.Identifier);
+					FUHTException::Throwf(*this, TEXT("Incorrect = after metadata specifier '%s'"), Specifier.Identifier);
 				}
 
 				FString Value = ReadNewStyleValue(TypeOfSpecifier);
-				MetadataKeyword->ApplyToMetadata(MetaData, &Value);
+				MetadataKeyword->ApplyToMetadata(*this, MetaData, &Value);
 			}
 			else
 			{
 				if (MetadataKeyword->ValueArgument == EMetadataValueArgument::Required)
 				{
-					FError::Throwf(TEXT("Missing = after metadata specifier '%s'"), Specifier.Identifier);
+					FUHTException::Throwf(*this, TEXT("Missing = after metadata specifier '%s'"), Specifier.Identifier);
 				}
 
-				MetadataKeyword->ApplyToMetadata(MetaData);
+				MetadataKeyword->ApplyToMetadata(*this, MetaData);
 			}
 		}
 		else
@@ -1185,7 +1191,7 @@ void FBaseParser::InsertMetaDataPair(TMap<FName, FString>& MetaData, FString Key
 	// make sure the key is valid
 	if (Key.Len() == 0)
 	{
-		FError::Throwf(TEXT("Invalid metadata"));
+		FUHTException::Throwf(*this, TEXT("Invalid metadata"));
 	}
 
 	// trim extra white space and quotes
@@ -1202,7 +1208,7 @@ void FBaseParser::InsertMetaDataPair(TMap<FName, FString>& MetaData, FName KeyNa
 	FString* ExistingValue = MetaData.Find(KeyName);
 	if (ExistingValue && Value != *ExistingValue)
 	{
-		FError::Throwf(TEXT("Metadata key '%s' first seen with value '%s' then '%s'"), *KeyName.ToString(), **ExistingValue, *Value);
+		FUHTException::Throwf(*this, TEXT("Metadata key '%s' first seen with value '%s' then '%s'"), *KeyName.ToString(), **ExistingValue, *Value);
 	}
 
 	// finally we have enough to put it into our metadata
diff --git a/Engine/Source/Programs/UnrealHeaderTool/Private/BaseParser.h b/Engine/Source/Programs/UnrealHeaderTool/Private/BaseParser.h
--- a/Engine/Source/Programs/UnrealHeaderTool/Private/BaseParser.h
+++ b/Engine/Source/Programs/UnrealHeaderTool/Private/BaseParser.h
@@ -3,10 +3,10 @@
 #pragma once
 
 #include "CoreMinimal.h"
+#include "Exceptions.h"
 #include "GeneratedCodeVersion.h"
 #include "ParserHelper.h"
 #include "Containers/UnrealString.h"
-#include "UObject/ErrorException.h"
 
 class FToken;
 
@@ -83,15 +83,20 @@ public:
 // Base class of header parsers.
 //
 
-class FBaseParser
+class FBaseParser 
+	: public FUHTExceptionContext
 {
 protected:
-	FBaseParser();
+	FBaseParser(FUnrealSourceFile& InSourceFile);
+	virtual ~FBaseParser() = default;
 
 	// UHTConfig data
 	const FUHTConfig& UHTConfig;
 
 public:
+	// Source being parsed
+	FUnrealSourceFile& SourceFile;
+
 	// Input text.
 	const TCHAR* Input;
 
@@ -119,6 +124,13 @@ public:
 	// Total number of lines parsed.
 	int32 LinesParsed;
 
+	virtual FString GetFilename() const override;
+
+	virtual int32 GetLineNumber() const override
+	{
+		return InputLine;
+	};
+
 	void ResetParser(const TCHAR* SourceBuffer, int32 StartingLineNumber = 1);
 
 	// Low-level parsing functions.
@@ -228,10 +240,10 @@ public:
 	void ReadSpecifierSetInsideMacro(TArray<FPropertySpecifier>& SpecifiersFound, const TCHAR* TypeOfSpecifier, TMap<FName, FString>& MetaData);
 
 	// Validates and inserts one key-value pair into the meta data map
-	static void InsertMetaDataPair(TMap<FName, FString>& MetaData, FString InKey, FString InValue);
+	void InsertMetaDataPair(TMap<FName, FString>& MetaData, FString InKey, FString InValue);
 
 	// Validates and inserts one key-value pair into the meta data map
-	static void InsertMetaDataPair(TMap<FName, FString>& MetaData, FName InKey, FString InValue);
+	void InsertMetaDataPair(TMap<FName, FString>& MetaData, FName InKey, FString InValue);
 
 	/**
 	 * Parse class/struct inheritance.
@@ -264,7 +276,7 @@ void FBaseParser::ParseInheritance(const TCHAR* What, Lambda&& InLambda)
 		RequireIdentifier(TEXT("public"), ESearchCase::CaseSensitive, TEXT("inheritance"));
 		if (!GetIdentifier(Token))
 		{
-			FError::Throwf(TEXT("Missing %s name"), What);
+			FUHTException::Throwf(*this, TEXT("Missing %s name"), What);
 		}
 		RedirectTypeIdentifier(Token);
 		InLambda(Token.Identifier, true);
@@ -282,7 +294,7 @@ void FBaseParser::ParseInheritance(const TCHAR* What, Lambda&& InLambda)
 			FToken Token;
 			if (!GetIdentifier(Token, true))
 			{
-				FError::Throwf(TEXT("Failed to get interface class identifier"));
+				FUHTException::Throwf(*this, TEXT("Failed to get interface class identifier"));
 			}
 
 			InterfaceName += Token.Identifier;
@@ -297,7 +309,7 @@ void FBaseParser::ParseInheritance(const TCHAR* What, Lambda&& InLambda)
 				{
 					if (!GetToken(Token))
 					{
-						FError::Throwf(TEXT("Unexpected end of file"));
+						FUHTException::Throwf(*this, TEXT("Unexpected end of file"));
 					}
 
 					if (Token.TokenType == TOKEN_Symbol)
diff --git a/Engine/Source/Programs/UnrealHeaderTool/Private/ClassMaps.h b/Engine/Source/Programs/UnrealHeaderTool/Private/ClassMaps.h
--- a/Engine/Source/Programs/UnrealHeaderTool/Private/ClassMaps.h
+++ b/Engine/Source/Programs/UnrealHeaderTool/Private/ClassMaps.h
@@ -7,7 +7,6 @@
 #include "UnrealSourceFile.h"
 #include "UnrealTypeDefinitionInfo.h"
 #include "UObject/Stack.h"
-#include "UObject/ErrorException.h"
 
 class UField;
 class UClass;
diff --git a/Engine/Source/Programs/UnrealHeaderTool/Private/CodeGenerator.cpp b/Engine/Source/Programs/UnrealHeaderTool/Private/CodeGenerator.cpp
--- a/Engine/Source/Programs/UnrealHeaderTool/Private/CodeGenerator.cpp
+++ b/Engine/Source/Programs/UnrealHeaderTool/Private/CodeGenerator.cpp
@@ -22,7 +22,6 @@
 #include "Misc/FeedbackContext.h"
 #include "Misc/OutputDeviceNull.h"
 #include "UObject/ClassTree.h"
-#include "UObject/ErrorException.h"
 #include "UObject/Script.h"
 #include "UObject/ObjectMacros.h"
 #include "UObject/UObjectGlobals.h"
@@ -78,8 +77,8 @@ static bool bWriteContents = false;
 static bool bVerifyContents = false;
 
 UClass* ProcessParsedClass(FUnrealClassDefinitionInfo& ClassDef);
-UEnum* ProcessParsedEnum(const FString& EnumName, UObject* InParent, EObjectFlags Flags);
-UScriptStruct* ProcessParsedStruct(const FString& StructName, UObject* InParent, EObjectFlags Flags);
+UEnum* ProcessParsedEnum(FUnrealEnumDefinitionInfo& EnumDef);
+UScriptStruct* ProcessParsedStruct(FUnrealScriptStructDefinitionInfo& ScriptStructDef);
 
 // Array of all the temporary header async file tasks so we can ensure they have completed before issuing our timings
 static FGraphEventArray GAsyncFileTasks;
@@ -285,7 +284,7 @@ namespace
 			}
 			else
 			{
-				FError::Throwf(TEXT("Class %s has Net flagged properties and should declare member function: void GetLifetimeReplicatedProps(TArray<FLifetimeProperty>& OutLifetimeProps) const override"), ClassCPPName);
+				FUHTException::Throwf(ClassDef, TEXT("Class %s has Net flagged properties and should declare member function: void GetLifetimeReplicatedProps(TArray<FLifetimeProperty>& OutLifetimeProps) const override"), ClassCPPName);
 			}
 		}
 
@@ -1654,7 +1653,7 @@ FString GetEventStructParamsName(FUnrealObjectDefinitionInfo& OuterDef, const TC
 	}
 	else
 	{
-		FError::Throwf(TEXT("Unrecognized outer type"));
+		FUHTException::Throwf(OuterDef, TEXT("Unrecognized outer type"));
 	}
 
 	FString Result = FString::Printf(TEXT("%s_event%s_Parms"), *OuterName, FunctionName);
@@ -1929,7 +1928,7 @@ void FNativeClassHeaderGenerator::ExportNativeGeneratedInitCode(FOutputDevice& O
 			// In a dynamic class the same function signature may be used for a Multi- and a Single-cast delegate.
 			if (!LocalFuncDef->IsDelegateFunction() || !bIsDynamic)
 			{
-				FError::Throwf(TEXT("The same function linked twice. Function: %s Class: %s"), *LocalFuncDef->GetName(), *ClassDef.GetName());
+				FUHTException::Throwf(*LocalFuncDef, TEXT("The same function linked twice. Function: %s Class: %s"), *LocalFuncDef->GetName(), *ClassDef.GetName());
 			}
 			continue;
 		}
@@ -2755,7 +2754,7 @@ void FNativeClassHeaderGenerator::ExportClassFromSourceFileInner(
 
 	FString PPOMacroName;
 
-	ClassDef.GetDefinitionRange().Validate();
+	ClassDef.ValidateDefinitionRange();
 
 	FString GeneratedSerializeFunctionCPP;
 	FString GeneratedSerializeFunctionHeaderMacroName;
@@ -4419,34 +4418,6 @@ FString FNativeClassHeaderGenerator::GetFunctionReturnString(FUnrealFunctionDefi
 	return Result;
 }
 
-/**
- * Converts Position within File to Line and Column.
- *
- * @param File File contents.
- * @param Position Position in string to convert.
- * @param OutLine Result line.
- * @param OutColumn Result column.
- */
-void GetLineAndColumnFromPositionInFile(const FString& File, int32 Position, int32& OutLine, int32& OutColumn)
-{
-	OutLine = 1;
-	OutColumn = 1;
-
-	int32 i;
-	for (i = 1; i <= Position; ++i)
-	{
-		if (File[i] == '\n')
-		{
-			++OutLine;
-			OutColumn = 0;
-		}
-		else
-		{
-			++OutColumn;
-		}
-	}
-}
-
 bool FNativeClassHeaderGenerator::IsMissingVirtualSpecifier(const FString& SourceFile, int32 FunctionNamePosition)
 {
 	auto IsEndOfSearchChar = [](TCHAR C) { return (C == TEXT('}')) || (C == TEXT('{')) || (C == TEXT(';')); };
@@ -4459,11 +4430,6 @@ bool FNativeClassHeaderGenerator::IsMissingVirtualSpecifier(const FString& Sourc
 	return !HasIdentifierExactMatch(&SourceFile[EndOfSearchCharIndex], &SourceFile[FunctionNamePosition], TEXT("virtual"));
 }
 
-FString CreateClickableErrorMessage(const FString& Filename, int32 Line, int32 Column)
-{
-	return FString::Printf(TEXT("%s(%d,%d): error: "), *Filename, Line, Column);
-}
-
 void FNativeClassHeaderGenerator::CheckRPCFunctions(FReferenceGatherers& OutReferenceGatherers, FUnrealFunctionDefinitionInfo& FunctionDef, const FString& ClassName, int32 ImplementationPosition, int32 ValidatePosition, const FUnrealSourceFile& SourceFile) const
 {
 	bool bHasImplementation = ImplementationPosition != INDEX_NONE;
@@ -4514,21 +4480,16 @@ void FNativeClassHeaderGenerator::CheckRPCFunctions(FReferenceGatherers& OutRefe
 	//
 	// Check if functions are missing.
 	//
-	int32 Line;
-	int32 Column;
-	GetLineAndColumnFromPositionInFile(FileContent, FunctionData.InputPos, Line, Column);
 	if (bNeedsImplementation && !bHasImplementation)
 	{
-		FString ErrorPosition = CreateClickableErrorMessage(Filename, Line, Column);
 		FString FunctionDecl = FString::Printf(TEXT("virtual %s %s::%s(%s) %s"), *FunctionReturnType, *ClassName, *FunctionData.CppImplName, *ParameterString, *ConstModifier);
-		FError::Throwf(TEXT("%s%s Declare function %s"), *ErrorPosition, *AssertMessage, *FunctionDecl);
+		FUHTException::Throwf(FunctionDef, TEXT("%s Declare function %s"), *AssertMessage, *FunctionDecl);
 	}
 
 	if (bNeedsValidate && !bHasValidate)
 	{
-		FString ErrorPosition = CreateClickableErrorMessage(Filename, Line, Column);
 		FString FunctionDecl = FString::Printf(TEXT("virtual bool %s::%s(%s) %s"), *ClassName, *FunctionData.CppValidationImplName, *ParameterString, *ConstModifier);
-		FError::Throwf(TEXT("%s%s Declare function %s"), *ErrorPosition, *AssertMessage, *FunctionDecl);
+		FUHTException::Throwf(FunctionDef, TEXT("%s Declare function %s"), *AssertMessage, *FunctionDecl);
 	}
 
 	//
@@ -4536,18 +4497,14 @@ void FNativeClassHeaderGenerator::CheckRPCFunctions(FReferenceGatherers& OutRefe
 	//
 	if (bNeedsImplementation && bHasImplementation && IsMissingVirtualSpecifier(FileContent, ImplementationPosition))
 	{
-		GetLineAndColumnFromPositionInFile(FileContent, ImplementationPosition, Line, Column);
-		FString ErrorPosition = CreateClickableErrorMessage(Filename, Line, Column);
 		FString FunctionDecl = FString::Printf(TEXT("%s %s::%s(%s) %s"), *FunctionReturnType, *ClassName, *FunctionData.CppImplName, *ParameterString, *ConstModifier);
-		FError::Throwf(TEXT("%sDeclared function %sis not marked as virtual."), *ErrorPosition, *FunctionDecl);
+		FUHTException::Throwf(FunctionDef, TEXT("Declared function %sis not marked as virtual."), *FunctionDecl);
 	}
 
 	if (bNeedsValidate && bHasValidate && IsMissingVirtualSpecifier(FileContent, ValidatePosition))
 	{
-		GetLineAndColumnFromPositionInFile(FileContent, ValidatePosition, Line, Column);
-		FString ErrorPosition = CreateClickableErrorMessage(Filename, Line, Column);
 		FString FunctionDecl = FString::Printf(TEXT("bool %s::%s(%s) %s"), *ClassName, *FunctionData.CppValidationImplName, *ParameterString, *ConstModifier);
-		FError::Throwf(TEXT("%sDeclared function %sis not marked as virtual."), *ErrorPosition, *FunctionDecl);
+		FUHTException::Throwf(FunctionDef, TEXT("Declared function %sis not marked as virtual."), *FunctionDecl);
 	}
 }
 
@@ -4855,7 +4812,7 @@ void FNativeClassHeaderGenerator::ExportFunctionThunk(FUHTStringBuilder& RPCWrap
 	FUnrealClassDefinitionInfo* OwnerClassDef = FunctionDef.GetOwnerClass();
 	check(OwnerClassDef);
 
-	OwnerClassDef->GetDefinitionRange().Validate();
+	OwnerClassDef->ValidateDefinitionRange();
 	const TCHAR* ClassStart = OwnerClassDef->GetDefinitionRange().Start;
 	const TCHAR* ClassEnd = OwnerClassDef->GetDefinitionRange().End;
 	FString      ClassName  = OwnerClassDef->GetName();
@@ -4969,7 +4926,7 @@ void FNativeClassHeaderGenerator::ExportNativeFunctions(FOutputDevice& OutGenera
 	FStructMetaData& StructData = ClassDef.GetStructMetaData();
 	const FString ClassCPPName = ClassDef.GetAlternateNameCPP(ClassDef.HasAnyClassFlags(CLASS_Interface));
 
-	ClassDef.GetDefinitionRange().Validate();
+	ClassDef.ValidateDefinitionRange();
 	const TCHAR* ClassStart = ClassDef.GetDefinitionRange().Start;
 	const TCHAR* ClassEnd = ClassDef.GetDefinitionRange().End;
 
@@ -6056,19 +6013,19 @@ bool FNativeClassHeaderGenerator::SaveHeaderIfChanged(FGeneratedFileInfo& FileIn
 			FFileHelper::SaveStringToFile(InNewHeaderContents, *ConflictPath);
 
 			FResults::SetResult(ECompilationResult::FailedDueToHeaderChange);
-			FError::Throwf(TEXT("ERROR: '%s': Changes to generated code are not allowed - conflicts written to '%s'"), *HeaderPathStr, *ConflictPath);
+			FUHTException::Throwf(FString(FileInfo.GetFilename()), 1, TEXT("ERROR: '%s': Changes to generated code are not allowed - conflicts written to '%s'"), *HeaderPathStr, *ConflictPath);
 		}
 
 		// save the updated version to a tmp file so that the user can see what will be changing
 		FString TmpHeaderFilename = GenerateTempHeaderName(HeaderPathStr, false);
 
-		auto SaveTempTask = [TmpHeaderFilename, InNewHeaderContents = MoveTemp(InNewHeaderContents)]()
+		auto SaveTempTask = [&FileInfo, TmpHeaderFilename, InNewHeaderContents = MoveTemp(InNewHeaderContents)]()
 		{
 			// delete any existing temp file
 			IFileManager::Get().Delete(*TmpHeaderFilename, false, true);
 			if (!FFileHelper::SaveStringToFile(InNewHeaderContents, *TmpHeaderFilename))
 			{
-				UE_LOG_WARNING_UHT(TEXT("Failed to save header export preview: '%s'"), *TmpHeaderFilename);
+				FResults::LogWarning(FString(FileInfo.GetFilename()), 1, FString::Printf(TEXT("Failed to save header export preview: '%s'"), *TmpHeaderFilename));
 			}
 		};
 
@@ -6213,7 +6170,7 @@ void ResolveSuperClasses(UPackage* Package)
 					if (FoundBaseClassDef == nullptr)
 					{
 						// Don't know its parent class. Raise error.
-						FError::Throwf(TEXT("Couldn't find parent type for '%s' named '%s' in current module (Package: %s) or any other module parsed so far."), *DefinedClass->GetName(), *BaseClassName, *Package->GetName());
+						FUHTException::Throwf(ClassDef, TEXT("Couldn't find parent type for '%s' named '%s' in current module (Package: %s) or any other module parsed so far."), *DefinedClass->GetName(), *BaseClassName, *Package->GetName());
 					}
 
 					DefinedClass->SetSuperStruct(FoundBaseClassDef->GetClass());
@@ -6436,7 +6393,7 @@ void PreparseSources(TArray<FUnrealPackageDefinitionInfo*>& PackageDefs, const F
 						FString Content;
 						if (!FFileHelper::LoadFileToString(Content, *FullFilename))
 						{
-							FError::Throwf(TEXT("UnrealHeaderTool was unable to load source file '%s'"), *FullFilename);
+							FUHTException::Throwf(SourceFile, 1, TEXT("UnrealHeaderTool was unable to load source file '%s'"), *FullFilename);
 						}
 						SourceFile.SetContent(MoveTemp(Content));
 					}
@@ -6494,7 +6451,7 @@ void DefineTypes(TArray<FUnrealPackageDefinitionInfo*>& PackageDefs)
 				for (TSharedRef<FUnrealTypeDefinitionInfo>& TypeDef : SourceFile.GetDefinedEnums())
 				{
 					FUnrealEnumDefinitionInfo& EnumDef = TypeDef->AsEnumChecked();
-					UEnum* ResultEnum = ProcessParsedEnum(EnumDef.GetNameCPP(), Package, RF_Public);
+					UEnum* ResultEnum = ProcessParsedEnum(EnumDef);
 					EnumDef.SetObject(ResultEnum);
 					GTypeDefinitionInfoMap.Add(ResultEnum, TypeDef);
 				}
@@ -6502,7 +6459,7 @@ void DefineTypes(TArray<FUnrealPackageDefinitionInfo*>& PackageDefs)
 				for (TSharedRef<FUnrealTypeDefinitionInfo>& TypeDef : SourceFile.GetDefinedStructs())
 				{
 					FUnrealScriptStructDefinitionInfo& ScriptStructDef = TypeDef->AsScriptStructChecked();
-					UScriptStruct* ResultScriptStruct = ProcessParsedStruct(ScriptStructDef.GetNameCPP(), Package, RF_Public);
+					UScriptStruct* ResultScriptStruct = ProcessParsedStruct(ScriptStructDef);
 					ScriptStructDef.SetObject(ResultScriptStruct);
 					GTypeDefinitionInfoMap.Add(ResultScriptStruct, TypeDef);
 				}
@@ -6539,16 +6496,7 @@ void ResolveParents(TArray<FUnrealPackageDefinitionInfo*>& PackageDefs)
 
 void PrepareTypesForParsing(TArray<FUnrealPackageDefinitionInfo*>& PackageDefs)
 {
-	// Verify that all script declared superclasses exist.
-	TMap<const UClass*, TArray<const UClass*>> Uniques;
-	for (UClass* ScriptClass : TObjectRange<UClass>())
-	{
-		const UClass* ScriptSuperClass = ScriptClass->GetSuperClass();
-		if (ScriptSuperClass && !ScriptSuperClass->HasAnyClassFlags(CLASS_Intrinsic) && GTypeDefinitionInfoMap.Contains(ScriptClass) && !GTypeDefinitionInfoMap.Contains(ScriptSuperClass))
-		{
-			FResults::LogError(*ScriptClass, *FString::Printf(TEXT("Superclass %s of class %s not found"), *ScriptSuperClass->GetName(), *ScriptClass->GetName()));
-		}
-	}
+	// Does nothing now
 }
 
 void TopologicalRecursion(FUnrealSourceFile& First, FUnrealSourceFile& Visit)
@@ -6637,12 +6585,7 @@ void ParseSourceFiles(TArray<FUnrealSourceFile*>& OrderedSourceFiles)
 		FUnrealPackageDefinitionInfo& PackageDef = SourceFile->GetPackageDef();
 		FScopedDurationTimer SourceTimer(SourceFile->GetTime(ESourceFileTime::Parse));
 
-		ECompilationResult::Type Result = FHeaderParser::Parse(PackageDef, GWarn, *SourceFile);
-		if (Result != ECompilationResult::Succeeded)
-		{
-			FResults::SetResult(Result);
-			break;
-		}
+		FResults::TryAlways([&PackageDef, SourceFile]() { FHeaderParser::Parse(PackageDef, *SourceFile); });
 	}
 }
 
@@ -6659,21 +6602,20 @@ void PrepareTypesForExport(TArray<FUnrealPackageDefinitionInfo*>& PackageDefs)
 	for (FUnrealPackageDefinitionInfo* PackageDef : PackageDefs)
 	{
 		FResults::Try([PackageDef]()
-			{
-				PackageDef->PostParseFinalize();
-			}
-		);
+		{
+			PackageDef->PostParseFinalize();
+		});
 	}
 	FResults::WaitForErrorTasks();
 
 	// Until we fix the issue with types not having sources, we need to do this
 	GTypeDefinitionInfoMap.ForAllTypes([](FUnrealTypeDefinitionInfo& TypeDef) 
-		{ 
-			if (!TypeDef.HasSource())
-			{
-				TypeDef.PostParseFinalize();
-			}
-		});
+	{ 
+		if (!TypeDef.HasSource())
+		{
+			TypeDef.PostParseFinalize();
+		}
+	});
 }
 
 void Export(TArray<FUnrealPackageDefinitionInfo*>& PackageDefs, TArray<FUnrealSourceFile*>& OrderedSourceFiles)
@@ -6838,21 +6780,8 @@ ECompilationResult::Type UnrealHeaderTool_Main(const FString& ModuleInfoFilename
 		GTypeDefinitionInfoMap.Add(*ClassIt, MakeShareable<FUnrealTypeDefinitionInfo>(ClassDef));
 	}
 
-
 	// Load the manifest file, giving a list of all modules to be processed, pre-sorted by dependency ordering
-#if !PLATFORM_EXCEPTIONS_DISABLED
-	try
-#endif
-	{
-		GManifest = FManifest::LoadFromFile(ModuleInfoFilename);
-	}
-#if !PLATFORM_EXCEPTIONS_DISABLED
-	catch (const TCHAR* Ex)
-	{
-		FResults::LogError(*FString::Printf(TEXT("Failed to load manifest file '%s': %s"), *ModuleInfoFilename, Ex));
-		return FResults::GetOverallResults();
-	}
-#endif
+	FResults::Try([&ModuleInfoFilename]() { GManifest = FManifest::LoadFromFile(ModuleInfoFilename); });
 
 	TArray<FUnrealSourceFile*> OrderedSourceFiles;
 	TArray<FUnrealPackageDefinitionInfo*> PackageDefs;
@@ -6944,12 +6873,12 @@ UClass* ProcessParsedClass(FUnrealClassDefinitionInfo& ClassDef)
 	// All classes must start with a valid unreal prefix
 	if (!FHeaderParser::ClassNameHasValidPrefix(ClassName, ClassNameStripped))
 	{
-		FError::Throwf(TEXT("Invalid class name '%s'. The class name must have an appropriate prefix added (A for Actors, U for other classes)."), *ClassName);
+		FUHTException::Throwf(ClassDef, TEXT("Invalid class name '%s'. The class name must have an appropriate prefix added (A for Actors, U for other classes)."), *ClassName);
 	}
 
 	if(FHeaderParser::IsReservedTypeName(ClassNameStripped))
 	{
-		FError::Throwf(TEXT("Invalid class name '%s'. Cannot use a reserved name ('%s')."), *ClassName, *ClassNameStripped);
+		FUHTException::Throwf(ClassDef, TEXT("Invalid class name '%s'. Cannot use a reserved name ('%s')."), *ClassName, *ClassNameStripped);
 	}
 
 	// Ensure the base class has any valid prefix and exists as a valid class. Checking for the 'correct' prefix will occur during compilation
@@ -6959,7 +6888,7 @@ UClass* ProcessParsedClass(FUnrealClassDefinitionInfo& ClassDef)
 		BaseClassNameStripped = GetClassNameWithPrefixRemoved(BaseClassName);
 		if (!FHeaderParser::ClassNameHasValidPrefix(BaseClassName, BaseClassNameStripped))
 		{
-			FError::Throwf(TEXT("No prefix or invalid identifier for base class %s.\nClass names must match Unreal prefix specifications (e.g., \"UObject\" or \"AActor\")"), *BaseClassName);
+			FUHTException::Throwf(ClassDef, TEXT("No prefix or invalid identifier for base class %s.\nClass names must match Unreal prefix specifications (e.g., \"UObject\" or \"AActor\")"), *BaseClassName);
 		}
 	}
 
@@ -6967,12 +6896,12 @@ UClass* ProcessParsedClass(FUnrealClassDefinitionInfo& ClassDef)
 	// Handle failure and non-class headers.
 	if (BaseClassName.IsEmpty() && (ClassName != TEXT("UObject")))
 	{
-		FError::Throwf(TEXT("Class '%s' must inherit UObject or a UObject-derived class"), *ClassName);
+		FUHTException::Throwf(ClassDef, TEXT("Class '%s' must inherit UObject or a UObject-derived class"), *ClassName);
 	}
 
 	if (ClassName == BaseClassName)
 	{
-		FError::Throwf(TEXT("Class '%s' cannot inherit from itself"), *ClassName);
+		FUHTException::Throwf(ClassDef, TEXT("Class '%s' cannot inherit from itself"), *ClassName);
 	}
 
 	UClass* ResultClass = FEngineAPI::FindObject<UClass>(Package, *ClassNameStripped);
@@ -6989,7 +6918,7 @@ UClass* ProcessParsedClass(FUnrealClassDefinitionInfo& ClassDef)
 			UClass* ConflictingClass = FEngineAPI::FindObject<UClass>(ANY_PACKAGE, *ClassNameStripped, true);
 			if (ConflictingClass != nullptr)
 			{
-				UE_LOG_WARNING_UHT(TEXT("Duplicate class name: %s also exists in file %s"), *ClassName, *ConflictingClass->GetOutermost()->GetName());
+				UE_LOG_WARNING_UHT(ClassDef, TEXT("Duplicate class name: %s also exists in file %s"), *ClassName, *ConflictingClass->GetOutermost()->GetName());
 			}
 		}
 
@@ -7017,40 +6946,46 @@ UClass* ProcessParsedClass(FUnrealClassDefinitionInfo& ClassDef)
 	return ResultClass;
 }
 
-UEnum* ProcessParsedEnum(const FString& EnumName, UObject* InParent, EObjectFlags Flags)
+UEnum* ProcessParsedEnum(FUnrealEnumDefinitionInfo& EnumDef)
 {
+	UPackage* Package = EnumDef.GetPackageDef().GetPackage();
+	const FString& EnumName = EnumDef.GetNameCPP();
+
 	if (UEnum* Existing = FEngineAPI::FindObject<UEnum>(ANY_PACKAGE, *EnumName))
 	{
-		FError::Throwf(TEXT("Duplicate enum name: %s also exists in file %s"), *EnumName, *Existing->GetOutermost()->GetName());
+		FUHTException::Throwf(EnumDef, TEXT("Duplicate enum name: %s also exists in file %s"), *EnumName, *Existing->GetOutermost()->GetName());
 	}
 
 	// Check if the enum name is using a reserved keyword
 	if (FHeaderParser::IsReservedTypeName(EnumName))
 	{
-		FError::Throwf(TEXT("enum: '%s' uses a reserved type name."), *EnumName);
+		FUHTException::Throwf(EnumDef, TEXT("enum: '%s' uses a reserved type name."), *EnumName);
 	}
 
 	// Create enum definition.
-	UEnum* Enum = new(EC_InternalUseOnlyConstructor, InParent, FName(EnumName), Flags) UEnum(FObjectInitializer());
+	UEnum* Enum = new(EC_InternalUseOnlyConstructor, Package, FName(EnumName), RF_Public) UEnum(FObjectInitializer());
 	return Enum;
 }
 
-UScriptStruct* ProcessParsedStruct(const FString& StructName, UObject* InParent, EObjectFlags Flags)
+UScriptStruct* ProcessParsedStruct(FUnrealScriptStructDefinitionInfo& ScriptStructDef)
 {
+	UPackage* Package = ScriptStructDef.GetPackageDef().GetPackage();
+	const FString& StructName = ScriptStructDef.GetNameCPP();
+
 	FString StructNameStripped = GetClassNameWithPrefixRemoved(*StructName);
 
 	if (UScriptStruct* Existing = FEngineAPI::FindObject<UScriptStruct>(ANY_PACKAGE, *StructNameStripped))
 	{
-		FError::Throwf(TEXT("Duplicate struct name: %s also exists in file %s"), *StructNameStripped, *Existing->GetOutermost()->GetName());
+		FUHTException::Throwf(ScriptStructDef, TEXT("Duplicate struct name: %s also exists in file %s"), *StructNameStripped, *Existing->GetOutermost()->GetName());
 	}
 
 	// Check if the enum name is using a reserved keyword
 	if (FHeaderParser::IsReservedTypeName(StructNameStripped))
 	{
-		FError::Throwf(TEXT("struct: '%s' uses a reserved type name."), *StructNameStripped);
+		FUHTException::Throwf(ScriptStructDef, TEXT("struct: '%s' uses a reserved type name."), *StructNameStripped);
 	}
 
 	// Create enum definition.
-	UScriptStruct* Struct = new(EC_InternalUseOnlyConstructor, InParent, FName(StructNameStripped), Flags) UScriptStruct(FObjectInitializer());
+	UScriptStruct* Struct = new(EC_InternalUseOnlyConstructor, Package, FName(StructNameStripped), RF_Public) UScriptStruct(FObjectInitializer());
 	return Struct;
 }
diff --git a/Engine/Source/Programs/UnrealHeaderTool/Private/Exceptions.cpp b/Engine/Source/Programs/UnrealHeaderTool/Private/Exceptions.cpp
--- a/Engine/Source/Programs/UnrealHeaderTool/Private/Exceptions.cpp
+++ b/Engine/Source/Programs/UnrealHeaderTool/Private/Exceptions.cpp
@@ -3,6 +3,7 @@
 #include "Exceptions.h"
 
 #include "ClassMaps.h"
+#include "BaseParser.h"
 #include "UnrealHeaderTool.h"
 #include "UnrealHeaderToolGlobals.h"
 #include "UnrealSourceFile.h"
@@ -32,11 +33,11 @@ FUHTException::FUHTException(ECompilationResult::Type InResult, const FUnrealSou
 {
 }
 
-FUHTException::FUHTException(ECompilationResult::Type InResult, FUnrealTypeDefinitionInfo& TypeDef, FString&& InMessage)
+FUHTException::FUHTException(ECompilationResult::Type InResult, const FUHTExceptionContext& InContext, FString&& InMessage)
 	: Result(InResult)
 	, Message(MoveTemp(InMessage))
-	, Filename(TypeDef.GetUnrealSourceFile().GetFilename())
-	, Line(TypeDef.GetLineNumber())
+	, Filename(InContext.GetFilename())
+	, Line(InContext.GetLineNumber())
 {
 }
 
@@ -44,6 +45,8 @@ namespace UE::UnrealHeaderTool::Exceptions::Private
 {
 	std::atomic<ECompilationResult::Type> OverallResults = ECompilationResult::Succeeded;
 	std::atomic<int32> NumFailures = 0;
+	std::atomic<bool> OverallWarnings = false;
+	std::atomic<int32> NumWarnings = 0;
 	FGraphEventArray ErrorTasks;
 	FCriticalSection ErrorTasksCS;
 
@@ -66,6 +69,26 @@ namespace UE::UnrealHeaderTool::Exceptions::Private
 
 		FResults::SetResult(InResult);
 	}
+
+	void LogWarningInternal(const FString& Filename, int32 Line, const FString& Message)
+	{
+		TGuardValue<ELogTimes::Type> DisableLogTimes(GPrintLogTimes, ELogTimes::None);
+
+		FString FormattedErrorMessage;
+		if (Filename.IsEmpty())
+		{
+			FormattedErrorMessage = FString::Printf(TEXT("Warning: %s\r\n"), *Message);
+		}
+		else
+		{
+			FormattedErrorMessage = FString::Printf(TEXT("%s(%d): Warning: %s\r\n"), *Filename, Line, *Message);
+		}
+
+		UE_LOG(LogCompile, Log, TEXT("%s"), *FormattedErrorMessage);
+		GWarn->Log(ELogVerbosity::Warning, FormattedErrorMessage);
+
+		FResults::MarkWarning();
+	}
 }
 
 void FResults::LogError(FString&& Filename, int32 Line, const FString& Message, ECompilationResult::Type InResult/* = ECompilationResult::OtherCompilationError*/)
@@ -122,30 +145,41 @@ void FResults::LogError(const FUnrealSourceFile& SourceFile, int32 Line, const T
 	LogError(MoveTemp(AbsFilename), Line, ErrorMsg, InResult);
 }
 
-void FResults::LogError(const UObject& Object, const TCHAR* ErrorMsg, ECompilationResult::Type InResult/* = ECompilationResult::OtherCompilationError*/)
+void FResults::LogError(const FUHTExceptionContext& Context, const TCHAR* ErrorMsg, ECompilationResult::Type InResult/* = ECompilationResult::OtherCompilationError*/)
 {
-	if (const UField* Field = Cast<UField>(&Object))
-	{
-		if (TSharedRef<FUnrealTypeDefinitionInfo>* TypeDef = GTypeDefinitionInfoMap.Find(Field))
-		{
-			LogError((*TypeDef)->GetUnrealSourceFile(), (*TypeDef)->GetLineNumber(), ErrorMsg, InResult);
-		}
-	}
-	LogError(ErrorMsg, InResult);
+	LogError(Context.GetFilename(), Context.GetLineNumber(), ErrorMsg, InResult);
+}
+
+void FResults::LogError(const TCHAR* ErrorMsg, ECompilationResult::Type InResult/* = ECompilationResult::OtherCompilationError*/)
+{
+	LogError(TEXT(""), 1, ErrorMsg, InResult);
 }
 
-void FResults::LogError(FUnrealTypeDefinitionInfo& InTypeDef, const TCHAR* ErrorMsg, ECompilationResult::Type InResult/* = ECompilationResult::OtherCompilationError*/)
+void FResults::LogWarning(FString&& Filename, int32 Line, const FString& Message)
 {
-	if (InTypeDef.HasSource())
+	using namespace UE::UnrealHeaderTool::Exceptions::Private;
+
+	if (IsInGameThread())
 	{
-		LogError(InTypeDef.GetUnrealSourceFile(), InTypeDef.GetLineNumber(), ErrorMsg, InResult);
+		LogWarningInternal(Filename, Line, Message);
+	}
+	else
+	{
+		auto LogWarningTask = [Filename = MoveTemp(Filename), Line, Message]()
+		{
+			LogWarningInternal(Filename, Line, Message);
+		};
+
+		FGraphEventRef EventRef = FFunctionGraphTask::CreateAndDispatchWhenReady(MoveTemp(LogWarningTask), TStatId(), nullptr, ENamedThreads::GameThread);
+
+		FScopeLock Lock(&ErrorTasksCS);
+		ErrorTasks.Add(EventRef);
 	}
-	LogError(ErrorMsg, InResult);
 }
 
-void FResults::LogError(const TCHAR* ErrorMsg, ECompilationResult::Type InResult/* = ECompilationResult::OtherCompilationError*/)
+void FResults::LogWarning(const FUHTExceptionContext& Context, const TCHAR* ErrorMsg)
 {
-	LogError(TEXT(""), 1, ErrorMsg, InResult);
+	LogWarning(Context.GetFilename(), Context.GetLineNumber(), ErrorMsg);
 }
 
 void FResults::WaitForErrorTasks()
@@ -183,6 +217,13 @@ ECompilationResult::Type FResults::GetResults()
 	return OverallResults;
 }
 
+void FResults::MarkWarning()
+{
+	using namespace UE::UnrealHeaderTool::Exceptions::Private;
+	OverallWarnings = true;
+	++NumWarnings;
+}
+
 ECompilationResult::Type FResults::GetOverallResults()
 {
 	using namespace UE::UnrealHeaderTool::Exceptions::Private;
diff --git a/Engine/Source/Programs/UnrealHeaderTool/Private/Exceptions.h b/Engine/Source/Programs/UnrealHeaderTool/Private/Exceptions.h
--- a/Engine/Source/Programs/UnrealHeaderTool/Private/Exceptions.h
+++ b/Engine/Source/Programs/UnrealHeaderTool/Private/Exceptions.h
@@ -7,9 +7,17 @@
 #include "Misc/CompilationResult.h"
 #include "ProfilingDebugging/ScopedTimers.h"
 
+class FBaseParser;
 class FUnrealSourceFile;
 class FUnrealTypeDefinitionInfo;
 
+class FUHTExceptionContext
+{
+public:
+	virtual FString GetFilename() const = 0;
+	virtual int32 GetLineNumber() const = 0;
+};
+
 class FUHTException
 {
 public:
@@ -74,29 +82,39 @@ public:
 
 	/**
 	 * Generate an exception
-	 * @param InTypeDef The type definition generating the exception.  The filename and line number will be retrieved from the type definition if possible.
+	 * @param InContext The context of the exception
+	 * @param InText The text of the error
+	 */
+	UE_NORETURN static void Throwf(const FUHTExceptionContext& InContext, FString&& InText)
+	{
+		throw FUHTException(ECompilationResult::OtherCompilationError, InContext, MoveTemp(InText));
+	}
+
+	/**
+	 * Generate an exception
+	 * @param InContext The context of the exception
 	 * @param InFmt The format string
 	 * @param InArgs Arguments supplied to the format string
 	 */
 	template <typename FmtType, typename... Types>
-	UE_NORETURN static void VARARGS Throwf(FUnrealTypeDefinitionInfo& InTypeDef, const FmtType& InFmt, Types... InArgs)
+	UE_NORETURN static void VARARGS Throwf(const FUHTExceptionContext& InContext, const FmtType& InFmt, Types... InArgs)
 	{
 		FString ResultString = FString::Printf(InFmt, InArgs ...);
-		throw FUHTException(ECompilationResult::OtherCompilationError, InTypeDef, MoveTemp(ResultString));
+		throw FUHTException(ECompilationResult::OtherCompilationError, InContext, MoveTemp(ResultString));
 	}
 
 	/**
 	 * Generate an exception
 	 * @param InResult Result code to be returned as the overall result of the compilation process
-	 * @param InTypeDef The type definition generating the exception.  The filename and line number will be retrieved from the type definition if possible.
+	 * @param InContext The context of the exception
 	 * @param InFmt The format string
 	 * @param InArgs Arguments supplied to the format string
 	 */
 	template <typename FmtType, typename... Types>
-	UE_NORETURN static void VARARGS Throwf(ECompilationResult::Type InResult, FUnrealTypeDefinitionInfo& InTypeDef, const FmtType& InFmt, Types... InArgs)
+	UE_NORETURN static void VARARGS Throwf(ECompilationResult::Type InResult, const FUHTExceptionContext& InContext, const FmtType& InFmt, Types... InArgs)
 	{
 		FString ResultString = FString::Printf(InFmt, InArgs ...);
-		throw FUHTException(InResult, InTypeDef, MoveTemp(ResultString));
+		throw FUHTException(InResult, InContext, MoveTemp(ResultString));
 	}
 
 	/**
@@ -139,8 +157,7 @@ private:
 
 	FUHTException(ECompilationResult::Type InResult, FString&& InFilename, int32 InLine, FString&& InMessage);
 	FUHTException(ECompilationResult::Type InResult, const FUnrealSourceFile& SourceFile, int32 InLine, FString&& InMessage);
-	FUHTException(ECompilationResult::Type InResult, FUnrealTypeDefinitionInfo& TypeDef, FString&& InMessage);
-	FUHTException(ECompilationResult::Type InResult, const UObject* Object, FString&& InMessage);
+	FUHTException(ECompilationResult::Type InResult, const FUHTExceptionContext& InContext, FString&& InMessage);
 };
 
 /** Helper methods for working with exceptions and compilation results */
@@ -171,6 +188,11 @@ struct FResults
 	 */
 	static ECompilationResult::Type GetResults();
 
+	/**
+	 * Mark that a warning has happened
+	 */
+	static void MarkWarning();
+
 	/**
 	 * Get the overall results to be returned from compilation.
 	 */
@@ -207,22 +229,14 @@ struct FResults
 	 */
 	static void LogError(const FUnrealSourceFile& SourceFile, int32 Line, const TCHAR* ErrorMsg, ECompilationResult::Type Result = ECompilationResult::OtherCompilationError);
 
-	/**
-	 * Log an error for the given source file where the object is defined
-	 * @param Object The source file and line number will be determined from the object
-	 * @param ErrorMsg The text of the error
-	 * @param Result Compilation result of the error
-	 */
-	static void LogError(const UObject& Object, const TCHAR* ErrorMsg, ECompilationResult::Type Result = ECompilationResult::OtherCompilationError);
-
 	/**
 	 * Log an error for the given source file where the type is defined
-	 * @param InTypeDef The source file and line number will be determined from the type
+	 * @param Context The context of the exception
 	 * @param ErrorMsg The text of the error
 	 * @param Result Compilation result of the error
 	 */
-	static void LogError(FUnrealTypeDefinitionInfo& InTypeDef, const TCHAR* ErrorMsg, ECompilationResult::Type Result = ECompilationResult::OtherCompilationError);
-	
+	static void LogError(const FUHTExceptionContext& Context, const TCHAR* ErrorMsg, ECompilationResult::Type Result = ECompilationResult::OtherCompilationError);
+
 	/**
 	 * Log an error without any source file information
 	 * @param ErrorMsg The text of the error
@@ -230,6 +244,23 @@ struct FResults
 	 */
 	static void LogError(const TCHAR* ErrorMsg, ECompilationResult::Type Result = ECompilationResult::OtherCompilationError);
 
+	/**
+	 * Log a warning.
+	 * @param Filename The filename generating the warning.  If empty, then no file and line number are included in the warning.
+	 * @param Line Line number of the warning
+	 * @param Message Message body of the warning
+	 * @param Result Compilation result of the warning
+	 */
+	static void LogWarning(FString&& Filename, int32 Line, const FString& Message);
+
+	/**
+	 * Log a warning for the given source file where the type is defined
+	 * @param Context The context of the warning
+	 * @param ErrorMsg The text of the warning
+	 * @param Result Compilation result of the warning
+	 */
+	static void LogWarning(const FUHTExceptionContext& Context, const TCHAR* ErrorMsg);
+
 	/**
 	 * Invoke the given lambda in a try block catching all supported exception types.
 	 * @param SourceFile The source file being processed
@@ -287,6 +318,31 @@ struct FResults
 		}
 	}
 
+	/**
+	 * Invoke the given lambda in a try block catching all supported exception types.
+	 * @param InLambda The code to be executed in the try block
+	 */
+	template<typename Lambda>
+	static void TryAlways(Lambda&& InLambda)
+	{
+#if !PLATFORM_EXCEPTIONS_DISABLED
+		try
+#endif
+		{
+			InLambda();
+		}
+#if !PLATFORM_EXCEPTIONS_DISABLED
+		catch (const FUHTException& Ex)
+		{
+			LogError(Ex);
+		}
+		catch (const TCHAR* ErrorMsg)
+		{
+			LogError(ErrorMsg);
+		}
+#endif
+	}
+
 	/**
 	 * Invoke the given lambda in a try block catching all supported exception types.
 	 * @param InLambda The code to be executed in the try block
@@ -303,3 +359,6 @@ struct FResults
 		return DeltaTime;
 	}
 };
+
+#define UE_LOG_WARNING_UHT(Context, Format, ...) { FResults::LogWarning(Context, *FString::Printf(Format, ##__VA_ARGS__)); }
+#define UE_LOG_ERROR_UHT(Context, Format, ...) { FResults::LogError(Context, *FString::Printf(Format, ##__VA_ARGS__)); }
diff --git a/Engine/Source/Programs/UnrealHeaderTool/Private/HeaderParser.cpp b/Engine/Source/Programs/UnrealHeaderTool/Private/HeaderParser.cpp
--- a/Engine/Source/Programs/UnrealHeaderTool/Private/HeaderParser.cpp
+++ b/Engine/Source/Programs/UnrealHeaderTool/Private/HeaderParser.cpp
@@ -42,33 +42,6 @@ extern FUnrealClassDefinitionInfo* GUInterfaceDef;
 	Constants & declarations.
 -----------------------------------------------------------------------------*/
 
-/**
- * Data struct that annotates source files that failed during parsing.
- */
-class FFailedFilesAnnotation
-{
-public:
-	/**
-	 * Gets annotation state for given source file.
-	 */
-	bool Get(FUnrealSourceFile* SourceFile) const
-	{
-		return AnnotatedSet.Contains(SourceFile);
-	}
-
-	/**
-	 * Sets annotation state to true for given source file.
-	 */
-	void Set(FUnrealSourceFile* SourceFile)
-	{
-		AnnotatedSet.Add(SourceFile);
-	}
-
-private:
-	// Annotation set.
-	TSet<FUnrealSourceFile*> AnnotatedSet;
-} static FailedFilesAnnotation;
-
 enum {MAX_ARRAY_SIZE=2048};
 
 namespace
@@ -239,7 +212,7 @@ namespace
 	 * @param FuncInfo function info for the current function
 	 * @param Identifiers identifiers inside the net service declaration
 	 */
-	void ParseNetServiceIdentifiers(FFuncInfo& FuncInfo, const TArray<FString>& Identifiers)
+	void ParseNetServiceIdentifiers(FHeaderParser& HeaderParser, FFuncInfo& FuncInfo, const TArray<FString>& Identifiers)
 	{
 		static const TCHAR IdTag         [] = TEXT("Id");
 		static const TCHAR ResponseIdTag [] = TEXT("ResponseId");
@@ -258,7 +231,7 @@ namespace
 					int32 TempInt = FCString::Atoi(Equals + 1);
 					if (TempInt <= 0 || TempInt > MAX_uint16)
 					{
-						FError::Throwf(TEXT("Invalid network identifier %s for function"), IdentifierPtr);
+						FUHTException::Throwf(HeaderParser, TEXT("Invalid network identifier %s for function"), IdentifierPtr);
 					}
 					FuncInfo.RPCId = (uint16)TempInt;
 				}
@@ -268,7 +241,7 @@ namespace
 					int32 TempInt = FCString::Atoi(Equals + 1);
 					if (TempInt <= 0 || TempInt > MAX_uint16)
 					{
-						FError::Throwf(TEXT("Invalid network identifier %s for function"), IdentifierPtr);
+						FUHTException::Throwf(HeaderParser, TEXT("Invalid network identifier %s for function"), IdentifierPtr);
 					}
 					FuncInfo.RPCResponseId = (uint16)TempInt;
 				}
@@ -279,7 +252,7 @@ namespace
 
 				if (FuncInfo.EndpointName.Len())
 				{
-					FError::Throwf(TEXT("Function should not specify multiple endpoints - '%s' found but already using '%s'"), *Identifier);
+					FUHTException::Throwf(HeaderParser, TEXT("Function should not specify multiple endpoints - '%s' found but already using '%s'"), *Identifier);
 				}
 
 				FuncInfo.EndpointName = Identifier;
@@ -293,7 +266,7 @@ namespace
 	 * @param FuncInfo   - The FFuncInfo object to populate.
 	 * @param Specifiers - The specifiers to process.
 	 */
-	void ProcessFunctionSpecifiers(FFuncInfo& FuncInfo, const TArray<FPropertySpecifier>& Specifiers, TMap<FName, FString>& MetaData)
+	void ProcessFunctionSpecifiers(FHeaderParser& HeaderParser, FFuncInfo& FuncInfo, const TArray<FPropertySpecifier>& Specifiers, TMap<FName, FString>& MetaData)
 	{
 		bool bSpecifiedUnreliable = false;
 		bool bSawPropertyAccessor = false;
@@ -304,7 +277,7 @@ namespace
 			{
 				default:
 				{
-					FError::Throwf(TEXT("Unknown function specifier '%s'"), *Specifier.Key);
+					FUHTException::Throwf(HeaderParser, TEXT("Unknown function specifier '%s'"), *Specifier.Key);
 				}
 				break;
 
@@ -312,20 +285,20 @@ namespace
 				{
 					if (FuncInfo.FunctionFlags & FUNC_Net)
 					{
-						UE_LOG_ERROR_UHT(TEXT("BlueprintNativeEvent functions cannot be replicated!") );
+						UE_LOG_ERROR_UHT(HeaderParser, TEXT("BlueprintNativeEvent functions cannot be replicated!") );
 					}
 					else if ( (FuncInfo.FunctionFlags & FUNC_BlueprintEvent) && !(FuncInfo.FunctionFlags & FUNC_Native) )
 					{
 						// already a BlueprintImplementableEvent
-						UE_LOG_ERROR_UHT(TEXT("A function cannot be both BlueprintNativeEvent and BlueprintImplementableEvent!") );
+						UE_LOG_ERROR_UHT(HeaderParser, TEXT("A function cannot be both BlueprintNativeEvent and BlueprintImplementableEvent!") );
 					}
 					else if (bSawPropertyAccessor)
 					{
-						UE_LOG_ERROR_UHT(TEXT("A function cannot be both BlueprintNativeEvent and a Blueprint Property accessor!"));
+						UE_LOG_ERROR_UHT(HeaderParser, TEXT("A function cannot be both BlueprintNativeEvent and a Blueprint Property accessor!"));
 					}
 					else if ( (FuncInfo.FunctionFlags & FUNC_Private) )
 					{
-						UE_LOG_ERROR_UHT(TEXT("A Private function cannot be a BlueprintNativeEvent!") );
+						UE_LOG_ERROR_UHT(HeaderParser, TEXT("A Private function cannot be a BlueprintNativeEvent!") );
 					}
 
 					FuncInfo.FunctionFlags |= FUNC_Event;
@@ -337,20 +310,20 @@ namespace
 				{
 					if (FuncInfo.FunctionFlags & FUNC_Net)
 					{
-						UE_LOG_ERROR_UHT(TEXT("BlueprintImplementableEvent functions cannot be replicated!") );
+						UE_LOG_ERROR_UHT(HeaderParser, TEXT("BlueprintImplementableEvent functions cannot be replicated!") );
 					}
 					else if ( (FuncInfo.FunctionFlags & FUNC_BlueprintEvent) && (FuncInfo.FunctionFlags & FUNC_Native) )
 					{
 						// already a BlueprintNativeEvent
-						UE_LOG_ERROR_UHT(TEXT("A function cannot be both BlueprintNativeEvent and BlueprintImplementableEvent!") );
+						UE_LOG_ERROR_UHT(HeaderParser, TEXT("A function cannot be both BlueprintNativeEvent and BlueprintImplementableEvent!") );
 					}
 					else if (bSawPropertyAccessor)
 					{
-						UE_LOG_ERROR_UHT(TEXT("A function cannot be both BlueprintImplementableEvent and a Blueprint Property accessor!"));
+						UE_LOG_ERROR_UHT(HeaderParser, TEXT("A function cannot be both BlueprintImplementableEvent and a Blueprint Property accessor!"));
 					}
 					else if ( (FuncInfo.FunctionFlags & FUNC_Private) )
 					{
-						UE_LOG_ERROR_UHT(TEXT("A Private function cannot be a BlueprintImplementableEvent!") );
+						UE_LOG_ERROR_UHT(HeaderParser, TEXT("A Private function cannot be a BlueprintImplementableEvent!") );
 					}
 
 					FuncInfo.FunctionFlags |= FUNC_Event;
@@ -364,7 +337,7 @@ namespace
 					FuncInfo.FunctionFlags |= FUNC_Exec;
 					if( FuncInfo.FunctionFlags & FUNC_Net )
 					{
-						UE_LOG_ERROR_UHT(TEXT("Exec functions cannot be replicated!") );
+						UE_LOG_ERROR_UHT(HeaderParser, TEXT("Exec functions cannot be replicated!") );
 					}
 				}
 				break;
@@ -379,7 +352,7 @@ namespace
 				{
 					if ((FuncInfo.FunctionFlags & FUNC_BlueprintEvent) != 0)
 					{
-						FError::Throwf(TEXT("BlueprintImplementableEvent or BlueprintNativeEvent functions cannot be declared as Client or Server"));
+						FUHTException::Throwf(HeaderParser, TEXT("BlueprintImplementableEvent or BlueprintNativeEvent functions cannot be declared as Client or Server"));
 					}
 
 					FuncInfo.FunctionFlags |= FUNC_Net;
@@ -392,7 +365,7 @@ namespace
 
 					if( FuncInfo.FunctionFlags & FUNC_Exec )
 					{
-						UE_LOG_ERROR_UHT(TEXT("Exec functions cannot be replicated!") );
+						UE_LOG_ERROR_UHT(HeaderParser, TEXT("Exec functions cannot be replicated!") );
 					}
 				}
 				break;
@@ -401,7 +374,7 @@ namespace
 				{
 					if ((FuncInfo.FunctionFlags & FUNC_BlueprintEvent) != 0)
 					{
-						FError::Throwf(TEXT("BlueprintImplementableEvent or BlueprintNativeEvent functions cannot be declared as Client or Server"));
+						FUHTException::Throwf(HeaderParser, TEXT("BlueprintImplementableEvent or BlueprintNativeEvent functions cannot be declared as Client or Server"));
 					}
 
 					FuncInfo.FunctionFlags |= FUNC_Net;
@@ -418,7 +391,7 @@ namespace
 				{
 					if ((FuncInfo.FunctionFlags & FUNC_BlueprintEvent) != 0)
 					{
-						FError::Throwf(TEXT("BlueprintImplementableEvent or BlueprintNativeEvent functions cannot be declared as Multicast"));
+						FUHTException::Throwf(HeaderParser, TEXT("BlueprintImplementableEvent or BlueprintNativeEvent functions cannot be declared as Multicast"));
 					}
 
 					FuncInfo.FunctionFlags |= FUNC_Net;
@@ -430,7 +403,7 @@ namespace
 				{
 					if ((FuncInfo.FunctionFlags & FUNC_BlueprintEvent) != 0)
 					{
-						FError::Throwf(TEXT("BlueprintImplementableEvent or BlueprintNativeEvent functions cannot be declared as a ServiceRequest"));
+						FUHTException::Throwf(HeaderParser, TEXT("BlueprintImplementableEvent or BlueprintNativeEvent functions cannot be declared as a ServiceRequest"));
 					}
 
 					FuncInfo.FunctionFlags |= FUNC_Net;
@@ -438,11 +411,11 @@ namespace
 					FuncInfo.FunctionFlags |= FUNC_NetRequest;
 					FuncInfo.FunctionExportFlags |= FUNCEXPORT_CustomThunk;
 
-					ParseNetServiceIdentifiers(FuncInfo, Specifier.Values);
+					ParseNetServiceIdentifiers(HeaderParser, FuncInfo, Specifier.Values);
 
 					if (FuncInfo.EndpointName.Len() == 0)
 					{
-						FError::Throwf(TEXT("ServiceRequest needs to specify an endpoint name"));
+						FUHTException::Throwf(HeaderParser, TEXT("ServiceRequest needs to specify an endpoint name"));
 					}
 				}
 				break;
@@ -451,18 +424,18 @@ namespace
 				{
 					if ((FuncInfo.FunctionFlags & FUNC_BlueprintEvent) != 0)
 					{
-						FError::Throwf(TEXT("BlueprintImplementableEvent or BlueprintNativeEvent functions cannot be declared as a ServiceResponse"));
+						FUHTException::Throwf(HeaderParser, TEXT("BlueprintImplementableEvent or BlueprintNativeEvent functions cannot be declared as a ServiceResponse"));
 					}
 
 					FuncInfo.FunctionFlags |= FUNC_Net;
 					FuncInfo.FunctionFlags |= FUNC_NetReliable;
 					FuncInfo.FunctionFlags |= FUNC_NetResponse;
 
-					ParseNetServiceIdentifiers(FuncInfo, Specifier.Values);
+					ParseNetServiceIdentifiers(HeaderParser, FuncInfo, Specifier.Values);
 
 					if (FuncInfo.EndpointName.Len() == 0)
 					{
-						FError::Throwf(TEXT("ServiceResponse needs to specify an endpoint name"));
+						FUHTException::Throwf(HeaderParser, TEXT("ServiceResponse needs to specify an endpoint name"));
 					}
 				}
 				break;
@@ -495,7 +468,7 @@ namespace
 				{
 					if (FuncInfo.FunctionFlags & FUNC_Event)
 					{
-						UE_LOG_ERROR_UHT(TEXT("Function cannot be a blueprint event and a blueprint getter."));
+						UE_LOG_ERROR_UHT(HeaderParser, TEXT("Function cannot be a blueprint event and a blueprint getter."));
 					}
 
 					bSawPropertyAccessor = true;
@@ -509,7 +482,7 @@ namespace
 				{
 					if (FuncInfo.FunctionFlags & FUNC_Event)
 					{
-						UE_LOG_ERROR_UHT(TEXT("Function cannot be a blueprint event and a blueprint setter."));
+						UE_LOG_ERROR_UHT(HeaderParser, TEXT("Function cannot be a blueprint event and a blueprint setter."));
 					}
 
 					bSawPropertyAccessor = true;
@@ -578,45 +551,45 @@ namespace
 
 			if (FuncInfo.FunctionFlags & FUNC_Static)
 			{
-				UE_LOG_ERROR_UHT(TEXT("Static functions can't be replicated"));
+				UE_LOG_ERROR_UHT(HeaderParser, TEXT("Static functions can't be replicated"));
 			}
 
 			if (!bIsNetReliable && !bSpecifiedUnreliable && !bIsNetService)
 			{
-				UE_LOG_ERROR_UHT(TEXT("Replicated function: 'reliable' or 'unreliable' is required"));
+				UE_LOG_ERROR_UHT(HeaderParser, TEXT("Replicated function: 'reliable' or 'unreliable' is required"));
 			}
 
 			if (bIsNetReliable && bSpecifiedUnreliable && !bIsNetService)
 			{
-				UE_LOG_ERROR_UHT(TEXT("'reliable' and 'unreliable' are mutually exclusive"));
+				UE_LOG_ERROR_UHT(HeaderParser, TEXT("'reliable' and 'unreliable' are mutually exclusive"));
 			}
 		}
 		else if (FuncInfo.FunctionFlags & FUNC_NetReliable)
 		{
-			UE_LOG_ERROR_UHT(TEXT("'reliable' specified without 'client' or 'server'"));
+			UE_LOG_ERROR_UHT(HeaderParser, TEXT("'reliable' specified without 'client' or 'server'"));
 		}
 		else if (bSpecifiedUnreliable)
 		{
-			UE_LOG_ERROR_UHT(TEXT("'unreliable' specified without 'client' or 'server'"));
+			UE_LOG_ERROR_UHT(HeaderParser, TEXT("'unreliable' specified without 'client' or 'server'"));
 		}
 
 		if (FuncInfo.bSealedEvent && !(FuncInfo.FunctionFlags & FUNC_Event))
 		{
-			UE_LOG_ERROR_UHT(TEXT("SealedEvent may only be used on events"));
+			UE_LOG_ERROR_UHT(HeaderParser, TEXT("SealedEvent may only be used on events"));
 		}
 
 		if (FuncInfo.bSealedEvent && FuncInfo.FunctionFlags & FUNC_BlueprintEvent)
 		{
-			UE_LOG_ERROR_UHT(TEXT("SealedEvent cannot be used on Blueprint events"));
+			UE_LOG_ERROR_UHT(HeaderParser, TEXT("SealedEvent cannot be used on Blueprint events"));
 		}
 
 		if (FuncInfo.bForceBlueprintImpure && (FuncInfo.FunctionFlags & FUNC_BlueprintPure) != 0)
 		{
-			UE_LOG_ERROR_UHT(TEXT("BlueprintPure (or BlueprintPure=true) and BlueprintPure=false should not both appear on the same function, they are mutually exclusive"));
+			UE_LOG_ERROR_UHT(HeaderParser, TEXT("BlueprintPure (or BlueprintPure=true) and BlueprintPure=false should not both appear on the same function, they are mutually exclusive"));
 		}
 	}
 
-	const TCHAR* GetHintText(EVariableCategory VariableCategory)
+	const TCHAR* GetHintText(FHeaderParser& HeaderParser, EVariableCategory VariableCategory)
 	{
 		switch (VariableCategory)
 		{
@@ -631,7 +604,7 @@ namespace
 				return TEXT("Member variable declaration");
 
 			default:
-				FError::Throwf(TEXT("Unknown variable category"));
+				FUHTException::Throwf(HeaderParser, TEXT("Unknown variable category"));
 		}
 
 		// Unreachable
@@ -670,13 +643,13 @@ namespace
 		FToken Token;
 		if (Parser.GetToken(Token) && (Token.ConstantType != CPT_Float || Token.TokenType != TOKEN_Const))
 		{
-			FError::Throwf(TEXT("Expected engine version in %s macro"), MacroToken.Identifier);
+			FUHTException::Throwf(Parser, TEXT("Expected engine version in %s macro"), MacroToken.Identifier);
 		}
 
 		Parser.RequireSymbol(TEXT(','), ErrorMessageGetter);
 		if (Parser.GetToken(Token) && (Token.ConstantType != CPT_String || Token.TokenType != TOKEN_Const))
 		{
-			FError::Throwf(TEXT("Expected deprecation message in %s macro"), MacroToken.Identifier);
+			FUHTException::Throwf(Parser, TEXT("Expected deprecation message in %s macro"), MacroToken.Identifier);
 		}
 
 		Parser.RequireSymbol(TEXT(')'), ErrorMessageGetter);
@@ -708,7 +681,7 @@ namespace
 			return EPointerMemberBehavior::AllowAndLog;
 		}
 
-		FError::Throwf(TEXT("Unrecognized native pointer member behavior: %s"), *InString );
+		FUHTException::Throwf(TEXT("Configuration"), 1, TEXT("Unrecognized native pointer member behavior: %s"), *InString );
 		return EPointerMemberBehavior::Disallow;
 	}
 
@@ -794,16 +767,6 @@ const FUHTConfig& FUHTConfig::Get()
 /////////////////////////////////////////////////////
 // FHeaderParser
 
-FString FHeaderParser::GetSourceFileContext() const
-{
-	return Filename;
-}
-
-FString FHeaderParser::GetContext()
-{
-	return FString::Printf(TEXT("%s(%i)"), *GetSourceFileContext(), InputLine);
-}
-
 /*-----------------------------------------------------------------------------
 	Code emitting.
 -----------------------------------------------------------------------------*/
@@ -826,10 +789,10 @@ FUnrealClassDefinitionInfo* FHeaderParser::GetQualifiedClass(const TCHAR* Thing)
 
 	if (!ClassName[0])
 	{
-		FError::Throwf(TEXT("%s: Missing class name"), Thing );
+		FUHTException::Throwf(*this, TEXT("%s: Missing class name"), Thing);
 	}
 
-	return FUnrealClassDefinitionInfo::FindScriptClassOrThrow(ClassName);
+	return FUnrealClassDefinitionInfo::FindScriptClassOrThrow(*this, ClassName);
 }
 
 /*-----------------------------------------------------------------------------
@@ -849,7 +812,8 @@ FUnrealClassDefinitionInfo* FHeaderParser::GetQualifiedClass(const TCHAR* Thing)
  */
 FUnrealFunctionDefinitionInfo* FHeaderParser::FindFunction
 (
-	const FUnrealStructDefinitionInfo&  InScope,
+
+	const FUnrealStructDefinitionInfo& InScope,
 	const TCHAR*	InIdentifier,
 	bool			bIncludeParents,
 	const TCHAR*	Thing
@@ -867,7 +831,7 @@ FUnrealFunctionDefinitionInfo* FHeaderParser::FindFunction
 				{
 					if (Thing)
 					{
-						FError::Throwf(TEXT("%s: expecting function or delegate, got a property"), Thing);
+						FUHTException::Throwf(InScope, TEXT("%s: expecting function or delegate, got a property"), Thing);
 					}
 					return nullptr;
 				}
@@ -904,7 +868,7 @@ FUnrealPropertyDefinitionInfo* FHeaderParser::FindProperty(const FUnrealStructDe
 				{
 					if (Thing)
 					{
-						FError::Throwf(TEXT("%s: expecting a property, got a function or delegate"), Thing);
+						FUHTException::Throwf(InScope, TEXT("%s: expecting a property, got a function or delegate"), Thing);
 					}
 					return nullptr;
 				}
@@ -989,7 +953,7 @@ FUnrealEnumDefinitionInfo& FHeaderParser::CompileEnum()
 	EEnumFlags      Flags         = EEnumFlags::None;
 	if (!GetIdentifier(EnumToken))
 	{
-		FError::Throwf(TEXT("Missing identifier after UENUM()") );
+		FUHTException::Throwf(*this, TEXT("Missing identifier after UENUM()") );
 	}
 
 	if (EnumToken.Matches(TEXT("namespace"), ESearchCase::CaseSensitive))
@@ -1004,7 +968,7 @@ FUnrealEnumDefinitionInfo& FHeaderParser::CompileEnum()
 	{
 		if (!GetIdentifier(EnumToken))
 		{
-			FError::Throwf(TEXT("Missing identifier after enum") );
+			FUHTException::Throwf(*this, TEXT("Missing identifier after enum") );
 		}
 
 		if (EnumToken.Matches(TEXT("class"), ESearchCase::CaseSensitive) || EnumToken.Matches(TEXT("struct"), ESearchCase::CaseSensitive))
@@ -1025,13 +989,13 @@ FUnrealEnumDefinitionInfo& FHeaderParser::CompileEnum()
 	}
 	else
 	{
-		FError::Throwf(TEXT("UENUM() should be followed by \'enum\' or \'namespace\' keywords.") );
+		FUHTException::Throwf(*this, TEXT("UENUM() should be followed by \'enum\' or \'namespace\' keywords.") );
 	}
 
 	// Get enumeration name.
 	if (!bReadEnumName)
 	{
-		FError::Throwf(TEXT("Missing enumeration name") );
+		FUHTException::Throwf(*this, TEXT("Missing enumeration name") );
 	}
 
 	ParseFieldMetaData(EnumMetaData, EnumToken.Identifier);
@@ -1046,7 +1010,7 @@ FUnrealEnumDefinitionInfo& FHeaderParser::CompileEnum()
 		switch ((EEnumSpecifier)Algo::FindSortedStringCaseInsensitive(*Specifier.Key, GEnumSpecifierStrings))
 		{
 		default:
-			FError::Throwf(TEXT("Unknown enum specifier '%s'"), *Specifier.Key);
+			FUHTException::Throwf(*this, TEXT("Unknown enum specifier '%s'"), *Specifier.Key);
 
 		case EEnumSpecifier::Flags:
 			Flags |= EEnumFlags::Flags;
@@ -1071,7 +1035,7 @@ FUnrealEnumDefinitionInfo& FHeaderParser::CompileEnum()
 			FToken BaseToken;
 			if (!GetIdentifier(BaseToken))
 			{
-				FError::Throwf(TEXT("Missing enum base") );
+				FUHTException::Throwf(*this, TEXT("Missing enum base") );
 			}
 
 			if (!FCString::Strcmp(BaseToken.Identifier, TEXT("uint8")))
@@ -1108,7 +1072,7 @@ FUnrealEnumDefinitionInfo& FHeaderParser::CompileEnum()
 			}
 			else
 			{
-				FError::Throwf(TEXT("Unsupported enum class base type: %s"), BaseToken.Identifier);
+				FUHTException::Throwf(*this, TEXT("Unsupported enum class base type: %s"), BaseToken.Identifier);
 			}
 		}
 		else
@@ -1122,13 +1086,13 @@ FUnrealEnumDefinitionInfo& FHeaderParser::CompileEnum()
 	{
 		if (EnumHasAnyFlags(Flags, EEnumFlags::Flags))
 		{
-			FError::Throwf(TEXT("The 'Flags' specifier can only be used on enum classes"));
+			FUHTException::Throwf(*this, TEXT("The 'Flags' specifier can only be used on enum classes"));
 		}
 	}
 
 	if (UnderlyingType != EUnderlyingEnumType::uint8 && EnumMetaData.Contains(FHeaderParserNames::NAME_BlueprintType))
 	{
-		FError::Throwf(TEXT("Invalid BlueprintType enum base - currently only uint8 supported"));
+		FUHTException::Throwf(*this, TEXT("Invalid BlueprintType enum base - currently only uint8 supported"));
 	}
 
 	// Get opening brace.
@@ -1146,7 +1110,7 @@ FUnrealEnumDefinitionInfo& FHeaderParser::CompileEnum()
 			FToken InnerEnumToken;
 			if (!GetIdentifier(InnerEnumToken))
 			{
-				FError::Throwf(TEXT("Missing enumeration name") );
+				FUHTException::Throwf(*this, TEXT("Missing enumeration name") );
 			}
 
 			EnumDef.SetCppType(FString::Printf(TEXT("%s::%s"), EnumToken.Identifier, InnerEnumToken.Identifier));
@@ -1186,7 +1150,7 @@ FUnrealEnumDefinitionInfo& FHeaderParser::CompileEnum()
 			FToken InitToken;
 			if (!GetToken(InitToken))
 			{
-				FError::Throwf(TEXT("UENUM: missing enumerator initializer"));
+				FUHTException::Throwf(*this, TEXT("UENUM: missing enumerator initializer"));
 			}
 
 			int64 NewEnumValue = -1;
@@ -1201,7 +1165,7 @@ FUnrealEnumDefinitionInfo& FHeaderParser::CompileEnum()
 			{
 				if (!GetToken(InitToken))
 				{
-					FError::Throwf(TEXT("Enumerator: end of file encountered while parsing the initializer"));
+					FUHTException::Throwf(*this, TEXT("Enumerator: end of file encountered while parsing the initializer"));
 				}
 
 				if (InitToken.TokenType == TOKEN_Symbol)
@@ -1279,7 +1243,7 @@ FUnrealEnumDefinitionInfo& FHeaderParser::CompileEnum()
 			FToken ClosingBrace;
 			if (!GetToken(ClosingBrace))
 			{
-				FError::Throwf(TEXT("UENUM: end of file encountered"));
+				FUHTException::Throwf(*this, TEXT("UENUM: end of file encountered"));
 			}
 
 			if (ClosingBrace.TokenType == TOKEN_Symbol && ClosingBrace.Matches(TEXT('}')))
@@ -1307,10 +1271,10 @@ FUnrealEnumDefinitionInfo& FHeaderParser::CompileEnum()
 		const int32 MaxEnumItemIndex = EnumDef.GetIndexByName(MaxEnumItem);
 		if (MaxEnumItemIndex != INDEX_NONE)
 		{
-			FError::Throwf(TEXT("Illegal enumeration tag specified.  Conflicts with auto-generated tag '%s'"), *MaxEnumItem.ToString());
+			FUHTException::Throwf(*this, TEXT("Illegal enumeration tag specified.  Conflicts with auto-generated tag '%s'"), *MaxEnumItem.ToString());
 		}
 
-		FError::Throwf(TEXT("Unable to generate enum MAX entry '%s' due to name collision"), *MaxEnumItem.ToString());
+		FUHTException::Throwf(*this, TEXT("Unable to generate enum MAX entry '%s' due to name collision"), *MaxEnumItem.ToString());
 	}
 
 	CheckDocumentationPolicyForEnum(EnumDef, EnumValueMetaData, EntryMetaData);
@@ -1326,7 +1290,7 @@ FUnrealEnumDefinitionInfo& FHeaderParser::CompileEnum()
 
 	if (!EnumDef.IsValidEnumValue(0) && EnumMetaData.Contains(FHeaderParserNames::NAME_BlueprintType))
 	{
-		UE_LOG_WARNING_UHT(TEXT("'%s' does not have a 0 entry! (This is a problem when the enum is initalized by default)"), *EnumDef.GetName());
+		UE_LOG_WARNING_UHT(EnumDef, TEXT("'%s' does not have a 0 entry! (This is a problem when the enum is initalized by default)"), *EnumDef.GetName());
 	}
 
 	return EnumDef;
@@ -1803,7 +1767,7 @@ FUnrealScriptStructDefinitionInfo& FHeaderParser::CompileStructDeclaration()
 		{
 			default:
 			{
-				FError::Throwf(TEXT("Unknown struct specifier '%s'"), *Specifier.Key);
+				FUHTException::Throwf(*this, TEXT("Unknown struct specifier '%s'"), *Specifier.Key);
 			}
 			break;
 
@@ -1827,7 +1791,7 @@ FUnrealScriptStructDefinitionInfo& FHeaderParser::CompileStructDeclaration()
 
 				if (!FPaths::IsSamePath(Filename, GUObjectDef->GetUnrealSourceFile().GetFilename()))
 				{
-					UE_LOG_ERROR_UHT(TEXT("Immutable is being phased out in favor of SerializeNative, and is only legal on the mirror structs declared in UObject"));
+					UE_LOG_ERROR_UHT(*this, TEXT("Immutable is being phased out in favor of SerializeNative, and is only legal on the mirror structs declared in UObject"));
 				}
 			}
 			break;
@@ -1880,7 +1844,7 @@ FUnrealScriptStructDefinitionInfo& FHeaderParser::CompileStructDeclaration()
 		// If the struct still wasn't found, throw an error
 		if (BaseStructDef == NULL)
 		{
-			FError::Throwf(TEXT("'struct': Can't find struct '%s'"), *ParentStructNameInScript);
+			FUHTException::Throwf(*this, TEXT("'struct': Can't find struct '%s'"), *ParentStructNameInScript);
 		}
 		else
 		{
@@ -1890,7 +1854,7 @@ FUnrealScriptStructDefinitionInfo& FHeaderParser::CompileStructDeclaration()
 				const TCHAR* PrefixCPP = UHTConfig.StructsWithTPrefix.Contains(ParentStructNameStripped) ? TEXT("T") : BaseStructDef->AsScriptStructChecked().GetPrefixCPP();
 				if (ParentStructNameInScript != FString::Printf(TEXT("%s%s"), PrefixCPP, *ParentStructNameStripped))
 				{
-					FError::Throwf(TEXT("Parent Struct '%s' is missing a valid Unreal prefix, expecting '%s'"), *ParentStructNameInScript, *FString::Printf(TEXT("%s%s"), PrefixCPP, *BaseStructDef->GetName()));
+					FUHTException::Throwf(*this, TEXT("Parent Struct '%s' is missing a valid Unreal prefix, expecting '%s'"), *ParentStructNameInScript, *FString::Printf(TEXT("%s%s"), PrefixCPP, *BaseStructDef->GetName()));
 					BaseStructDef = nullptr;
 				}
 			}
@@ -1919,14 +1883,14 @@ FUnrealScriptStructDefinitionInfo& FHeaderParser::CompileStructDeclaration()
 		FString ExpectedStructName = FString::Printf(TEXT("%s%s"), ExpectedPrefixCPP, *StructNameStripped);
 		if (StructNameInScript != ExpectedStructName)
 		{
-			FError::Throwf(TEXT("Struct '%s' has an invalid Unreal prefix, expecting '%s'"), *StructNameInScript, *ExpectedStructName);
+			FUHTException::Throwf(*this, TEXT("Struct '%s' has an invalid Unreal prefix, expecting '%s'"), *StructNameInScript, *ExpectedStructName);
 		}
 	}
 	else
 	{
 		const TCHAR* ExpectedPrefixCPP = UHTConfig.StructsWithTPrefix.Contains(StructNameInScript) ? TEXT("T") : StructDef.GetPrefixCPP();
 		FString ExpectedStructName = FString::Printf(TEXT("%s%s"), ExpectedPrefixCPP, *StructNameInScript);
-		FError::Throwf(TEXT("Struct '%s' is missing a valid Unreal prefix, expecting '%s'"), *StructNameInScript, *ExpectedStructName);
+		FUHTException::Throwf(*this, TEXT("Struct '%s' is missing a valid Unreal prefix, expecting '%s'"), *StructNameInScript, *ExpectedStructName);
 	}
 
 	Struct->StructFlags = EStructFlags(Struct->StructFlags | StructFlags);
@@ -1965,7 +1929,7 @@ FUnrealScriptStructDefinitionInfo& FHeaderParser::CompileStructDeclaration()
 		}
 		else if (Token.Matches(TEXT("UFUNCTION"), ESearchCase::CaseSensitive))
 		{
-			FError::Throwf(TEXT("USTRUCTs cannot contain UFUNCTIONs."));
+			FUHTException::Throwf(*this, TEXT("USTRUCTs cannot contain UFUNCTIONs."));
 		}
 		else if (Token.Matches(TEXT("RIGVM_METHOD"), ESearchCase::CaseSensitive))
 		{
@@ -1976,12 +1940,12 @@ FUnrealScriptStructDefinitionInfo& FHeaderParser::CompileStructDeclaration()
 			// Match 'GENERATED_USTRUCT_BODY' '(' [StructName] ')' or 'GENERATED_BODY' '(' [StructName] ')'
 			if (CurrentAccessSpecifier != ACCESS_Public)
 			{
-				FError::Throwf(TEXT("%s must be in the public scope of '%s', not private or protected."), Token.Identifier, *StructNameInScript);
+				FUHTException::Throwf(*this, TEXT("%s must be in the public scope of '%s', not private or protected."), Token.Identifier, *StructNameInScript);
 			}
 
 			if (StructDef.GetMacroDeclaredLineNumber() != INDEX_NONE)
 			{
-				FError::Throwf(TEXT("Multiple %s declarations found in '%s'"), Token.Identifier, *StructNameInScript);
+				FUHTException::Throwf(*this, TEXT("Multiple %s declarations found in '%s'"), Token.Identifier, *StructNameInScript);
 			}
 
 			StructDef.SetMacroDeclaredLineNumber(InputLine);
@@ -2006,7 +1970,7 @@ FUnrealScriptStructDefinitionInfo& FHeaderParser::CompileStructDeclaration()
 		{
 			if (CompilerDirectiveStack.Num() < 1)
 			{
-				FError::Throwf(TEXT("Unmatched '#endif' in class or global scope"));
+				FUHTException::Throwf(*this, TEXT("Unmatched '#endif' in class or global scope"));
 			}
 			CompilerDirectiveStack.Pop();
 			// Do nothing and hope that the if code below worked out OK earlier
@@ -2021,7 +1985,7 @@ FUnrealScriptStructDefinitionInfo& FHeaderParser::CompileStructDeclaration()
 			{
 				if (bInvertConditional)
 				{
-					FError::Throwf(TEXT("Cannot use !WITH_EDITORONLY_DATA"));
+					FUHTException::Throwf(*this, TEXT("Cannot use !WITH_EDITORONLY_DATA"));
 				}
 
 				PushCompilerDirective(ECompilerDirective::WithEditorOnlyData);
@@ -2030,7 +1994,7 @@ FUnrealScriptStructDefinitionInfo& FHeaderParser::CompileStructDeclaration()
 			{
 				if (bInvertConditional)
 				{
-					FError::Throwf(TEXT("Cannot use !WITH_EDITOR"));
+					FUHTException::Throwf(*this, TEXT("Cannot use !WITH_EDITOR"));
 				}
 				PushCompilerDirective(ECompilerDirective::WithEditor);
 			}
@@ -2041,7 +2005,7 @@ FUnrealScriptStructDefinitionInfo& FHeaderParser::CompileStructDeclaration()
 			}
 			else
 			{
-				FError::Throwf(TEXT("'struct': Unsupported preprocessor directive inside a struct.") );
+				FUHTException::Throwf(*this, TEXT("'struct': Unsupported preprocessor directive inside a struct.") );
 			}
 
 			if (bConsumeAsCppText)
@@ -2054,7 +2018,7 @@ FUnrealScriptStructDefinitionInfo& FHeaderParser::CompileStructDeclaration()
 
 					if ( ch==0 )
 					{
-						FError::Throwf(TEXT("Unexpected end of struct definition %s"), *Struct->GetName());
+						FUHTException::Throwf(*this, TEXT("Unexpected end of struct definition %s"), *Struct->GetName());
 					}
 					else if ( ch=='{' || (ch=='#' && (PeekIdentifier(TEXT("if"), ESearchCase::CaseSensitive) || PeekIdentifier(TEXT("ifdef"), ESearchCase::CaseSensitive))) )
 					{
@@ -2094,7 +2058,7 @@ FUnrealScriptStructDefinitionInfo& FHeaderParser::CompileStructDeclaration()
 
 				if (!SkipDeclaration(Token))
 				{
-					FError::Throwf(TEXT("'struct': Unexpected '%s'"), FirstTokenIdentifier);
+					FUHTException::Throwf(*this, TEXT("'struct': Unexpected '%s'"), FirstTokenIdentifier);
 				}	
 			}
 			else
@@ -2112,7 +2076,7 @@ FUnrealScriptStructDefinitionInfo& FHeaderParser::CompileStructDeclaration()
 	{
 		// Roll the line number back to the start of the struct body and error out
 		InputLine = SavedLineNumber;
-		FError::Throwf(TEXT("Expected a GENERATED_BODY() at the start of struct"));
+		FUHTException::Throwf(*this, TEXT("Expected a GENERATED_BODY() at the start of struct"));
 	}
 
 	// Validate sparse class data
@@ -2201,11 +2165,11 @@ void FHeaderParser::CheckAllow( const TCHAR* Thing, ENestAllowFlags AllowFlags )
 	{
 		if (TopNest->NestType == ENestType::GlobalScope)
 		{
-			FError::Throwf(TEXT("%s is not allowed before the Class definition"), Thing );
+			FUHTException::Throwf(*this, TEXT("%s is not allowed before the Class definition"), Thing );
 		}
 		else
 		{
-			FError::Throwf(TEXT("%s is not allowed here"), Thing );
+			FUHTException::Throwf(*this, TEXT("%s is not allowed here"), Thing );
 		}
 	}
 }
@@ -2224,7 +2188,7 @@ void FHeaderParser::PushNest(ENestType NestType, FUnrealStructDefinitionInfo* In
 	// Prevent overnesting.
 	if (NestLevel >= MAX_NEST_LEVELS)
 	{
-		FError::Throwf(TEXT("Maximum nesting limit exceeded"));
+		FUHTException::Throwf(*this, TEXT("Maximum nesting limit exceeded"));
 	}
 
 	// Inherit info from stack node above us.
@@ -2256,7 +2220,7 @@ void FHeaderParser::PushNest(ENestType NestType, FUnrealStructDefinitionInfo* In
 		break;
 
 	default:
-		FError::Throwf(TEXT("Internal error in PushNest, type %i"), (uint8)NestType);
+		FUHTException::Throwf(*this, TEXT("Internal error in PushNest, type %i"), (uint8)NestType);
 		break;
 	}
 }
@@ -2272,16 +2236,16 @@ void FHeaderParser::PopNest(ENestType NestType, const TCHAR* Descr)
 	// Validate the nesting state.
 	if (NestLevel <= 0)
 	{
-		FError::Throwf(TEXT("Unexpected '%s' at global scope"), Descr, NestTypeName(NestType));
+		FUHTException::Throwf(*this, TEXT("Unexpected '%s' at global scope"), Descr, NestTypeName(NestType));
 	}
 	else if (TopNest->NestType != NestType)
 	{
-		FError::Throwf(TEXT("Unexpected end of %s in '%s' block"), Descr, NestTypeName(TopNest->NestType));
+		FUHTException::Throwf(*this, TEXT("Unexpected end of %s in '%s' block"), Descr, NestTypeName(TopNest->NestType));
 	}
 
 	if (NestType != ENestType::GlobalScope && NestType != ENestType::Class && NestType != ENestType::Interface && NestType != ENestType::NativeInterface && NestType != ENestType::FunctionDeclaration)
 	{
-		FError::Throwf(TEXT("Bad first pass NestType %i"), (uint8)NestType);
+		FUHTException::Throwf(*this, TEXT("Bad first pass NestType %i"), (uint8)NestType);
 	}
 
 	bool bLinkProps = true;
@@ -2370,21 +2334,21 @@ void FHeaderParser::FixupDelegateProperties(FUnrealStructDefinitionInfo& StructD
 				if ( DelegateClassName.Len() == 0 )
 				{
 					UngetToken(PropertyDef->GetLineNumber(), PropertyDef->GetParsePosition());
-					FError::Throwf(TEXT("Invalid scope specified in delegate property function reference: '%s'"), *NameOfDelegateFunction);
+					FUHTException::Throwf(*this, TEXT("Invalid scope specified in delegate property function reference: '%s'"), *NameOfDelegateFunction);
 				}
 
 				// verify that we got a valid string for the name of the function
 				if ( DelegateName.Len() == 0 )
 				{
 					UngetToken(PropertyDef->GetLineNumber(), PropertyDef->GetParsePosition());
-					FError::Throwf(TEXT("Invalid delegate name specified in delegate property function reference '%s'"), *NameOfDelegateFunction);
+					FUHTException::Throwf(*this, TEXT("Invalid delegate name specified in delegate property function reference '%s'"), *NameOfDelegateFunction);
 				}
 
 				// make sure that the class that contains the delegate can be referenced here
-				FUnrealClassDefinitionInfo* DelegateOwnerClassDef = FUnrealClassDefinitionInfo::FindScriptClassOrThrow(DelegateClassName);
+				FUnrealClassDefinitionInfo* DelegateOwnerClassDef = FUnrealClassDefinitionInfo::FindScriptClassOrThrow(*this, DelegateClassName);
 				if (DelegateOwnerClassDef->GetScope()->FindTypeByName(*DelegateName) != nullptr)
 				{
-					FError::Throwf(TEXT("Inaccessible type: '%s'"), *DelegateOwnerClassDef->GetPathName());
+					FUHTException::Throwf(*this, TEXT("Inaccessible type: '%s'"), *DelegateOwnerClassDef->GetPathName());
 				}
 				SourceDelegateFunctionDef = FindFunction(*DelegateOwnerClassDef, *DelegateName, false, nullptr);
 			}
@@ -2392,12 +2356,12 @@ void FHeaderParser::FixupDelegateProperties(FUnrealStructDefinitionInfo& StructD
 			if (SourceDelegateFunctionDef == NULL )
 			{
 				UngetToken(PropertyDef->GetLineNumber(), PropertyDef->GetParsePosition());
-				FError::Throwf(TEXT("Failed to find delegate function '%s'"), *NameOfDelegateFunction);
+				FUHTException::Throwf(*this, TEXT("Failed to find delegate function '%s'"), *NameOfDelegateFunction);
 			}
 			else if (!SourceDelegateFunctionDef->HasAnyFunctionFlags(FUNC_Delegate))
 			{
 				UngetToken(PropertyDef->GetLineNumber(), PropertyDef->GetParsePosition());
-				FError::Throwf(TEXT("Only delegate functions can be used as the type for a delegate property; '%s' is not a delegate."), *NameOfDelegateFunction);
+				FUHTException::Throwf(*this, TEXT("Only delegate functions can be used as the type for a delegate property; '%s' is not a delegate."), *NameOfDelegateFunction);
 			}
 		}
 
@@ -2415,7 +2379,7 @@ void FHeaderParser::FixupDelegateProperties(FUnrealStructDefinitionInfo& StructD
 			}
 			else
 			{
-				FError::Throwf(TEXT("Unable to declare a single-cast delegate property for a multi-cast delegate type '%s'.  Either add a 'multicast' qualifier to the property or change the delegate type to be single-cast as well."), *SourceDelegateFunctionDef->GetName());
+				FUHTException::Throwf(*this, TEXT("Unable to declare a single-cast delegate property for a multi-cast delegate type '%s'.  Either add a 'multicast' qualifier to the property or change the delegate type to be single-cast as well."), *SourceDelegateFunctionDef->GetName());
 			}
 		}
 		else if (PropertyBase.Type == CPT_MulticastDelegate)
@@ -2432,7 +2396,7 @@ void FHeaderParser::FixupDelegateProperties(FUnrealStructDefinitionInfo& StructD
 						{
 							FString ExtendedCPPType;
 							FString CPPType = FuncParamDef->GetCPPType(&ExtendedCPPType);
-							UE_LOG_ERROR_UHT(TEXT("Type '%s%s' is not supported by blueprint. %s.%s"), *CPPType, *ExtendedCPPType, *SourceDelegateFunctionDef->GetName(), *FuncParamDef->GetName());
+							UE_LOG_ERROR_UHT(*this, TEXT("Type '%s%s' is not supported by blueprint. %s.%s"), *CPPType, *ExtendedCPPType, *SourceDelegateFunctionDef->GetName(), *FuncParamDef->GetName());
 						}
 
 						if(FuncParamDef->HasAllPropertyFlags(CPF_OutParm) && !FuncParamDef->HasAllPropertyFlags(CPF_ConstParm)  )
@@ -2441,7 +2405,7 @@ void FHeaderParser::FixupDelegateProperties(FUnrealStructDefinitionInfo& StructD
 							const bool bAllowedArrayRefFromBP = bClassGeneratedFromBP && FuncParamDef->IsDynamicArray();
 							if (!bAllowedArrayRefFromBP)
 							{
-								UE_LOG_ERROR_UHT(TEXT("BlueprintAssignable delegates do not support non-const references at the moment. Function: %s Parameter: '%s'"), *SourceDelegateFunctionDef->GetName(), *FuncParamDef->GetName());
+								UE_LOG_ERROR_UHT(*this, TEXT("BlueprintAssignable delegates do not support non-const references at the moment. Function: %s Parameter: '%s'"), *SourceDelegateFunctionDef->GetName(), *FuncParamDef->GetName());
 							}
 						}
 					}
@@ -2450,7 +2414,7 @@ void FHeaderParser::FixupDelegateProperties(FUnrealStructDefinitionInfo& StructD
 			}
 			else
 			{
-				FError::Throwf(TEXT("Unable to declare a multi-cast delegate property for a single-cast delegate type '%s'.  Either remove the 'multicast' qualifier from the property or change the delegate type to be 'multicast' as well."), *SourceDelegateFunctionDef->GetName());
+				FUHTException::Throwf(*this, TEXT("Unable to declare a multi-cast delegate property for a single-cast delegate type '%s'.  Either remove the 'multicast' qualifier from the property or change the delegate type to be 'multicast' as well."), *SourceDelegateFunctionDef->GetName());
 			}
 		}
 	}
@@ -2476,7 +2440,7 @@ void FHeaderParser::CheckSparseClassData(const FUnrealStructDefinitionInfo& Stru
 	{
 		if (ClassDef == nullptr)
 		{
-			FError::Throwf(TEXT("%s contains sparse class data but is not a class."), *StructDef.GetName());
+			FUHTException::Throwf(StructDef, TEXT("%s contains sparse class data but is not a class."), *StructDef.GetName());
 		}
 	}
 	else
@@ -2490,12 +2454,12 @@ void FHeaderParser::CheckSparseClassData(const FUnrealStructDefinitionInfo& Stru
 	// for now we only support one sparse class data structure per class
 	if (SparseClassDataTypes.Num() > 1)
 	{
-		FError::Throwf(TEXT("Class %s contains multiple sparse class data types."), *ClassDef->GetName());
+		FUHTException::Throwf(StructDef, TEXT("Class %s contains multiple sparse class data types."), *ClassDef->GetName());
 		return;
 	}
 	if (SparseClassDataTypes.Num() == 0)
 	{
-		FError::Throwf(TEXT("Class %s has sparse class metadata but does not specify a type."), *ClassDef->GetName());
+		FUHTException::Throwf(StructDef, TEXT("Class %s has sparse class metadata but does not specify a type."), *ClassDef->GetName());
 		return;
 	}
 
@@ -2506,7 +2470,7 @@ void FHeaderParser::CheckSparseClassData(const FUnrealStructDefinitionInfo& Stru
 		// make sure the sparse class data struct actually exists
 		if (!SparseClassDataStructDef)
 		{
-			FError::Throwf(TEXT("Unable to find sparse data type %s for class %s."), *SparseClassDataTypeName, *ClassDef->GetName());
+			FUHTException::Throwf(StructDef, TEXT("Unable to find sparse data type %s for class %s."), *SparseClassDataTypeName, *ClassDef->GetName());
 			return;
 		}
 
@@ -2515,19 +2479,19 @@ void FHeaderParser::CheckSparseClassData(const FUnrealStructDefinitionInfo& Stru
 		{
 			if (PropertyDef->HasAnyPropertyFlags(CPF_BlueprintAssignable))
 			{
-				FError::Throwf(TEXT("Sparse class data types can not contain blueprint assignable delegates. Type '%s' Delegate '%s'"), *SparseClassDataStructDef->GetName(), *PropertyDef->GetName());
+				FUHTException::Throwf(StructDef, TEXT("Sparse class data types can not contain blueprint assignable delegates. Type '%s' Delegate '%s'"), *SparseClassDataStructDef->GetName(), *PropertyDef->GetName());
 			}
 
 			// all sparse properties should have EditDefaultsOnly
 			if (!PropertyDef->HasAllPropertyFlags(CPF_Edit | CPF_DisableEditOnInstance))
 			{
-				FError::Throwf(TEXT("Sparse class data types must be EditDefaultsOnly. Type '%s' Property '%s'"), *SparseClassDataStructDef->GetName(), *PropertyDef->GetName());
+				FUHTException::Throwf(StructDef, TEXT("Sparse class data types must be EditDefaultsOnly. Type '%s' Property '%s'"), *SparseClassDataStructDef->GetName(), *PropertyDef->GetName());
 			}
 
 			// no sparse properties should have BlueprintReadWrite
 			if (PropertyDef->HasAllPropertyFlags(CPF_BlueprintVisible) && !PropertyDef->HasAllPropertyFlags(CPF_BlueprintReadOnly))
 			{
-				FError::Throwf(TEXT("Sparse class data types must not be BlueprintReadWrite. Type '%s' Property '%s'"), *SparseClassDataStructDef->GetName(), *PropertyDef->GetName());
+				FUHTException::Throwf(StructDef, TEXT("Sparse class data types must not be BlueprintReadWrite. Type '%s' Property '%s'"), *SparseClassDataStructDef->GetName(), *PropertyDef->GetName());
 			}
 		}
 
@@ -2543,7 +2507,7 @@ void FHeaderParser::CheckSparseClassData(const FUnrealStructDefinitionInfo& Stru
 				UScriptStruct* ParentSparseClassDataStruct = ParentSparseClassDataStructDef->GetScriptStruct();
 				if (!SparseClassDataStructDef->GetScriptStruct()->IsChildOf(ParentSparseClassDataStruct))
 				{
-					FError::Throwf(TEXT("Class %s is a child of %s but its sparse class data struct, %s, does not inherit from %s."), *ClassDef->GetName(), *ParentClassDef->GetName(), *SparseClassDataStructDef->GetName(), *ParentSparseClassDataStruct->GetName());
+					FUHTException::Throwf(StructDef, TEXT("Class %s is a child of %s but its sparse class data struct, %s, does not inherit from %s."), *ClassDef->GetName(), *ParentClassDef->GetName(), *SparseClassDataStructDef->GetName(), *ParentSparseClassDataStruct->GetName());
 				}
 			}
 		}
@@ -2555,7 +2519,7 @@ void FHeaderParser::ValidateClassFlags(const FUnrealClassDefinitionInfo& ToValid
 	if (ToValidate.HasAnyClassFlags(CLASS_NeedsDeferredDependencyLoading) && !ToValidate.GetClass()->IsChildOf(UClass::StaticClass()))
 	{
 		// CLASS_NeedsDeferredDependencyLoading can only be set on classes derived from UClass
-		FError::Throwf(TEXT("NeedsDeferredDependencyLoading is set on %s but the flag can only be used with classes derived from UClass."), *ToValidate.GetClass()->GetName());
+		FUHTException::Throwf(ToValidate, TEXT("NeedsDeferredDependencyLoading is set on %s but the flag can only be used with classes derived from UClass."), *ToValidate.GetClass()->GetName());
 	}
 }
 
@@ -2567,23 +2531,23 @@ void FHeaderParser::VerifyBlueprintPropertyGetter(FUnrealPropertyDefinitionInfo&
 	FUnrealPropertyDefinitionInfo* ReturnPropDef = TargetFuncDef->GetReturn();
 	if (Properties.Num() > 1 || (Properties.Num() == 1 && ReturnPropDef == nullptr))
 	{
-		UE_LOG_ERROR_UHT(TEXT("Blueprint Property getter function %s must not have parameters."), *TargetFuncDef->GetName());
+		UE_LOG_ERROR_UHT(*this, TEXT("Blueprint Property getter function %s must not have parameters."), *TargetFuncDef->GetName());
 	}
 
 	if (ReturnPropDef == nullptr || !PropertyDef.SameType(*ReturnPropDef))
 	{
 		FString ExtendedCPPType;
 		FString CPPType = PropertyDef.GetCPPType(&ExtendedCPPType);
-		UE_LOG_ERROR_UHT(TEXT("Blueprint Property getter function %s must have return value of type %s%s."), *TargetFuncDef->GetName(), *CPPType, *ExtendedCPPType);
+		UE_LOG_ERROR_UHT(*this, TEXT("Blueprint Property getter function %s must have return value of type %s%s."), *TargetFuncDef->GetName(), *CPPType, *ExtendedCPPType);
 	}
 
 	if (TargetFuncDef->HasAnyFunctionFlags(FUNC_Event))
 	{
-		UE_LOG_ERROR_UHT(TEXT("Blueprint Property setter function cannot be a blueprint event."));
+		UE_LOG_ERROR_UHT(*this, TEXT("Blueprint Property setter function cannot be a blueprint event."));
 	}
 	else if (!TargetFuncDef->HasAnyFunctionFlags(FUNC_BlueprintPure))
 	{
-		UE_LOG_ERROR_UHT(TEXT("Blueprint Property getter function must be pure."));
+		UE_LOG_ERROR_UHT(*this, TEXT("Blueprint Property getter function must be pure."));
 	}
 }
 
@@ -2594,7 +2558,7 @@ void FHeaderParser::VerifyBlueprintPropertySetter(FUnrealPropertyDefinitionInfo&
 
 	if (ReturnPropDef)
 	{
-		UE_LOG_ERROR_UHT(TEXT("Blueprint Property setter function %s must not have a return value."), *TargetFuncDef->GetName());
+		UE_LOG_ERROR_UHT(*this, TEXT("Blueprint Property setter function %s must not have a return value."), *TargetFuncDef->GetName());
 	}
 	else
 	{
@@ -2603,21 +2567,21 @@ void FHeaderParser::VerifyBlueprintPropertySetter(FUnrealPropertyDefinitionInfo&
 		{
 			FString ExtendedCPPType;
 			FString CPPType = PropertyDef.GetCPPType(&ExtendedCPPType);
-			UE_LOG_ERROR_UHT(TEXT("Blueprint Property setter function %s must have exactly one parameter of type %s%s."), *TargetFuncDef->GetName(), *CPPType, *ExtendedCPPType);
+			UE_LOG_ERROR_UHT(*this, TEXT("Blueprint Property setter function %s must have exactly one parameter of type %s%s."), *TargetFuncDef->GetName(), *CPPType, *ExtendedCPPType);
 		}
 	}
 
 	if (TargetFuncDef->HasAnyFunctionFlags(FUNC_Event))
 	{
-		UE_LOG_ERROR_UHT(TEXT("Blueprint Property setter function cannot be a blueprint event."));
+		UE_LOG_ERROR_UHT(*this, TEXT("Blueprint Property setter function cannot be a blueprint event."));
 	}
 	else if (!TargetFuncDef->HasAnyFunctionFlags(FUNC_BlueprintCallable))
 	{
-		UE_LOG_ERROR_UHT(TEXT("Blueprint Property setter function must be blueprint callable."));
+		UE_LOG_ERROR_UHT(*this, TEXT("Blueprint Property setter function must be blueprint callable."));
 	}
 	else if (TargetFuncDef->HasAnyFunctionFlags(FUNC_BlueprintPure))
 	{
-		UE_LOG_ERROR_UHT(TEXT("Blueprint Property setter function must not be pure."));
+		UE_LOG_ERROR_UHT(*this, TEXT("Blueprint Property setter function must not be pure."));
 	}
 }
 
@@ -2627,7 +2591,7 @@ void FHeaderParser::VerifyRepNotifyCallback(FUnrealPropertyDefinitionInfo& Prope
 	{
 		if (TargetFuncDef->GetReturn())
 		{
-			UE_LOG_ERROR_UHT(TEXT("Replication notification function %s must not have return value."), *TargetFuncDef->GetName());
+			UE_LOG_ERROR_UHT(*this, TEXT("Replication notification function %s must not have return value."), *TargetFuncDef->GetName());
 		}
 
 		const bool bIsArrayProperty = PropertyDef.IsStaticArray() || PropertyDef.IsDynamicArray();
@@ -2636,7 +2600,7 @@ void FHeaderParser::VerifyRepNotifyCallback(FUnrealPropertyDefinitionInfo& Prope
 		const TArray<TSharedRef<FUnrealPropertyDefinitionInfo>>& Properties = TargetFuncDef->GetProperties();
 		if (Properties.Num() > MaxParms)
 		{
-			UE_LOG_ERROR_UHT(TEXT("Replication notification function %s has too many parameters."), *TargetFuncDef->GetName());
+			UE_LOG_ERROR_UHT(*this, TEXT("Replication notification function %s has too many parameters."), *TargetFuncDef->GetName());
 		}
 
 		if (Properties.Num() >= 1)
@@ -2647,7 +2611,7 @@ void FHeaderParser::VerifyRepNotifyCallback(FUnrealPropertyDefinitionInfo& Prope
 			{
 				FString ExtendedCPPType;
 				FString CPPType = ParmDef->GetCPPType(&ExtendedCPPType);
-				UE_LOG_ERROR_UHT(TEXT("Replication notification function %s has invalid parameter for property %s. First (optional) parameter must be of type %s%s."), *TargetFuncDef->GetName(), *ParmDef->GetName(), *CPPType, *ExtendedCPPType);
+				UE_LOG_ERROR_UHT(*this, TEXT("Replication notification function %s has invalid parameter for property %s. First (optional) parameter must be of type %s%s."), *TargetFuncDef->GetName(), *ParmDef->GetName(), *CPPType, *ExtendedCPPType);
 			}
 		}
 
@@ -2658,14 +2622,14 @@ void FHeaderParser::VerifyRepNotifyCallback(FUnrealPropertyDefinitionInfo& Prope
 			// A 2nd parameter for arrays can be specified as a const TArray<uint8>&. This is a list of element indices that have changed
 			if (!(ParmBase.Type == CPT_Byte && ParmBase.ArrayType == EArrayType::Dynamic && ParmDef->HasAllPropertyFlags(CPF_ConstParm | CPF_ReferenceParm)))
 			{
-				UE_LOG_ERROR_UHT(TEXT("Replication notification function %s (optional) second parameter must be of type 'const TArray<uint8>&'"), *TargetFuncDef->GetName());
+				UE_LOG_ERROR_UHT(*this, TEXT("Replication notification function %s (optional) second parameter must be of type 'const TArray<uint8>&'"), *TargetFuncDef->GetName());
 			}
 		}
 	}
 	else
 	{
 		// Couldn't find a valid function...
-		UE_LOG_ERROR_UHT(TEXT("Replication notification function %s not found"), *PropertyDef.GetRepNotifyFunc().ToString() );
+		UE_LOG_ERROR_UHT(*this, TEXT("Replication notification function %s not found"), *PropertyDef.GetRepNotifyFunc().ToString() );
 	}
 }
 void FHeaderParser::VerifyPropertyMarkups(FUnrealClassDefinitionInfo& TargetClassDef)
@@ -2712,7 +2676,7 @@ void FHeaderParser::VerifyPropertyMarkups(FUnrealClassDefinitionInfo& TargetClas
 				else
 				{
 					// Couldn't find a valid function...
-					UE_LOG_ERROR_UHT(TEXT("Blueprint Property getter function %s not found"), *GetterFuncName);
+					UE_LOG_ERROR_UHT(*this, TEXT("Blueprint Property getter function %s not found"), *GetterFuncName);
 				}
 			}
 
@@ -2728,7 +2692,7 @@ void FHeaderParser::VerifyPropertyMarkups(FUnrealClassDefinitionInfo& TargetClas
 					else
 					{
 						// Couldn't find a valid function...
-						UE_LOG_ERROR_UHT(TEXT("Blueprint Property setter function %s not found"), *SetterFuncName);
+						UE_LOG_ERROR_UHT(*this, TEXT("Blueprint Property setter function %s not found"), *SetterFuncName);
 					}
 				}
 			}
@@ -2754,11 +2718,11 @@ void FHeaderParser::CompileDirective()
 
 	if (!GetIdentifier(Directive))
 	{
-		FError::Throwf(TEXT("Missing compiler directive after '#'") );
+		FUHTException::Throwf(*this, TEXT("Missing compiler directive after '#'") );
 	}
 	else if (Directive.Matches(TEXT("error"), ESearchCase::CaseSensitive))
 	{
-		FError::Throwf(TEXT("#error directive encountered") );
+		FUHTException::Throwf(*this, TEXT("#error directive encountered") );
 	}
 	else if (Directive.Matches(TEXT("pragma"), ESearchCase::CaseSensitive))
 	{
@@ -2769,7 +2733,7 @@ void FHeaderParser::CompileDirective()
 		FToken Number;
 		if (!GetToken(Number) || (Number.TokenType != TOKEN_Const) || (Number.ConstantType != CPT_Int && Number.ConstantType != CPT_Int64))
 		{
-			FError::Throwf(TEXT("Missing line number in line number directive"));
+			FUHTException::Throwf(*this, TEXT("Missing line number in line number directive"));
 		}
 
 		int32 newInputLine;
@@ -2806,7 +2770,7 @@ void FHeaderParser::CompileDirective()
 			FToken Define;
 			if (!GetIdentifier(Define))
 			{
-				FError::Throwf(TEXT("Missing define name '#if'") );
+				FUHTException::Throwf(*this, TEXT("Missing define name '#if'") );
 			}
 
 			if ( Define.Matches(TEXT("WITH_EDITORONLY_DATA"), ESearchCase::CaseSensitive) )
@@ -2827,7 +2791,7 @@ void FHeaderParser::CompileDirective()
 			}
 			else
 			{
-				FError::Throwf(TEXT("Unknown define '#if %s' in class or global scope"), Define.Identifier);
+				FUHTException::Throwf(*this, TEXT("Unknown define '#if %s' in class or global scope"), Define.Identifier);
 			}
 		}
 	}
@@ -2835,7 +2799,7 @@ void FHeaderParser::CompileDirective()
 	{
 		if (CompilerDirectiveStack.Num() < 1)
 		{
-			FError::Throwf(TEXT("Unmatched '#endif' in class or global scope"));
+			FUHTException::Throwf(*this, TEXT("Unmatched '#endif' in class or global scope"));
 		}
 		CompilerDirectiveStack.Pop();
 	}
@@ -2854,7 +2818,7 @@ void FHeaderParser::CompileDirective()
 	}
 	else
 	{
-		FError::Throwf(TEXT("Unrecognized compiler directive %s"), Directive.Identifier );
+		FUHTException::Throwf(*this, TEXT("Unrecognized compiler directive %s"), Directive.Identifier );
 	}
 
 	// Skip to end of line (or end of multiline #define).
@@ -2971,7 +2935,7 @@ void FHeaderParser::GetVarType(
 				{
 					if (bSeenEditSpecifier)
 					{
-						UE_LOG_ERROR_UHT(TEXT("Found more than one edit/visibility specifier (%s), only one is allowed"), *Specifier.Key);
+						UE_LOG_ERROR_UHT(*this, TEXT("Found more than one edit/visibility specifier (%s), only one is allowed"), *Specifier.Key);
 					}
 					Flags |= CPF_Edit;
 					bSeenEditSpecifier = true;
@@ -2982,7 +2946,7 @@ void FHeaderParser::GetVarType(
 				{
 					if (bSeenEditSpecifier)
 					{
-						UE_LOG_ERROR_UHT(TEXT("Found more than one edit/visibility specifier (%s), only one is allowed"), *Specifier.Key);
+						UE_LOG_ERROR_UHT(*this, TEXT("Found more than one edit/visibility specifier (%s), only one is allowed"), *Specifier.Key);
 					}
 					Flags |= CPF_Edit | CPF_DisableEditOnTemplate;
 					bSeenEditSpecifier = true;
@@ -2993,7 +2957,7 @@ void FHeaderParser::GetVarType(
 				{
 					if (bSeenEditSpecifier)
 					{
-						UE_LOG_ERROR_UHT(TEXT("Found more than one edit/visibility specifier (%s), only one is allowed"), *Specifier.Key);
+						UE_LOG_ERROR_UHT(*this, TEXT("Found more than one edit/visibility specifier (%s), only one is allowed"), *Specifier.Key);
 					}
 					Flags |= CPF_Edit | CPF_DisableEditOnInstance;
 					bSeenEditSpecifier = true;
@@ -3004,7 +2968,7 @@ void FHeaderParser::GetVarType(
 				{
 					if (bSeenEditSpecifier)
 					{
-						UE_LOG_ERROR_UHT(TEXT("Found more than one edit/visibility specifier (%s), only one is allowed"), *Specifier.Key);
+						UE_LOG_ERROR_UHT(*this, TEXT("Found more than one edit/visibility specifier (%s), only one is allowed"), *Specifier.Key);
 					}
 					Flags |= CPF_Edit | CPF_EditConst;
 					bSeenEditSpecifier = true;
@@ -3015,7 +2979,7 @@ void FHeaderParser::GetVarType(
 				{
 					if (bSeenEditSpecifier)
 					{
-						UE_LOG_ERROR_UHT(TEXT("Found more than one edit/visibility specifier (%s), only one is allowed"), *Specifier.Key);
+						UE_LOG_ERROR_UHT(*this, TEXT("Found more than one edit/visibility specifier (%s), only one is allowed"), *Specifier.Key);
 					}
 					Flags |= CPF_Edit | CPF_EditConst | CPF_DisableEditOnTemplate;
 					bSeenEditSpecifier = true;
@@ -3026,7 +2990,7 @@ void FHeaderParser::GetVarType(
 				{
 					if (bSeenEditSpecifier)
 					{
-						UE_LOG_ERROR_UHT(TEXT("Found more than one edit/visibility specifier (%s), only one is allowed"), *Specifier.Key);
+						UE_LOG_ERROR_UHT(*this, TEXT("Found more than one edit/visibility specifier (%s), only one is allowed"), *Specifier.Key);
 					}
 					Flags |= CPF_Edit | CPF_EditConst | CPF_DisableEditOnInstance;
 					bSeenEditSpecifier = true;
@@ -3037,19 +3001,19 @@ void FHeaderParser::GetVarType(
 				{
 					if (bSeenBlueprintReadOnlySpecifier)
 					{
-						UE_LOG_ERROR_UHT(TEXT("Cannot specify a property as being both BlueprintReadOnly and BlueprintReadWrite."));
+						UE_LOG_ERROR_UHT(*this, TEXT("Cannot specify a property as being both BlueprintReadOnly and BlueprintReadWrite."));
 					}
 
 					const FString* PrivateAccessMD = MetaDataFromNewStyle.Find(NAME_AllowPrivateAccess);  // FBlueprintMetadata::MD_AllowPrivateAccess
 					const bool bAllowPrivateAccess = PrivateAccessMD ? (*PrivateAccessMD != TEXT("false")) : false;
 					if (CurrentAccessSpecifier == ACCESS_Private && !bAllowPrivateAccess)
 					{
-						UE_LOG_ERROR_UHT(TEXT("BlueprintReadWrite should not be used on private members"));
+						UE_LOG_ERROR_UHT(*this, TEXT("BlueprintReadWrite should not be used on private members"));
 					}
 
 					if ((Flags & CPF_EditorOnly) != 0 && OwnerStruct->IsA<UScriptStruct>())
 					{
-						UE_LOG_ERROR_UHT(TEXT("Blueprint exposed struct members cannot be editor only"));
+						UE_LOG_ERROR_UHT(*this, TEXT("Blueprint exposed struct members cannot be editor only"));
 					}
 
 					Flags |= CPF_BlueprintVisible;
@@ -3061,15 +3025,15 @@ void FHeaderParser::GetVarType(
 				{
 					if (bSeenBlueprintReadOnlySpecifier)
 					{
-						UE_LOG_ERROR_UHT(TEXT("Cannot specify a property as being both BlueprintReadOnly and having a BlueprintSetter."));
+						UE_LOG_ERROR_UHT(*this, TEXT("Cannot specify a property as being both BlueprintReadOnly and having a BlueprintSetter."));
 					}
 
 					if (OwnerStruct->IsA<UScriptStruct>())
 					{
-						UE_LOG_ERROR_UHT(TEXT("Cannot specify BlueprintSetter for a struct member."))
+						UE_LOG_ERROR_UHT(*this, TEXT("Cannot specify BlueprintSetter for a struct member."))
 					}
 
-					const FString BlueprintSetterFunction = RequireExactlyOneSpecifierValue(Specifier);
+					const FString BlueprintSetterFunction = RequireExactlyOneSpecifierValue(*this, Specifier);
 					MetaDataFromNewStyle.Add(NAME_BlueprintSetter, BlueprintSetterFunction);
 
 					Flags |= CPF_BlueprintVisible;
@@ -3081,19 +3045,19 @@ void FHeaderParser::GetVarType(
 				{
 					if (bSeenBlueprintWriteSpecifier)
 					{
-						UE_LOG_ERROR_UHT(TEXT("Cannot specify both BlueprintReadOnly and BlueprintReadWrite or BlueprintSetter."), *Specifier.Key);
+						UE_LOG_ERROR_UHT(*this, TEXT("Cannot specify both BlueprintReadOnly and BlueprintReadWrite or BlueprintSetter."), *Specifier.Key);
 					}
 
 					const FString* PrivateAccessMD = MetaDataFromNewStyle.Find(NAME_AllowPrivateAccess);  // FBlueprintMetadata::MD_AllowPrivateAccess
 					const bool bAllowPrivateAccess = PrivateAccessMD ? (*PrivateAccessMD != TEXT("false")) : false;
 					if (CurrentAccessSpecifier == ACCESS_Private && !bAllowPrivateAccess)
 					{
-						UE_LOG_ERROR_UHT(TEXT("BlueprintReadOnly should not be used on private members"));
+						UE_LOG_ERROR_UHT(*this, TEXT("BlueprintReadOnly should not be used on private members"));
 					}
 
 					if ((Flags & CPF_EditorOnly) != 0 && OwnerStruct->IsA<UScriptStruct>())
 					{
-						UE_LOG_ERROR_UHT(TEXT("Blueprint exposed struct members cannot be editor only"));
+						UE_LOG_ERROR_UHT(*this, TEXT("Blueprint exposed struct members cannot be editor only"));
 					}
 
 					Flags        |= CPF_BlueprintVisible | CPF_BlueprintReadOnly;
@@ -3106,10 +3070,10 @@ void FHeaderParser::GetVarType(
 				{
 					if (OwnerStruct->IsA<UScriptStruct>())
 					{
-						UE_LOG_ERROR_UHT(TEXT("Cannot specify BlueprintGetter for a struct member."))
+						UE_LOG_ERROR_UHT(*this, TEXT("Cannot specify BlueprintGetter for a struct member."))
 					}
 
-					const FString BlueprintGetterFunction = RequireExactlyOneSpecifierValue(Specifier);
+					const FString BlueprintGetterFunction = RequireExactlyOneSpecifierValue(*this, Specifier);
 					MetaDataFromNewStyle.Add(NAME_BlueprintGetter, BlueprintGetterFunction);
 
 					Flags        |= CPF_BlueprintVisible;
@@ -3131,7 +3095,7 @@ void FHeaderParser::GetVarType(
 
 				case EVariableSpecifier::Localized:
 				{
-					UE_LOG_ERROR_UHT(TEXT("The Localized specifier is deprecated"));
+					UE_LOG_ERROR_UHT(*this, TEXT("The Localized specifier is deprecated"));
 				}
 				break;
 
@@ -3155,7 +3119,7 @@ void FHeaderParser::GetVarType(
 
 				case EVariableSpecifier::NonPIETransient:
 				{
-					UE_LOG_WARNING_UHT(TEXT("NonPIETransient is deprecated - NonPIEDuplicateTransient should be used instead"));
+					UE_LOG_WARNING_UHT(*this, TEXT("NonPIETransient is deprecated - NonPIEDuplicateTransient should be used instead"));
 					Flags |= CPF_NonPIEDuplicateTransient;
 				}
 				break;
@@ -3174,7 +3138,7 @@ void FHeaderParser::GetVarType(
 
 				case EVariableSpecifier::EditInline:
 				{
-					UE_LOG_ERROR_UHT(TEXT("EditInline is deprecated. Remove it, or use Instanced instead."));
+					UE_LOG_ERROR_UHT(*this, TEXT("EditInline is deprecated. Remove it, or use Instanced instead."));
 				}
 				break;
 
@@ -3195,7 +3159,7 @@ void FHeaderParser::GetVarType(
 				{
 					if (OwnerStruct->IsA<UScriptStruct>())
 					{
-						UE_LOG_ERROR_UHT(TEXT("Struct members cannot be replicated"));
+						UE_LOG_ERROR_UHT(*this, TEXT("Struct members cannot be replicated"));
 					}
 
 					Flags |= CPF_Net;
@@ -3203,7 +3167,7 @@ void FHeaderParser::GetVarType(
 					// See if we've specified a rep notification function
 					if (SpecID == EVariableSpecifier::ReplicatedUsing)
 					{
-						RepCallbackName = FName(*RequireExactlyOneSpecifierValue(Specifier));
+						RepCallbackName = FName(*RequireExactlyOneSpecifierValue(*this, Specifier));
 						Flags |= CPF_RepNotify;
 					}
 				}
@@ -3213,7 +3177,7 @@ void FHeaderParser::GetVarType(
 				{
 					if (!OwnerStruct->IsA<UScriptStruct>())
 					{
-						UE_LOG_ERROR_UHT(TEXT("Only Struct members can be marked NotReplicated"));
+						UE_LOG_ERROR_UHT(*this, TEXT("Only Struct members can be marked NotReplicated"));
 					}
 
 					Flags |= CPF_RepSkip;
@@ -3222,7 +3186,7 @@ void FHeaderParser::GetVarType(
 
 				case EVariableSpecifier::RepRetry:
 				{
-					UE_LOG_ERROR_UHT(TEXT("'RepRetry' is deprecated."));
+					UE_LOG_ERROR_UHT(*this, TEXT("'RepRetry' is deprecated."));
 				}
 				break;
 
@@ -3297,7 +3261,7 @@ void FHeaderParser::GetVarType(
 
 				default:
 				{
-					UE_LOG_ERROR_UHT(TEXT("Unknown variable specifier '%s'"), *Specifier.Key);
+					UE_LOG_ERROR_UHT(*this, TEXT("Unknown variable specifier '%s'"), *Specifier.Key);
 				}
 				break;
 			}
@@ -3327,14 +3291,14 @@ void FHeaderParser::GetVarType(
 					}
 					else
 					{
-						UE_LOG_ERROR_UHT(TEXT("Only parameters in service request functions can be marked NotReplicated"));
+						UE_LOG_ERROR_UHT(*this, TEXT("Only parameters in service request functions can be marked NotReplicated"));
 					}
 				}
 				break;
 
 				default:
 				{
-					UE_LOG_ERROR_UHT(TEXT("Unknown variable specifier '%s'"), *Specifier.Key);
+					UE_LOG_ERROR_UHT(*this, TEXT("Unknown variable specifier '%s'"), *Specifier.Key);
 				}
 				break;
 			}
@@ -3356,11 +3320,11 @@ void FHeaderParser::GetVarType(
 		{
 			if (0 != (CPF_DisableEditOnInstance & Flags))
 			{
-				UE_LOG_WARNING_UHT(TEXT("Property cannot have both 'DisableEditOnInstance' and 'ExposeOnSpawn' flags"));
+				UE_LOG_WARNING_UHT(*this, TEXT("Property cannot have both 'DisableEditOnInstance' and 'ExposeOnSpawn' flags"));
 			}
 			if (0 == (CPF_BlueprintVisible & Flags))
 			{
-				UE_LOG_WARNING_UHT(TEXT("Property cannot have 'ExposeOnSpawn' without 'BlueprintVisible' flag."));
+				UE_LOG_WARNING_UHT(*this, TEXT("Property cannot have 'ExposeOnSpawn' without 'BlueprintVisible' flag."));
 			}
 			Flags |= CPF_ExposeOnSpawn;
 		}
@@ -3395,7 +3359,7 @@ void FHeaderParser::GetVarType(
 	}
 	else
 	{
-		FError::Throwf(TEXT("Unknown access level"));
+		FUHTException::Throwf(*this, TEXT("Unknown access level"));
 	}
 
 	// Swallow inline keywords
@@ -3404,7 +3368,7 @@ void FHeaderParser::GetVarType(
 		FToken InlineToken;
 		if (!GetIdentifier(InlineToken, true))
 		{
-			FError::Throwf(TEXT("%s: Missing variable type"), GetHintText(VariableCategory));
+			FUHTException::Throwf(*this, TEXT("%s: Missing variable type"), GetHintText(*this, VariableCategory));
 		}
 
 		if (FCString::Strcmp(InlineToken.Identifier, TEXT("inline")) != 0
@@ -3489,7 +3453,7 @@ void FHeaderParser::GetVarType(
 	{
 		if (VariableCategory == EVariableCategory::Member)
 		{
-			FError::Throwf(TEXT("%s: Cannot declare enum at variable declaration"), GetHintText(VariableCategory));
+			FUHTException::Throwf(*this, TEXT("%s: Cannot declare enum at variable declaration"), GetHintText(*this, VariableCategory));
 		}
 
 		bUnconsumedEnumKeyword = true;
@@ -3499,7 +3463,7 @@ void FHeaderParser::GetVarType(
 	FToken VarType;
 	if ( !GetIdentifier(VarType,1) )
 	{
-		FError::Throwf(TEXT("%s: Missing variable type"), GetHintText(VariableCategory));
+		FUHTException::Throwf(*this, TEXT("%s: Missing variable type"), GetHintText(*this, VariableCategory));
 	}
 
 	RedirectTypeIdentifier(VarType);
@@ -3562,7 +3526,7 @@ void FHeaderParser::GetVarType(
 	{
 		if (IsBitfieldProperty(LayoutMacroType))
 		{
-			UE_LOG_ERROR_UHT(TEXT("bool bitfields are not supported."));
+			UE_LOG_ERROR_UHT(*this, TEXT("bool bitfields are not supported."));
 		}
 		// C++ bool type
 		VarProperty = FPropertyBase(CPT_Bool);
@@ -3599,7 +3563,7 @@ void FHeaderParser::GetVarType(
 		const FString& ForSourceFile = Scope->GetFileScope()->GetSourceFile()->GetFilename();
 		if(!ForSourceFile.EndsWith(TEXT("NoExportTypes.h")))
 		{
-			FError::Throwf(TEXT("FLargeWorldCoordinatesReal is intended for LWC support only and should not be used outside of NoExportTypes.h"));
+			FUHTException::Throwf(*this, TEXT("FLargeWorldCoordinatesReal is intended for LWC support only and should not be used outside of NoExportTypes.h"));
 		}
 		VarProperty = FPropertyBase(CPT_FLargeWorldCoordinatesReal);
 	}
@@ -3617,7 +3581,7 @@ void FHeaderParser::GetVarType(
 		GetVarType(Scope, VarProperty, Disallow, &VarType, &VarTypePropertyFlags, EPropertyDeclarationStyle::None, VariableCategory);
 		if (VarProperty.IsContainer())
 		{
-			FError::Throwf(TEXT("Nested containers are not supported.") );
+			FUHTException::Throwf(*this, TEXT("Nested containers are not supported.") );
 		}
 		// TODO: Prevent sparse delegate types from being used in a container
 
@@ -3632,7 +3596,7 @@ void FHeaderParser::GetVarType(
 		FToken CloseTemplateToken;
 		if (!GetToken(CloseTemplateToken, /*bNoConsts=*/ true, ESymbolParseOption::CloseTemplateBracket))
 		{
-			FError::Throwf(TEXT("Missing token while parsing TArray."));
+			FUHTException::Throwf(*this, TEXT("Missing token while parsing TArray."));
 		}
 
 		if (CloseTemplateToken.TokenType != TOKEN_Symbol || !CloseTemplateToken.Matches(TEXT('>')))
@@ -3640,26 +3604,26 @@ void FHeaderParser::GetVarType(
 			// If we didn't find a comma, report it
 			if (!CloseTemplateToken.Matches(TEXT(',')))
 			{
-				FError::Throwf(TEXT("Expected '>' but found '%s'"), CloseTemplateToken.Identifier);
+				FUHTException::Throwf(*this, TEXT("Expected '>' but found '%s'"), CloseTemplateToken.Identifier);
 			}
 
 			// If we found a comma, read the next thing, assume it's an allocator, and report that
 			FToken AllocatorToken;
 			if (!GetToken(AllocatorToken, /*bNoConsts=*/ true, ESymbolParseOption::CloseTemplateBracket))
 			{
-				FError::Throwf(TEXT("Unexpected end of file when parsing TArray allocator."));
+				FUHTException::Throwf(*this, TEXT("Unexpected end of file when parsing TArray allocator."));
 			}
 
 			if (AllocatorToken.TokenType != TOKEN_Identifier)
 			{
-				FError::Throwf(TEXT("Found '%s' - expected a '>' or ','."), AllocatorToken.Identifier);
+				FUHTException::Throwf(*this, TEXT("Found '%s' - expected a '>' or ','."), AllocatorToken.Identifier);
 			}
 
 			if (FCString::Strcmp(AllocatorToken.Identifier, TEXT("FMemoryImageAllocator")) == 0)
 			{
 				if (EnumHasAnyFlags(Flags, CPF_Net))
 				{
-					FError::Throwf(TEXT("Replicated arrays with MemoryImageAllocators are not yet supported"));
+					FUHTException::Throwf(*this, TEXT("Replicated arrays with MemoryImageAllocators are not yet supported"));
 				}
 
 				RequireSymbol(TEXT('>'), TEXT("TArray template arguments"), ESymbolParseOption::CloseTemplateBracket);
@@ -3670,7 +3634,7 @@ void FHeaderParser::GetVarType(
 			{
 				if (EnumHasAnyFlags(Flags, CPF_Net))
 				{
-					FError::Throwf(TEXT("Replicated arrays with MemoryImageAllocators are not yet supported"));
+					FUHTException::Throwf(*this, TEXT("Replicated arrays with MemoryImageAllocators are not yet supported"));
 				}
 
 				RequireSymbol(TEXT('<'), TEXT("TMemoryImageAllocator template arguments"));
@@ -3680,7 +3644,7 @@ void FHeaderParser::GetVarType(
 				{
 					if (!GetToken(SkipToken, /*bNoConsts=*/ false, ESymbolParseOption::CloseTemplateBracket))
 					{
-						FError::Throwf(TEXT("Unexpected end of file when parsing TMemoryImageAllocator template arguments."));
+						FUHTException::Throwf(*this, TEXT("Unexpected end of file when parsing TMemoryImageAllocator template arguments."));
 					}
 
 					if (SkipToken.TokenType == TOKEN_Symbol && FCString::Strcmp(SkipToken.Identifier, TEXT(">")) == 0)
@@ -3693,7 +3657,7 @@ void FHeaderParser::GetVarType(
 			}
 			else
 			{
-				FError::Throwf(TEXT("Found '%s' - explicit allocators are not supported in TArray properties."), AllocatorToken.Identifier);
+				FUHTException::Throwf(*this, TEXT("Found '%s' - explicit allocators are not supported in TArray properties."), AllocatorToken.Identifier);
 			}
 		}
 	}
@@ -3707,35 +3671,35 @@ void FHeaderParser::GetVarType(
 		GetVarType(Scope, MapKeyType, Disallow, &VarType, &VarTypePropertyFlags, EPropertyDeclarationStyle::None, VariableCategory);
 		if (MapKeyType.IsContainer())
 		{
-			FError::Throwf(TEXT("Nested containers are not supported.") );
+			FUHTException::Throwf(*this, TEXT("Nested containers are not supported.") );
 		}
 		// TODO: Prevent sparse delegate types from being used in a container
 
 		if (MapKeyType.Type == CPT_Interface)
 		{
-			FError::Throwf(TEXT("UINTERFACEs are not currently supported as key types."));
+			FUHTException::Throwf(*this, TEXT("UINTERFACEs are not currently supported as key types."));
 		}
 
 		if (MapKeyType.Type == CPT_Text)
 		{
-			FError::Throwf(TEXT("FText is not currently supported as a key type."));
+			FUHTException::Throwf(*this, TEXT("FText is not currently supported as a key type."));
 		}
 		
 		if (EnumHasAnyFlags(Flags, CPF_Net))
 		{
-			UE_LOG_ERROR_UHT(TEXT("Replicated maps are not supported."));
+			UE_LOG_ERROR_UHT(*this, TEXT("Replicated maps are not supported."));
 		}
 
 		FToken CommaToken;
 		if (!GetToken(CommaToken, /*bNoConsts=*/ true) || CommaToken.TokenType != TOKEN_Symbol || !CommaToken.Matches(TEXT(',')))
 		{
-			FError::Throwf(TEXT("Missing value type while parsing TMap."));
+			FUHTException::Throwf(*this, TEXT("Missing value type while parsing TMap."));
 		}
 
 		GetVarType(Scope, VarProperty, Disallow, &VarType, &VarTypePropertyFlags, EPropertyDeclarationStyle::None, VariableCategory);
 		if (VarProperty.IsContainer())
 		{
-			FError::Throwf(TEXT("Nested containers are not supported.") );
+			FUHTException::Throwf(*this, TEXT("Nested containers are not supported.") );
 		}
 		// TODO: Prevent sparse delegate types from being used in a container
 
@@ -3747,7 +3711,7 @@ void FHeaderParser::GetVarType(
 		FToken CloseTemplateToken;
 		if (!GetToken(CloseTemplateToken, /*bNoConsts=*/ true, ESymbolParseOption::CloseTemplateBracket))
 		{
-			FError::Throwf(TEXT("Missing token while parsing TMap."));
+			FUHTException::Throwf(*this, TEXT("Missing token while parsing TMap."));
 		}
 
 		if (CloseTemplateToken.TokenType != TOKEN_Symbol || !CloseTemplateToken.Matches(TEXT('>')))
@@ -3755,19 +3719,19 @@ void FHeaderParser::GetVarType(
 			// If we didn't find a comma, report it
 			if (!CloseTemplateToken.Matches(TEXT(',')))
 			{
-				FError::Throwf(TEXT("Expected '>' but found '%s'"), CloseTemplateToken.Identifier);
+				FUHTException::Throwf(*this, TEXT("Expected '>' but found '%s'"), CloseTemplateToken.Identifier);
 			}
 
 			// If we found a comma, read the next thing, assume it's an allocator, and report that
 			FToken AllocatorToken;
 			if (!GetToken(AllocatorToken, /*bNoConsts=*/ true, ESymbolParseOption::CloseTemplateBracket))
 			{
-				FError::Throwf(TEXT("Unexpected end of file when parsing TArray allocator."));
+				FUHTException::Throwf(*this, TEXT("Unexpected end of file when parsing TArray allocator."));
 			}
 
 			if (AllocatorToken.TokenType != TOKEN_Identifier)
 			{
-				FError::Throwf(TEXT("Found '%s' - expected a '>' or ','."), AllocatorToken.Identifier);
+				FUHTException::Throwf(*this, TEXT("Found '%s' - expected a '>' or ','."), AllocatorToken.Identifier);
 			}
 
 			if (FCString::Strcmp(AllocatorToken.Identifier, TEXT("FMemoryImageSetAllocator")) == 0)
@@ -3778,7 +3742,7 @@ void FHeaderParser::GetVarType(
 			}
 			else
 			{
-				FError::Throwf(TEXT("Found '%s' - explicit allocators are not supported in TMap properties."), AllocatorToken.Identifier);
+				FUHTException::Throwf(*this, TEXT("Found '%s' - explicit allocators are not supported in TMap properties."), AllocatorToken.Identifier);
 			}
 		}
 	}
@@ -3791,23 +3755,23 @@ void FHeaderParser::GetVarType(
 		GetVarType(Scope, VarProperty, Disallow, &VarType, &VarTypePropertyFlags, EPropertyDeclarationStyle::None, VariableCategory);
 		if (VarProperty.IsContainer())
 		{
-			FError::Throwf(TEXT("Nested containers are not supported.") );
+			FUHTException::Throwf(*this, TEXT("Nested containers are not supported.") );
 		}
 		// TODO: Prevent sparse delegate types from being used in a container
 
 		if (VarProperty.Type == CPT_Interface)
 		{
-			FError::Throwf(TEXT("UINTERFACEs are not currently supported as element types."));
+			FUHTException::Throwf(*this, TEXT("UINTERFACEs are not currently supported as element types."));
 		}
 
 		if (VarProperty.Type == CPT_Text)
 		{
-			FError::Throwf(TEXT("FText is not currently supported as an element type."));
+			FUHTException::Throwf(*this, TEXT("FText is not currently supported as an element type."));
 		}
 
 		if (EnumHasAnyFlags(Flags, CPF_Net))
 		{
-			UE_LOG_ERROR_UHT(TEXT("Replicated sets are not supported."));
+			UE_LOG_ERROR_UHT(*this, TEXT("Replicated sets are not supported."));
 		}
 
 		VarTypePropertyFlags = VarProperty.PropertyFlags & (CPF_ContainsInstancedReference | CPF_InstancedReference); // propagate these to the set, we will fix them later
@@ -3816,7 +3780,7 @@ void FHeaderParser::GetVarType(
 		FToken CloseTemplateToken;
 		if (!GetToken(CloseTemplateToken, /*bNoConsts=*/ true, ESymbolParseOption::CloseTemplateBracket))
 		{
-			FError::Throwf(TEXT("Missing token while parsing TArray."));
+			FUHTException::Throwf(*this, TEXT("Missing token while parsing TArray."));
 		}
 
 		if (CloseTemplateToken.TokenType != TOKEN_Symbol || !CloseTemplateToken.Matches(TEXT('>')))
@@ -3824,17 +3788,17 @@ void FHeaderParser::GetVarType(
 			// If we didn't find a comma, report it
 			if (!CloseTemplateToken.Matches(TEXT(',')))
 			{
-				FError::Throwf(TEXT("Expected '>' but found '%s'"), CloseTemplateToken.Identifier);
+				FUHTException::Throwf(*this, TEXT("Expected '>' but found '%s'"), CloseTemplateToken.Identifier);
 			}
 
 			// If we found a comma, read the next thing, assume it's a keyfuncs, and report that
 			FToken AllocatorToken;
 			if (!GetToken(AllocatorToken, /*bNoConsts=*/ true, ESymbolParseOption::CloseTemplateBracket))
 			{
-				FError::Throwf(TEXT("Expected '>' but found '%s'"), CloseTemplateToken.Identifier);
+				FUHTException::Throwf(*this, TEXT("Expected '>' but found '%s'"), CloseTemplateToken.Identifier);
 			}
 
-			FError::Throwf(TEXT("Found '%s' - explicit KeyFuncs are not supported in TSet properties."), AllocatorToken.Identifier);
+			FUHTException::Throwf(*this, TEXT("Found '%s' - explicit KeyFuncs are not supported in TSet properties."), AllocatorToken.Identifier);
 		}
 	}
 	else if ( VarType.Matches(TEXT("FString"), ESearchCase::CaseSensitive) || VarType.Matches(TEXT("FMemoryImageString"), ESearchCase::CaseSensitive))
@@ -3866,7 +3830,7 @@ void FHeaderParser::GetVarType(
 	}
 	else if ( VarType.Matches(TEXT("Text"), ESearchCase::IgnoreCase) )
 	{
-		FError::Throwf(TEXT("%s' is missing a prefix, expecting 'FText'"), VarType.Identifier);
+		FUHTException::Throwf(*this, TEXT("%s' is missing a prefix, expecting 'FText'"), VarType.Identifier);
 	}
 	else if ( VarType.Matches(TEXT("FText"), ESearchCase::CaseSensitive) )
 	{
@@ -3899,13 +3863,13 @@ void FHeaderParser::GetVarType(
 			FToken ScopedTrueEnumName;
 			if (!GetIdentifier(ScopedTrueEnumName, true))
 			{
-				FError::Throwf(TEXT("Expected a namespace scoped enum name.") );
+				FUHTException::Throwf(*this, TEXT("Expected a namespace scoped enum name.") );
 			}
 		}
 
 		if (!bFoundEnum)
 		{
-			FError::Throwf(TEXT("Expected the name of a previously defined enum"));
+			FUHTException::Throwf(*this, TEXT("Expected the name of a previously defined enum"));
 		}
 
 		RequireSymbol(TEXT('>'), VarType.Identifier, ESymbolParseOption::CloseTemplateBracket);
@@ -3918,14 +3882,14 @@ void FHeaderParser::GetVarType(
 		FToken PropertyTypeToken;
 		if (!GetToken(PropertyTypeToken, /* bNoConsts = */ true))
 		{
-			FError::Throwf(TEXT("Expected the property type"));
+			FUHTException::Throwf(*this, TEXT("Expected the property type"));
 		}
 		else
 		{
 			FieldClassName = FName(PropertyTypeToken.Identifier + 1, FNAME_Add);
 			if (!FPropertyTraits::IsValidFieldClass(FieldClassName))
 			{
-				FError::Throwf(TEXT("Undefined property type: %s"), PropertyTypeToken.Identifier);
+				FUHTException::Throwf(*this, TEXT("Undefined property type: %s"), PropertyTypeToken.Identifier);
 			}
 		}
 
@@ -3941,7 +3905,7 @@ void FHeaderParser::GetVarType(
 		{
 			if (EnumDef->GetCppForm() != UEnum::ECppForm::EnumClass)
 			{
-				FError::Throwf(TEXT("You cannot use the raw enum name as a type for member variables, instead use TEnumAsByte or a C++11 enum class with an explicit underlying type."), *EnumDef->GetCppType());
+				FUHTException::Throwf(*this, TEXT("You cannot use the raw enum name as a type for member variables, instead use TEnumAsByte or a C++11 enum class with an explicit underlying type."), *EnumDef->GetCppType());
 			}
 		}
 
@@ -3952,7 +3916,7 @@ void FHeaderParser::GetVarType(
 			FToken ScopedTrueEnumName;
 			if (!GetIdentifier(ScopedTrueEnumName, true))
 			{
-				FError::Throwf(TEXT("Expected a namespace scoped enum name.") );
+				FUHTException::Throwf(*this, TEXT("Expected a namespace scoped enum name.") );
 			}
 		}
 
@@ -4006,7 +3970,7 @@ void FHeaderParser::GetVarType(
 				}
 				else
 				{
-					FError::Throwf(TEXT("Cannot find class '%s', to resolve delegate '%s'"), *IdentifierStripped, DelegateName.Identifier);
+					FUHTException::Throwf(*this, TEXT("Cannot find class '%s', to resolve delegate '%s'"), *IdentifierStripped, DelegateName.Identifier);
 				}
 			}
 		}
@@ -4023,13 +3987,13 @@ void FHeaderParser::GetVarType(
 				FString ExpectedStructName = FString::Printf(TEXT("%s%s"), PrefixCPP, *StructDef->GetName());
 				if (FString(VarType.Identifier) != ExpectedStructName)
 				{
-					FError::Throwf(TEXT("Struct '%s' is missing or has an incorrect prefix, expecting '%s'"), VarType.Identifier, *ExpectedStructName);
+					FUHTException::Throwf(*this, TEXT("Struct '%s' is missing or has an incorrect prefix, expecting '%s'"), VarType.Identifier, *ExpectedStructName);
 				}
 			}
 			else if (!UHTConfig.StructsWithNoPrefix.Contains(VarType.Identifier))
 			{
 				const TCHAR* PrefixCPP = UHTConfig.StructsWithTPrefix.Contains(VarType.Identifier) ? TEXT("T") : StructDef->GetPrefixCPP();
-				FError::Throwf(TEXT("Struct '%s' is missing a prefix, expecting '%s'"), VarType.Identifier, *FString::Printf(TEXT("%s%s"), PrefixCPP, *StructDef->GetName()));
+				FUHTException::Throwf(*this, TEXT("Struct '%s' is missing a prefix, expecting '%s'"), VarType.Identifier, *FString::Printf(TEXT("%s%s"), PrefixCPP, *StructDef->GetName()));
 			}
 
 			bHandledType = true;
@@ -4108,7 +4072,7 @@ void FHeaderParser::GetVarType(
 							TempClassDef = FUnrealClassDefinitionInfo::FindScriptClass(InnerClass.Identifier);
 							if (TempClassDef == nullptr)
 							{
-								FError::Throwf(TEXT("Unrecognized type '%s' (in expression %s<%s>) - type must be a UCLASS"), InnerClass.Identifier, VarType.Identifier, InnerClass.Identifier);
+								FUHTException::Throwf(*this, TEXT("Unrecognized type '%s' (in expression %s<%s>) - type must be a UCLASS"), InnerClass.Identifier, VarType.Identifier, InnerClass.Identifier);
 							}
 
 							if (bIsAutoweakPtrTemplate)
@@ -4137,7 +4101,7 @@ void FHeaderParser::GetVarType(
 						}
 						else
 						{
-							FError::Throwf(TEXT("%s: Missing template type"), VarType.Identifier);
+							FUHTException::Throwf(*this, TEXT("%s: Missing template type"), VarType.Identifier);
 						}
 
 						// Const after template argument type but before end of template symbol
@@ -4180,12 +4144,12 @@ void FHeaderParser::GetVarType(
 							FToken Limitor;
 							if (!GetIdentifier(Limitor))
 							{
-								FError::Throwf(TEXT("'class': Missing class limitor"));
+								FUHTException::Throwf(*this, TEXT("'class': Missing class limitor"));
 							}
 
 							RedirectTypeIdentifier(Limitor);
 
-							VarProperty.MetaClassDef = FUnrealClassDefinitionInfo::FindScriptClassOrThrow(Limitor.Identifier);
+							VarProperty.MetaClassDef = FUnrealClassDefinitionInfo::FindScriptClassOrThrow(*this, Limitor.Identifier);
 
 							RequireSymbol(TEXT('>'), TEXT("'class limitor'"), ESymbolParseOption::CloseTemplateBracket);
 						}
@@ -4196,16 +4160,16 @@ void FHeaderParser::GetVarType(
 
 						if (PropertyType == CPT_WeakObjectReference)
 						{
-							FError::Throwf(TEXT("Class variables cannot be weak, they are always strong."));
+							FUHTException::Throwf(*this, TEXT("Class variables cannot be weak, they are always strong."));
 						}
 						else if (PropertyType == CPT_LazyObjectReference)
 						{
-							FError::Throwf(TEXT("Class variables cannot be lazy, they are always strong."));
+							FUHTException::Throwf(*this, TEXT("Class variables cannot be lazy, they are always strong."));
 						}
 
 						if (bIsSoftObjectPtrTemplate)
 						{
-							FError::Throwf(TEXT("Class variables cannot be stored in TSoftObjectPtr, use TSoftClassPtr instead."));
+							FUHTException::Throwf(*this, TEXT("Class variables cannot be stored in TSoftObjectPtr, use TSoftClassPtr instead."));
 						}
 					}
 
@@ -4260,7 +4224,7 @@ void FHeaderParser::GetVarType(
 
 				if (!bHandledType)
 				{
-					FError::Throwf(TEXT("Unrecognized type '%s' - type must be a UCLASS, USTRUCT or UENUM"), VarType.Identifier);
+					FUHTException::Throwf(*this, TEXT("Unrecognized type '%s' - type must be a UCLASS, USTRUCT or UENUM"), VarType.Identifier);
 				}
 			}
 		}
@@ -4280,32 +4244,32 @@ void FHeaderParser::GetVarType(
 	{
 		if (VariableCategory == EVariableCategory::Member)
 		{
-			FError::Throwf(TEXT("Const properties are not supported."));
+			FUHTException::Throwf(*this, TEXT("Const properties are not supported."));
 		}
 		else
 		{
-			FError::Throwf(TEXT("Inappropriate keyword 'const' on variable of type '%s'"), VarType.Identifier);
+			FUHTException::Throwf(*this, TEXT("Inappropriate keyword 'const' on variable of type '%s'"), VarType.Identifier);
 		}
 	}
 
 	if (bUnconsumedClassKeyword)
 	{
-		FError::Throwf(TEXT("Inappropriate keyword 'class' on variable of type '%s'"), VarType.Identifier );
+		FUHTException::Throwf(*this, TEXT("Inappropriate keyword 'class' on variable of type '%s'"), VarType.Identifier );
 	}
 
 	if (bUnconsumedStructKeyword)
 	{
-		FError::Throwf(TEXT("Inappropriate keyword 'struct' on variable of type '%s'"), VarType.Identifier );
+		FUHTException::Throwf(*this, TEXT("Inappropriate keyword 'struct' on variable of type '%s'"), VarType.Identifier );
 	}
 
 	if (bUnconsumedEnumKeyword)
 	{
-		FError::Throwf(TEXT("Inappropriate keyword 'enum' on variable of type '%s'"), VarType.Identifier );
+		FUHTException::Throwf(*this, TEXT("Inappropriate keyword 'enum' on variable of type '%s'"), VarType.Identifier );
 	}
 
 	if (MatchSymbol(TEXT('*')))
 	{
-		FError::Throwf(TEXT("Inappropriate '*' on variable of type '%s', cannot have an exposed pointer to this type."), VarType.Identifier );
+		FUHTException::Throwf(*this, TEXT("Inappropriate '*' on variable of type '%s', cannot have an exposed pointer to this type."), VarType.Identifier );
 	}
 
 	//@TODO: UCREMOVAL: 'const' member variables that will get written post-construction by defaultproperties
@@ -4337,7 +4301,7 @@ void FHeaderParser::GetVarType(
 			{
 				if (!(Flags & CPF_ConstParm))
 				{
-					FError::Throwf(TEXT("Replicated %s parameters cannot be passed by non-const reference"), VarType.Identifier);
+					FUHTException::Throwf(*this, TEXT("Replicated %s parameters cannot be passed by non-const reference"), VarType.Identifier);
 				}
 
 				Flags |= CPF_ReferenceParm;
@@ -4375,7 +4339,7 @@ void FHeaderParser::GetVarType(
 		}
 		else
 		{
-			FError::Throwf(TEXT("Must specify a valid function name for replication notifications"));
+			FUHTException::Throwf(*this, TEXT("Must specify a valid function name for replication notifications"));
 		}
 	}
 
@@ -4386,33 +4350,33 @@ void FHeaderParser::GetVarType(
 		{
 			if (VarProperty.ClassDef->GetClass()->IsChildOf<UClass>())
 			{
-				FError::Throwf(TEXT("'Instanced' cannot be applied to class properties (UClass* or TSubclassOf<>)"));
+				FUHTException::Throwf(*this, TEXT("'Instanced' cannot be applied to class properties (UClass* or TSubclassOf<>)"));
 			}
 		}
 		else
 		{
-			FError::Throwf(TEXT("'Instanced' is only allowed on an object property, an array of objects, a set of objects, or a map with an object value type."));
+			FUHTException::Throwf(*this, TEXT("'Instanced' is only allowed on an object property, an array of objects, a set of objects, or a map with an object value type."));
 		}
 	}
 
 	if ( VarProperty.IsObjectOrInterface() && VarProperty.Type != CPT_SoftObjectReference && VarProperty.MetaClassDef == nullptr && (VarProperty.PropertyFlags&CPF_Config) != 0 )
 	{
-		FError::Throwf(TEXT("Not allowed to use 'config' with object variables"));
+		FUHTException::Throwf(*this, TEXT("Not allowed to use 'config' with object variables"));
 	}
 
 	if ((VarProperty.PropertyFlags & CPF_BlueprintAssignable) && VarProperty.Type != CPT_MulticastDelegate)
 	{
-		FError::Throwf(TEXT("'BlueprintAssignable' is only allowed on multicast delegate properties"));
+		FUHTException::Throwf(*this, TEXT("'BlueprintAssignable' is only allowed on multicast delegate properties"));
 	}
 
 	if ((VarProperty.PropertyFlags & CPF_BlueprintCallable) && VarProperty.Type != CPT_MulticastDelegate)
 	{
-		FError::Throwf(TEXT("'BlueprintCallable' is only allowed on a property when it is a multicast delegate"));
+		FUHTException::Throwf(*this, TEXT("'BlueprintCallable' is only allowed on a property when it is a multicast delegate"));
 	}
 
 	if ((VarProperty.PropertyFlags & CPF_BlueprintAuthorityOnly) && VarProperty.Type != CPT_MulticastDelegate)
 	{
-		FError::Throwf(TEXT("'BlueprintAuthorityOnly' is only allowed on a property when it is a multicast delegate"));
+		FUHTException::Throwf(*this, TEXT("'BlueprintAuthorityOnly' is only allowed on a property when it is a multicast delegate"));
 	}
 	
 	if (VariableCategory != EVariableCategory::Member)
@@ -4426,13 +4390,13 @@ void FHeaderParser::GetVarType(
 	if (Transients && !Cast<UClass>(OwnerStruct))
 	{
 		TArray<const TCHAR*> FlagStrs = ParsePropertyFlags(Transients);
-		FError::Throwf(TEXT("'%s' specifier(s) are only allowed on class member variables"), *FString::Join(FlagStrs, TEXT(", ")));
+		FUHTException::Throwf(*this, TEXT("'%s' specifier(s) are only allowed on class member variables"), *FString::Join(FlagStrs, TEXT(", ")));
 	}
 
 	// Make sure the overrides are allowed here.
 	if( VarProperty.PropertyFlags & Disallow )
 	{
-		FError::Throwf(TEXT("Specified type modifiers not allowed here") );
+		FUHTException::Throwf(*this, TEXT("Specified type modifiers not allowed here") );
 	}
 
 	// For now, copy the flags that a TMap value has to the key
@@ -4466,7 +4430,7 @@ FUnrealPropertyDefinitionInfo& FHeaderParser::GetVarNameAndDim
 	ELayoutMacroType        LayoutMacroType
 )
 {
-	const TCHAR* HintText = GetHintText(VariableCategory);
+	const TCHAR* HintText = GetHintText(*this, VariableCategory);
 
 	AddModuleRelativePathToMetadata(ParentStruct, VarProperty.MetaData);
 	TCHAR Identifier[NAME_SIZE];
@@ -4482,7 +4446,7 @@ FUnrealPropertyDefinitionInfo& FHeaderParser::GetVarNameAndDim
 		FToken VarToken;
 		if (!GetIdentifier(VarToken))
 		{
-			FError::Throwf(TEXT("Missing variable name") );
+			FUHTException::Throwf(*this, TEXT("Missing variable name") );
 		}
 
 		switch (LayoutMacroType)
@@ -4513,7 +4477,7 @@ FUnrealPropertyDefinitionInfo& FHeaderParser::GetVarNameAndDim
 		{
 			if (DeprecatedIndex != VarName.Len() - 11)
 			{
-				FError::Throwf(TEXT("Deprecated variables must end with _DEPRECATED"));
+				FUHTException::Throwf(*this, TEXT("Deprecated variables must end with _DEPRECATED"));
 			}
 
 			// We allow deprecated properties in blueprints that have getters and setters assigned as they may be part of a backwards compatibility path
@@ -4523,12 +4487,12 @@ FUnrealPropertyDefinitionInfo& FHeaderParser::GetVarNameAndDim
 
 			if (bWarnOnGetter)
 			{
-				UE_LOG_WARNING_UHT(TEXT("%s: Deprecated property '%s' should not be marked as blueprint visible without having a BlueprintGetter"), HintText, *VarName);
+				UE_LOG_WARNING_UHT(*this, TEXT("%s: Deprecated property '%s' should not be marked as blueprint visible without having a BlueprintGetter"), HintText, *VarName);
 			}
 
 			if (bWarnOnSetter)
 			{
-				UE_LOG_WARNING_UHT(TEXT("%s: Deprecated property '%s' should not be marked as blueprint writeable without having a BlueprintSetter"), HintText, *VarName);
+				UE_LOG_WARNING_UHT(*this, TEXT("%s: Deprecated property '%s' should not be marked as blueprint writeable without having a BlueprintSetter"), HintText, *VarName);
 			}
 
 
@@ -4536,7 +4500,7 @@ FUnrealPropertyDefinitionInfo& FHeaderParser::GetVarNameAndDim
 			if (VarProperty.PropertyFlags & (CPF_Edit | CPF_EditConst) || // Property is marked as editable
 				(!bBlueprintVisible && (VarProperty.PropertyFlags & CPF_BlueprintReadOnly) && !(VarProperty.ImpliedPropertyFlags & CPF_BlueprintReadOnly)) ) // Is BPRO, but not via Implied Flags and not caught by Getter/Setter path above
 			{
-				UE_LOG_WARNING_UHT(TEXT("%s: Deprecated property '%s' should not be marked as visible or editable"), HintText, *VarName);
+				UE_LOG_WARNING_UHT(*this, TEXT("%s: Deprecated property '%s' should not be marked as visible or editable"), HintText, *VarName);
 			}
 
 			VarProperty.PropertyFlags |= CPF_Deprecated;
@@ -4576,7 +4540,7 @@ FUnrealPropertyDefinitionInfo& FHeaderParser::GetVarNameAndDim
 
 		if (bErrorDueToShadowing)
 		{
-			FError::Throwf(TEXT("%s: '%s' cannot be defined in '%s' as it is already defined in scope '%s' (shadowing is not allowed)"),
+			FUHTException::Throwf(*this, TEXT("%s: '%s' cannot be defined in '%s' as it is already defined in scope '%s' (shadowing is not allowed)"),
 				HintText,
 				Identifier,
 				*ParentStruct.GetName(),
@@ -4592,24 +4556,24 @@ FUnrealPropertyDefinitionInfo& FHeaderParser::GetVarNameAndDim
 		{
 			case EVariableCategory::Return:
 			{
-				FError::Throwf(TEXT("Arrays aren't allowed as return types"));
+				FUHTException::Throwf(*this, TEXT("Arrays aren't allowed as return types"));
 			}
 
 			case EVariableCategory::RegularParameter:
 			case EVariableCategory::ReplicatedParameter:
 			{
-				FError::Throwf(TEXT("Arrays aren't allowed as function parameters"));
+				FUHTException::Throwf(*this, TEXT("Arrays aren't allowed as function parameters"));
 			}
 		}
 
 		if (VarProperty.IsContainer())
 		{
-			FError::Throwf(TEXT("Static arrays of containers are not allowed"));
+			FUHTException::Throwf(*this, TEXT("Static arrays of containers are not allowed"));
 		}
 
 		if (VarProperty.IsBool())
 		{
-			FError::Throwf(TEXT("Bool arrays are not allowed") );
+			FUHTException::Throwf(*this, TEXT("Bool arrays are not allowed") );
 		}
 
 		if (LayoutMacroType == ELayoutMacroType::None)
@@ -4617,7 +4581,7 @@ FUnrealPropertyDefinitionInfo& FHeaderParser::GetVarNameAndDim
 			// Ignore how the actual array dimensions are actually defined - we'll calculate those with the compiler anyway.
 			if (!GetRawToken(Dimensions, TEXT(']')))
 			{
-				FError::Throwf(TEXT("%s %s: Missing ']'"), HintText, Identifier);
+				FUHTException::Throwf(*this, TEXT("%s %s: Missing ']'"), HintText, Identifier);
 			}
 		}
 		else
@@ -4625,7 +4589,7 @@ FUnrealPropertyDefinitionInfo& FHeaderParser::GetVarNameAndDim
 			// Ignore how the actual array dimensions are actually defined - we'll calculate those with the compiler anyway.
 			if (!GetRawToken(Dimensions, TEXT(')')))
 			{
-				FError::Throwf(TEXT("%s %s: Missing ']'"), HintText, Identifier);
+				FUHTException::Throwf(*this, TEXT("%s %s: Missing ']'"), HintText, Identifier);
 			}
 		}
 
@@ -4716,12 +4680,12 @@ FUnrealPropertyDefinitionInfo& FHeaderParser::GetVarNameAndDim
 		// WeakPointer is supported by VM as return type (see UObject::execLetWeakObjPtr), but there is no P_GET_... macro for WeakPointer.
 		if (VarProperty.Type == CPT_LazyObjectReference)
 		{
-			FError::Throwf(TEXT("UFunctions cannot take a lazy pointer as a parameter."));
+			FUHTException::Throwf(*this, TEXT("UFunctions cannot take a lazy pointer as a parameter."));
 		}
 		else if (VarProperty.Type == CPT_ObjectPtrReference)
 		{
 			// @TODO: OBJPTR: Investigate TObjectPtr support for UFunction parameters.
-			FError::Throwf(TEXT("UFunctions cannot take a TObjectPtr as a parameter."));
+			FUHTException::Throwf(*this, TEXT("UFunctions cannot take a TObjectPtr as a parameter."));
 		}
 	}
 
@@ -4748,7 +4712,7 @@ bool FHeaderParser::CompileDeclaration(TArray<FUnrealFunctionDefinitionInfo*>& D
 	{
 		if (!IsAllowedInThisNesting(ENestAllowFlags::VarDecl) && !IsAllowedInThisNesting(ENestAllowFlags::Function))
 		{
-			FError::Throwf(TEXT("Access specifier %s not allowed here."), Token.Identifier);
+			FUHTException::Throwf(*this, TEXT("Access specifier %s not allowed here."), Token.Identifier);
 		}
 		check(TopNest->NestType == ENestType::Class || TopNest->NestType == ENestType::Interface || TopNest->NestType == ENestType::NativeInterface);
 		CurrentAccessSpecifier = AccessSpecifier;
@@ -4760,7 +4724,7 @@ bool FHeaderParser::CompileDeclaration(TArray<FUnrealFunctionDefinitionInfo*>& D
 		// Make sure the previous class ended with valid nesting.
 		if (bEncounteredNewStyleClass_UnmatchedBrackets)
 		{
-			FError::Throwf(TEXT("Missing } at end of class"));
+			FUHTException::Throwf(*this, TEXT("Missing } at end of class"));
 		}
 
 		// Start parsing the second class
@@ -4780,7 +4744,7 @@ bool FHeaderParser::CompileDeclaration(TArray<FUnrealFunctionDefinitionInfo*>& D
 	{
 		if (TopNest->NestType != ENestType::NativeInterface)
 		{
-			FError::Throwf(TEXT("%s must occur inside the native interface definition"), Token.Identifier);
+			FUHTException::Throwf(*this, TEXT("%s must occur inside the native interface definition"), Token.Identifier);
 		}
 		RequireSymbol(TEXT('('), Token.Identifier);
 		CompileVersionDeclaration(GetCurrentClassDef());
@@ -4790,13 +4754,13 @@ bool FHeaderParser::CompileDeclaration(TArray<FUnrealFunctionDefinitionInfo*>& D
 		if (StructData.ParsedInterface == EParsedInterface::NotAnInterface)
 		{
 			FString CurrentClassName = GetCurrentClass()->GetName();
-			FError::Throwf(TEXT("Could not find the associated 'U%s' class while parsing 'I%s' - it could be missing or malformed"), *CurrentClassName, *CurrentClassName);
+			FUHTException::Throwf(*this, TEXT("Could not find the associated 'U%s' class while parsing 'I%s' - it could be missing or malformed"), *CurrentClassName, *CurrentClassName);
 		}
 
 		if (StructData.ParsedInterface == EParsedInterface::ParsedIInterface)
 		{
 			FString CurrentClassName = GetCurrentClass()->GetName();
-			FError::Throwf(TEXT("Duplicate IInterface definition found while parsing 'I%s'"), *CurrentClassName);
+			FUHTException::Throwf(*this, TEXT("Duplicate IInterface definition found while parsing 'I%s'"), *CurrentClassName);
 		}
 
 		check(StructData.ParsedInterface == EParsedInterface::ParsedUInterface);
@@ -4823,7 +4787,7 @@ bool FHeaderParser::CompileDeclaration(TArray<FUnrealFunctionDefinitionInfo*>& D
 	{
 		if (TopNest->NestType != ENestType::Interface)
 		{
-			FError::Throwf(TEXT("%s must occur inside the interface definition"), Token.Identifier);
+			FUHTException::Throwf(*this, TEXT("%s must occur inside the interface definition"), Token.Identifier);
 		}
 		RequireSymbol(TEXT('('), Token.Identifier);
 		CompileVersionDeclaration(GetCurrentClassDef());
@@ -4847,7 +4811,7 @@ bool FHeaderParser::CompileDeclaration(TArray<FUnrealFunctionDefinitionInfo*>& D
 	{
 		if (TopNest->NestType != ENestType::Class)
 		{
-			FError::Throwf(TEXT("%s must occur inside the class definition"), Token.Identifier);
+			FUHTException::Throwf(*this, TEXT("%s must occur inside the class definition"), Token.Identifier);
 		}
 
 		FStructMetaData& StructData = GetCurrentClassData();
@@ -4962,13 +4926,13 @@ bool FHeaderParser::CompileDeclaration(TArray<FUnrealFunctionDefinitionInfo*>& D
 			// Ensure the UINTERFACE classes have a GENERATED_BODY declaration
 			if (bHaveSeenUClass && !bClassHasGeneratedUInterfaceBody)
 			{
-				FError::Throwf(TEXT("Expected a GENERATED_BODY() at the start of class"));
+				FUHTException::Throwf(*this, TEXT("Expected a GENERATED_BODY() at the start of class"));
 			}
 
 			// Ensure the non-UINTERFACE interface classes have a GENERATED_BODY declaration
 			if (!bHaveSeenUClass && !bClassHasGeneratedIInterfaceBody)
 			{
-				FError::Throwf(TEXT("Expected a GENERATED_BODY() at the start of class"));
+				FUHTException::Throwf(*this, TEXT("Expected a GENERATED_BODY() at the start of class"));
 			}
 		}
 		else
@@ -4979,7 +4943,7 @@ bool FHeaderParser::CompileDeclaration(TArray<FUnrealFunctionDefinitionInfo*>& D
 			// Ensure classes have a GENERATED_BODY declaration
 			if (bHaveSeenUClass && !bClassHasGeneratedBody)
 			{
-				FError::Throwf(TEXT("Expected a GENERATED_BODY() at the start of class"));
+				FUHTException::Throwf(*this, TEXT("Expected a GENERATED_BODY() at the start of class"));
 			}
 		}
 
@@ -4996,11 +4960,11 @@ bool FHeaderParser::CompileDeclaration(TArray<FUnrealFunctionDefinitionInfo*>& D
 	{
 		if (GetToken(Token))
 		{
-			FError::Throwf(TEXT("Extra ';' before '%s'"), Token.Identifier);
+			FUHTException::Throwf(*this, TEXT("Extra ';' before '%s'"), Token.Identifier);
 		}
 		else
 		{
-			FError::Throwf(TEXT("Extra ';' before end of file"));
+			FUHTException::Throwf(*this, TEXT("Extra ';' before end of file"));
 		}
 	}
 
@@ -5140,7 +5104,7 @@ bool FHeaderParser::CompileDeclaration(TArray<FUnrealFunctionDefinitionInfo*>& D
 						}
 						else
 						{
-							FError::Throwf(TEXT("Serialize functions must not be inside preprocessor blocks, except for WITH_EDITORONLY_DATA"));
+							FUHTException::Throwf(*this, TEXT("Serialize functions must not be inside preprocessor blocks, except for WITH_EDITORONLY_DATA"));
 						}
 					}
 				}
@@ -5223,7 +5187,7 @@ bool FHeaderParser::SkipDeclaration(FToken& Token)
 
 			if (NestedScopes < 0)
 			{
-				FError::Throwf(TEXT("Unexpected '}'. Did you miss a semi-colon?"));
+				FUHTException::Throwf(*this, TEXT("Unexpected '}'. Did you miss a semi-colon?"));
 			}
 		}
 		else if (bMacroDeclaration && NestedScopes == 0)
@@ -5253,7 +5217,7 @@ bool FHeaderParser::SkipDeclaration(FToken& Token)
 			}
 			else if (!SafeMatchSymbol(TEXT(';')))
 			{
-				FError::Throwf(*FString::Printf(TEXT("Unexpected '%s'. Did you miss a semi-colon?"), VariableName.Identifier));
+				FUHTException::Throwf(*this, TEXT("Unexpected '%s'. Did you miss a semi-colon?"), VariableName.Identifier);
 			}
 		}
 
@@ -5318,7 +5282,7 @@ FUnrealClassDefinitionInfo& FHeaderParser::ParseClassNameDeclaration(FString& De
 		{
 			if (!BaseClassDef->HasAnyClassFlags(CLASS_Interface))
 			{
-				FError::Throwf(TEXT("Implements: Class %s is not an interface; Can only inherit from non-UObjects or UInterface derived interfaces"), *BaseClassDef->GetName());
+				FUHTException::Throwf(*this, TEXT("Implements: Class %s is not an interface; Can only inherit from non-UObjects or UInterface derived interfaces"), *BaseClassDef->GetName());
 			}
 
 			// Propagate the inheritable ClassFlags
@@ -5365,7 +5329,7 @@ void PostParsingClassSetup(FUnrealClassDefinitionInfo& ClassDef)
 		ClassDef.SetClassConfigName(ClassDef.GetSuperClass() ? ClassDef.GetSuperClass()->GetClassConfigName() : NAME_None);
 		if (ClassDef.GetClassConfigName() == NAME_None)
 		{
-			FError::Throwf(TEXT("Classes with config / globalconfig member variables need to specify config file."));
+			FUHTException::Throwf(ClassDef, TEXT("Classes with config / globalconfig member variables need to specify config file."));
 			ClassDef.SetClassConfigName(NAME_Engine);
 		}
 	}
@@ -5430,7 +5394,7 @@ FUnrealClassDefinitionInfo& FHeaderParser::CompileClassDeclaration()
 		bool bIsIntrinsic = TempDef->HasAnyClassFlags(CLASS_Intrinsic);
 		if (!(bIsParsed || bIsIntrinsic))
 		{
-			FError::Throwf(TEXT("'%s' can't be compiled: Parent class '%s' has errors"), *ClassDef.GetName(), *TempDef->GetName());
+			FUHTException::Throwf(*this, TEXT("'%s' can't be compiled: Parent class '%s' has errors"), *ClassDef.GetName(), *TempDef->GetName());
 		}
 	}
 
@@ -5530,7 +5494,7 @@ FUnrealClassDefinitionInfo* FHeaderParser::ParseInterfaceNameDeclaration(FString
 		// UInterface is special and actually extends from UObject, which isn't an interface
 		if (DeclaredInterfaceName != TEXT("UInterface"))
 		{
-			FError::Throwf(TEXT("Interface class '%s' cannot inherit from non-interface class '%s'"), *DeclaredInterfaceName, *TempClassDef->GetName());
+			FUHTException::Throwf(*this, TEXT("Interface class '%s' cannot inherit from non-interface class '%s'"), *DeclaredInterfaceName, *TempClassDef->GetName());
 		}
 	}
 
@@ -5539,7 +5503,7 @@ FUnrealClassDefinitionInfo* FHeaderParser::ParseInterfaceNameDeclaration(FString
 	check(SuperClassDef);
 	if (SuperClassDef != TempClassDef)
 	{
-		FError::Throwf(TEXT("%s's superclass must be %s, not %s"), *ClassDef->GetPathName(), *SuperClassDef->GetPathName(), *TempClassDef->GetPathName());
+		FUHTException::Throwf(*this, TEXT("%s's superclass must be %s, not %s"), *ClassDef->GetPathName(), *SuperClassDef->GetPathName(), *TempClassDef->GetPathName());
 	}
 
 	return ClassDef;
@@ -5627,7 +5591,7 @@ void FHeaderParser::CompileInterfaceDeclaration()
 	// All classes that are parsed are expected to be native
 	if (InterfaceClassDef->GetSuperClass() && !InterfaceClassDef->GetSuperClass()->HasAnyClassFlags(CLASS_Native))
 	{
-		FError::Throwf(TEXT("Native classes cannot extend non-native classes") );
+		FUHTException::Throwf(*this, TEXT("Native classes cannot extend non-native classes") );
 	}
 
 	InterfaceClassDef->SetInternalFlags(EInternalObjectFlags::Native);
@@ -5640,13 +5604,13 @@ void FHeaderParser::CompileInterfaceDeclaration()
 		{
 			default:
 			{
-				FError::Throwf(TEXT("Unknown interface specifier '%s'"), *Specifier.Key);
+				FUHTException::Throwf(*this, TEXT("Unknown interface specifier '%s'"), *Specifier.Key);
 			}
 			break;
 
 			case EInterfaceSpecifier::DependsOn:
 			{
-				FError::Throwf(TEXT("The dependsOn specifier is deprecated. Please use #include \"ClassHeaderFilename.h\" instead."));
+				FUHTException::Throwf(*this, TEXT("The dependsOn specifier is deprecated. Please use #include \"ClassHeaderFilename.h\" instead."));
 			}
 			break;
 
@@ -5668,7 +5632,7 @@ void FHeaderParser::CompileInterfaceDeclaration()
 	const FString ExpectedInterfaceName = InterfaceClassDef->GetNameWithPrefix(EEnforceInterfacePrefix::U);
 	if (DeclaredInterfaceName != ExpectedInterfaceName)
 	{
-		FError::Throwf(TEXT("Interface name '%s' is invalid, the first class should be identified as '%s'"), *DeclaredInterfaceName, *ExpectedInterfaceName );
+		FUHTException::Throwf(*this, TEXT("Interface name '%s' is invalid, the first class should be identified as '%s'"), *DeclaredInterfaceName, *ExpectedInterfaceName );
 	}
 
 	// Try parsing metadata for the interface
@@ -5696,7 +5660,7 @@ void FHeaderParser::CompileRigVMMethodDeclaration(FUnrealStructDefinitionInfo& S
 {
 	if (!MatchSymbol(TEXT("(")))
 	{
-		FError::Throwf(TEXT("Bad RIGVM_METHOD definition"));
+		FUHTException::Throwf(*this, TEXT("Bad RIGVM_METHOD definition"));
 	}
 
 	// find the next close brace
@@ -5734,7 +5698,7 @@ void FHeaderParser::CompileRigVMMethodDeclaration(FUnrealStructDefinitionInfo& S
 
 	if (!MatchSymbol(TEXT("(")))
 	{
-		FError::Throwf(TEXT("Bad RIGVM_METHOD definition"));
+		FUHTException::Throwf(*this, TEXT("Bad RIGVM_METHOD definition"));
 	}
 
 	TArray<FString> ParamsContent;
@@ -5870,7 +5834,7 @@ void FHeaderParser::ParseRigVMMethodParameters(FUnrealStructDefinitionInfo& Stru
 
 		if (Parameter.bEditorOnly)
 		{
-			UE_LOG_ERROR_UHT(TEXT("RigVM Struct '%s' - Member '%s' is editor only - WITH_EDITORONLY_DATA not allowed on structs with RIGVM_METHOD."), *StructDef.GetName(), *Parameter.Name, *MemberCPPType);
+			UE_LOG_ERROR_UHT(*this, TEXT("RigVM Struct '%s' - Member '%s' is editor only - WITH_EDITORONLY_DATA not allowed on structs with RIGVM_METHOD."), *StructDef.GetName(), *Parameter.Name, *MemberCPPType);
 		}
 
 		if (!ExtendedCPPType.IsEmpty())
@@ -5878,7 +5842,7 @@ void FHeaderParser::ParseRigVMMethodParameters(FUnrealStructDefinitionInfo& Stru
 			// we only support arrays - no maps or similar data structures
 			if (MemberCPPType != TArrayText && MemberCPPType != TEnumAsByteText)
 			{
-				UE_LOG_ERROR_UHT(TEXT("RigVM Struct '%s' - Member '%s' type '%s' not supported by RigVM."), *StructDef.GetName(), *Parameter.Name, *MemberCPPType);
+				UE_LOG_ERROR_UHT(*this, TEXT("RigVM Struct '%s' - Member '%s' type '%s' not supported by RigVM."), *StructDef.GetName(), *Parameter.Name, *MemberCPPType);
 				continue;
 			}
 		}
@@ -5904,12 +5868,12 @@ void FHeaderParser::ParseRigVMMethodParameters(FUnrealStructDefinitionInfo& Stru
 
 	if (StructRigVMInfo.Members.Num() == 0)
 	{
-		UE_LOG_ERROR_UHT(TEXT("RigVM Struct '%s' - has zero members - invalid RIGVM_METHOD."), *StructDef.GetName());
+		UE_LOG_ERROR_UHT(*this, TEXT("RigVM Struct '%s' - has zero members - invalid RIGVM_METHOD."), *StructDef.GetName());
 	}
 
 	if (StructRigVMInfo.Members.Num() > 64)
 	{
-		UE_LOG_ERROR_UHT(TEXT("RigVM Struct '%s' - has %d members (64 is the limit)."), *StructDef.GetName(), StructRigVMInfo.Members.Num());
+		UE_LOG_ERROR_UHT(*this, TEXT("RigVM Struct '%s' - has %d members (64 is the limit)."), *StructDef.GetName(), StructRigVMInfo.Members.Num());
 	}
 }
 
@@ -5956,14 +5920,14 @@ void FHeaderParser::ParseParameterList(FUnrealFunctionDefinitionInfo& FunctionDe
 			{
 				if (!FunctionDef.HasAnyFunctionFlags(FUNC_NetRequest | FUNC_NetResponse))
 				{
-					UE_LOG_ERROR_UHT(TEXT("Maps are not supported in an RPC."));
+					UE_LOG_ERROR_UHT(*this, TEXT("Maps are not supported in an RPC."));
 				}
 			}
 			else if (Property.ArrayType == EArrayType::Set)
 			{
 				if (!FunctionDef.HasAnyFunctionFlags(FUNC_NetRequest | FUNC_NetResponse))
 				{
-					UE_LOG_ERROR_UHT(TEXT("Sets are not supported in an RPC."));
+					UE_LOG_ERROR_UHT(*this, TEXT("Sets are not supported in an RPC."));
 				}
 			}
 
@@ -5979,39 +5943,39 @@ void FHeaderParser::ParseParameterList(FUnrealFunctionDefinitionInfo& FunctionDe
 			{
 				if (Property.PropertyFlags & CPF_OutParm)
 				{
-					UE_LOG_ERROR_UHT(TEXT("Replicated functions cannot contain out parameters"));
+					UE_LOG_ERROR_UHT(*this, TEXT("Replicated functions cannot contain out parameters"));
 				}
 
 				if (Property.PropertyFlags & CPF_RepSkip)
 				{
-					UE_LOG_ERROR_UHT(TEXT("Only service request functions cannot contain NoReplication parameters"));
+					UE_LOG_ERROR_UHT(*this, TEXT("Only service request functions cannot contain NoReplication parameters"));
 				}
 
 				if (PropDef.GetPropertyBase().IsDelegateOrDelegateStaticArray())
 				{
-					UE_LOG_ERROR_UHT(TEXT("Replicated functions cannot contain delegate parameters (this would be insecure)"));
+					UE_LOG_ERROR_UHT(*this, TEXT("Replicated functions cannot contain delegate parameters (this would be insecure)"));
 				}
 
 				if (Property.Type == CPT_String && Property.RefQualifier != ERefQualifier::ConstRef && !PropDef.IsStaticArray())
 				{
-					UE_LOG_ERROR_UHT(TEXT("Replicated FString parameters must be passed by const reference"));
+					UE_LOG_ERROR_UHT(*this, TEXT("Replicated FString parameters must be passed by const reference"));
 				}
 
 				if (Property.ArrayType == EArrayType::Dynamic && Property.RefQualifier != ERefQualifier::ConstRef && !PropDef.IsStaticArray())
 				{
-					UE_LOG_ERROR_UHT(TEXT("Replicated TArray parameters must be passed by const reference"));
+					UE_LOG_ERROR_UHT(*this, TEXT("Replicated TArray parameters must be passed by const reference"));
 				}
 			}
 			else
 			{
 				if (!(Property.PropertyFlags & CPF_RepSkip) && (Property.PropertyFlags & CPF_OutParm))
 				{
-					UE_LOG_ERROR_UHT(TEXT("Service request functions cannot contain out parameters, unless marked NotReplicated"));
+					UE_LOG_ERROR_UHT(*this, TEXT("Service request functions cannot contain out parameters, unless marked NotReplicated"));
 				}
 
 				if (!(Property.PropertyFlags & CPF_RepSkip) && PropDef.GetPropertyBase().IsDelegateOrDelegateStaticArray())
 				{
-					UE_LOG_ERROR_UHT(TEXT("Service request functions cannot contain delegate parameters, unless marked NotReplicated"));
+					UE_LOG_ERROR_UHT(*this, TEXT("Service request functions cannot contain delegate parameters, unless marked NotReplicated"));
 				}
 			}
 		}
@@ -6024,7 +5988,7 @@ void FHeaderParser::ParseParameterList(FUnrealFunctionDefinitionInfo& FunctionDe
 					if (EnumDef->GetUnderlyingType() != EUnderlyingEnumType::uint8 &&
 						EnumDef->GetUnderlyingType() != EUnderlyingEnumType::Unspecified)
 					{
-						FError::Throwf(TEXT("Invalid enum param for Blueprints - currently only uint8 supported"));
+						FUHTException::Throwf(*this, TEXT("Invalid enum param for Blueprints - currently only uint8 supported"));
 					}
 				}
 			}
@@ -6040,7 +6004,7 @@ void FHeaderParser::ParseParameterList(FUnrealFunctionDefinitionInfo& FunctionDe
 				{
 					if (FCString::Stricmp(*PropDef.GetNameCPP(), *InvalidName) == 0)
 					{
-						UE_LOG_ERROR_UHT(TEXT("Paramater name '%s' in function is invalid, '%s' is a reserved name."), *InvalidName, *InvalidName);
+						UE_LOG_ERROR_UHT(*this, TEXT("Paramater name '%s' in function is invalid, '%s' is a reserved name."), *InvalidName, *InvalidName);
 					}
 				}
 			}
@@ -6093,11 +6057,11 @@ void FHeaderParser::ParseParameterList(FUnrealFunctionDefinitionInfo& FunctionDe
 					const bool bDefaultValueParsed = FPropertyTraits::DefaultValueStringCppFormatToInnerFormat(PropDef, DefaultArgText, InnerDefaultValue);
 					if (!bDefaultValueParsed)
 					{
-						FError::Throwf(TEXT("C++ Default parameter not parsed: %s \"%s\" "), *PropDef.GetName(), *DefaultArgText);
+						FUHTException::Throwf(*this, TEXT("C++ Default parameter not parsed: %s \"%s\" "), *PropDef.GetName(), *DefaultArgText);
 					}
 
 					MetaData->Add(KeyName, InnerDefaultValue);
-						UE_LOG(LogCompile, Verbose, TEXT("C++ Default parameter parsed: %s \"%s\" -> \"%s\" "), *PropDef.GetName(), *DefaultArgText, *InnerDefaultValue );
+					UE_LOG(LogCompile, Verbose, TEXT("C++ Default parameter parsed: %s \"%s\" -> \"%s\" "), *PropDef.GetName(), *DefaultArgText, *InnerDefaultValue );
 				}
 			}
 		}
@@ -6120,14 +6084,14 @@ FUnrealFunctionDefinitionInfo& FHeaderParser::CompileDelegateDeclaration(const T
 		TArray<FPropertySpecifier> SpecifiersFound;
 		ReadSpecifierSetInsideMacro(SpecifiersFound, TEXT("Delegate"), MetaData);
 
-		ProcessFunctionSpecifiers(FuncInfo, SpecifiersFound, MetaData);
+		ProcessFunctionSpecifiers(*this, FuncInfo, SpecifiersFound, MetaData);
 
 		// Get the next token and ensure it looks like a delegate
 		FToken Token;
 		GetToken(Token);
 		if (!IsValidDelegateDeclaration(Token))
 		{
-			FError::Throwf(TEXT("Unexpected token following UDELEGATE(): %s"), Token.Identifier);
+			FUHTException::Throwf(*this, TEXT("Unexpected token following UDELEGATE(): %s"), Token.Identifier);
 		}
 
 		DelegateMacro = Token.Identifier;
@@ -6162,13 +6126,13 @@ FUnrealFunctionDefinitionInfo& FHeaderParser::CompileDelegateDeclaration(const T
 
 	if (DelegateMacro != ExpectedOriginalString)
 	{
-		FError::Throwf(TEXT("Unable to parse delegate declaration; expected '%s' but found '%s'."), *ExpectedOriginalString, *DelegateMacro);
+		FUHTException::Throwf(*this, TEXT("Unable to parse delegate declaration; expected '%s' but found '%s'."), *ExpectedOriginalString, *DelegateMacro);
 	}
 
 	// Multi-cast delegate function signatures are not allowed to have a return value
 	if (bHasReturnValue && bIsMulticast)
 	{
-		UE_LOG_ERROR_UHT(TEXT("Multi-cast delegates function signatures must not return a value"));
+		UE_LOG_ERROR_UHT(*this, TEXT("Multi-cast delegates function signatures must not return a value"));
 	}
 
 	// Delegate signature
@@ -6199,7 +6163,7 @@ FUnrealFunctionDefinitionInfo& FHeaderParser::CompileDelegateDeclaration(const T
 	// Get the delegate name
 	if (!GetIdentifier(FuncInfo.Function))
 	{
-		FError::Throwf(TEXT("Missing name for %s"), CurrentScopeName );
+		FUHTException::Throwf(*this, TEXT("Missing name for %s"), CurrentScopeName );
 	}
 
 	// If this is a delegate function then go ahead and mangle the name so we don't collide with
@@ -6212,7 +6176,7 @@ FUnrealFunctionDefinitionInfo& FHeaderParser::CompileDelegateDeclaration(const T
 
 		if (!Name.StartsWith(TEXT("F"), ESearchCase::CaseSensitive))
 		{
-			FError::Throwf(TEXT("Delegate type declarations must start with F"));
+			FUHTException::Throwf(*this, TEXT("Delegate type declarations must start with F"));
 		}
 
 		Name.RightChopInline(1, false);
@@ -6241,14 +6205,14 @@ FUnrealFunctionDefinitionInfo& FHeaderParser::CompileDelegateDeclaration(const T
 
 		if (!GetIdentifier(OwningClass))
 		{
-			FError::Throwf(TEXT("Missing OwningClass specifier."));
+			FUHTException::Throwf(*this, TEXT("Missing OwningClass specifier."));
 		}
 		RequireSymbol(TEXT(','), TEXT("Delegate Declaration"));
 		
 		FToken DelegateName;
 		if (!GetIdentifier(DelegateName))
 		{
-			FError::Throwf(TEXT("Missing Delegate Name."));
+			FUHTException::Throwf(*this, TEXT("Missing Delegate Name."));
 		}
 
 		DelegateSignatureFunctionDef.SetSparseOwningClassName(*GetClassNameWithoutPrefix(FString(OwningClass.Identifier)));
@@ -6268,7 +6232,7 @@ FUnrealFunctionDefinitionInfo& FHeaderParser::CompileDelegateDeclaration(const T
 		int32 ParamCount = UE_PTRDIFF_TO_INT32(FoundParamCount - UHTConfig.DelegateParameterCountStrings.GetData()) + 1;
 		if (DelegateSignatureFunctionDef.GetProperties().Num() != ParamCount)
 		{
-			FError::Throwf(TEXT("Expected %d parameters but found %d parameters"), ParamCount, DelegateSignatureFunctionDef.GetProperties().Num());
+			FUHTException::Throwf(*this, TEXT("Expected %d parameters but found %d parameters"), ParamCount, DelegateSignatureFunctionDef.GetProperties().Num());
 		}
 	}
 	else
@@ -6307,7 +6271,7 @@ FUnrealFunctionDefinitionInfo& FHeaderParser::CompileDelegateDeclaration(const T
 		FUnrealFunctionDefinitionInfo* TestFuncDef = *FunctionIterator;
 		if (TestFuncDef != &DelegateSignatureFunctionDef && TestFuncDef->GetFName() == DelegateSignatureFunctionDef.GetFName())
 		{
-			FError::Throwf(TEXT("Can't override delegate signature function '%s'"), FuncDefFuncInfo.Function.Identifier);
+			FUHTException::Throwf(*this, TEXT("Can't override delegate signature function '%s'"), FuncDefFuncInfo.Function.Identifier);
 		}
 	}
 
@@ -6330,7 +6294,7 @@ FUnrealFunctionDefinitionInfo& FHeaderParser::CompileFunctionDeclaration()
 
 	if (!GetCurrentClass()->HasAnyClassFlags(CLASS_Native))
 	{
-		FError::Throwf(TEXT("Should only be here for native classes!"));
+		FUHTException::Throwf(*this, TEXT("Should only be here for native classes!"));
 	}
 
 	// Process all specifiers.
@@ -6361,7 +6325,7 @@ FUnrealFunctionDefinitionInfo& FHeaderParser::CompileFunctionDeclaration()
 	}
 	else
 	{
-		FError::Throwf(TEXT("Unknown access level"));
+		FUHTException::Throwf(*this, TEXT("Unknown access level"));
 	}
 
 	// non-static functions in a const class must be const themselves
@@ -6386,7 +6350,7 @@ FUnrealFunctionDefinitionInfo& FHeaderParser::CompileFunctionDeclaration()
 		FuncInfo.FunctionFlags |= FUNC_EditorOnly;
 	}
 
-	ProcessFunctionSpecifiers(FuncInfo, SpecifiersFound, MetaData);
+	ProcessFunctionSpecifiers(*this, FuncInfo, SpecifiersFound, MetaData);
 
 	const bool bClassGeneratedFromBP = GetCurrentClassDef().IsDynamic();
 
@@ -6398,7 +6362,7 @@ FUnrealFunctionDefinitionInfo& FHeaderParser::CompileFunctionDeclaration()
 	if ((FuncInfo.FunctionFlags & FUNC_BlueprintPure) && GetCurrentClass()->HasAnyClassFlags(CLASS_Interface))
 	{
 		// Until pure interface casts are supported, we don't allow pures in interfaces
-		UE_LOG_ERROR_UHT(TEXT("BlueprintPure specifier is not allowed for interface functions"));
+		UE_LOG_ERROR_UHT(*this, TEXT("BlueprintPure specifier is not allowed for interface functions"));
 	}
 
 	if (FuncInfo.FunctionFlags & FUNC_Net)
@@ -6436,7 +6400,7 @@ FUnrealFunctionDefinitionInfo& FHeaderParser::CompileFunctionDeclaration()
 			// To allow for quick iteration, don't enforce the requirement that game functions have to be categorized
 			if (bIsCurrentModulePartOfEngine)
 			{
-				UE_LOG_ERROR_UHT(TEXT("An explicit Category specifier is required for Blueprint accessible functions in an Engine module."));
+				UE_LOG_ERROR_UHT(*this, TEXT("An explicit Category specifier is required for Blueprint accessible functions in an Engine module."));
 			}
 		}
 	}
@@ -6451,7 +6415,7 @@ FUnrealFunctionDefinitionInfo& FHeaderParser::CompileFunctionDeclaration()
 			// Ensure that blueprint events are only allowed in implementable interfaces. Internal only functions allowed
 			if (!bCanImplementInBlueprints)
 			{
-				UE_LOG_ERROR_UHT(TEXT("Interfaces that are not implementable in blueprints cannot have BlueprintImplementableEvent members."));
+				UE_LOG_ERROR_UHT(*this, TEXT("Interfaces that are not implementable in blueprints cannot have BlueprintImplementableEvent members."));
 			}
 		}
 		
@@ -6462,7 +6426,7 @@ FUnrealFunctionDefinitionInfo& FHeaderParser::CompileFunctionDeclaration()
 			// Ensure that if this interface contains blueprint callable functions that are not blueprint defined, that it must be implemented natively
 			if (bCanImplementInBlueprints)
 			{
-				UE_LOG_ERROR_UHT(TEXT("Blueprint implementable interfaces cannot contain BlueprintCallable functions that are not BlueprintImplementableEvents.  Use CannotImplementInterfaceInBlueprint on the interface if you wish to keep this function."));
+				UE_LOG_ERROR_UHT(*this, TEXT("Blueprint implementable interfaces cannot contain BlueprintCallable functions that are not BlueprintImplementableEvents.  Use CannotImplementInterfaceInBlueprint on the interface if you wish to keep this function."));
 			}
 		}
 	}
@@ -6484,7 +6448,7 @@ FUnrealFunctionDefinitionInfo& FHeaderParser::CompileFunctionDeclaration()
 
 					if (GetCurrentClass()->HasAnyClassFlags(CLASS_RequiredAPI))
 					{
-						FError::Throwf(TEXT("'%s' must not be used on methods of a class that is marked '%s' itself."), *RequiredAPIMacroIfPresent, *RequiredAPIMacroIfPresent);
+						FUHTException::Throwf(*this, TEXT("'%s' must not be used on methods of a class that is marked '%s' itself."), *RequiredAPIMacroIfPresent, *RequiredAPIMacroIfPresent);
 					}
 					FuncInfo.FunctionFlags |= FUNC_RequiredAPI;
 					FuncInfo.FunctionExportFlags |= FUNCEXPORT_RequiredAPI;
@@ -6503,7 +6467,7 @@ FUnrealFunctionDefinitionInfo& FHeaderParser::CompileFunctionDeclaration()
 	// Look for static again, in case there was an ENGINE_API token first
 	if (!APIMacroIfPresent.IsEmpty() && MatchIdentifier(TEXT("static"), ESearchCase::CaseSensitive))
 	{
-		FError::Throwf(TEXT("Unexpected API macro '%s'. Did you mean to put '%s' after the static keyword?"), *APIMacroIfPresent, *APIMacroIfPresent);
+		FUHTException::Throwf(*this, TEXT("Unexpected API macro '%s'. Did you mean to put '%s' after the static keyword?"), *APIMacroIfPresent, *APIMacroIfPresent);
 	}
 
 	// Look for virtual again, in case there was an ENGINE_API token first
@@ -6523,18 +6487,18 @@ FUnrealFunctionDefinitionInfo& FHeaderParser::CompileFunctionDeclaration()
 		{
 			if (GetCurrentClass()->HasAnyClassFlags(CLASS_Interface))
 			{
-				FError::Throwf(TEXT("BlueprintImplementableEvents in Interfaces must not be declared 'virtual'"));
+				FUHTException::Throwf(*this, TEXT("BlueprintImplementableEvents in Interfaces must not be declared 'virtual'"));
 			}
 
 			// if this is a BlueprintNativeEvent, make sure it's not "virtual"
 			else if (FuncInfo.FunctionFlags & FUNC_Native)
 			{
-				UE_LOG_ERROR_UHT(TEXT("BlueprintNativeEvent functions must be non-virtual."));
+				UE_LOG_ERROR_UHT(*this, TEXT("BlueprintNativeEvent functions must be non-virtual."));
 			}
 
 			else
 			{
-				UE_LOG_WARNING_UHT(TEXT("BlueprintImplementableEvents should not be virtual. Use BlueprintNativeEvent instead."));
+				UE_LOG_WARNING_UHT(*this, TEXT("BlueprintImplementableEvents should not be virtual. Use BlueprintNativeEvent instead."));
 			}
 		}
 	}
@@ -6543,7 +6507,7 @@ FUnrealFunctionDefinitionInfo& FHeaderParser::CompileFunctionDeclaration()
 		// if this is a function in an Interface, it must be marked 'virtual' unless it's an event
 		if (GetCurrentClass()->HasAnyClassFlags(CLASS_Interface) && !(FuncInfo.FunctionFlags & FUNC_BlueprintEvent))
 		{
-			FError::Throwf(TEXT("Interface functions that are not BlueprintImplementableEvents must be declared 'virtual'"));
+			FUHTException::Throwf(*this, TEXT("Interface functions that are not BlueprintImplementableEvents must be declared 'virtual'"));
 		}
 	}
 
@@ -6556,7 +6520,7 @@ FUnrealFunctionDefinitionInfo& FHeaderParser::CompileFunctionDeclaration()
 
 		if (GetCurrentClass()->HasAnyClassFlags(CLASS_Interface))
 		{
-			UE_LOG_ERROR_UHT(TEXT("Interface functions cannot be declared 'final'"));
+			UE_LOG_ERROR_UHT(*this, TEXT("Interface functions cannot be declared 'final'"));
 		}
 	}
 
@@ -6575,12 +6539,12 @@ FUnrealFunctionDefinitionInfo& FHeaderParser::CompileFunctionDeclaration()
 	// Get function or operator name.
 	if (!GetIdentifier(FuncInfo.Function))
 	{
-		FError::Throwf(TEXT("Missing %s name"), TypeOfFunction);
+		FUHTException::Throwf(*this, TEXT("Missing %s name"), TypeOfFunction);
 	}
 
 	if ( !MatchSymbol(TEXT('(')) )
 	{
-		FError::Throwf(TEXT("Bad %s definition"), TypeOfFunction);
+		FUHTException::Throwf(*this, TEXT("Bad %s definition"), TypeOfFunction);
 	}
 
 	if (FuncInfo.FunctionFlags & FUNC_Net)
@@ -6588,14 +6552,14 @@ FUnrealFunctionDefinitionInfo& FHeaderParser::CompileFunctionDeclaration()
 		bool bIsNetService = !!(FuncInfo.FunctionFlags & (FUNC_NetRequest | FUNC_NetResponse));
 		if (bHasReturnValue && !bIsNetService)
 		{
-			FError::Throwf(TEXT("Replicated functions can't have return values"));
+			FUHTException::Throwf(*this, TEXT("Replicated functions can't have return values"));
 		}
 
 		if (FuncInfo.RPCId > 0)
 		{
 			if (FString* ExistingFunc = UsedRPCIds.Find(FuncInfo.RPCId))
 			{
-				FError::Throwf(TEXT("Function %s already uses identifier %d"), **ExistingFunc, FuncInfo.RPCId);
+				FUHTException::Throwf(*this, TEXT("Function %s already uses identifier %d"), **ExistingFunc, FuncInfo.RPCId);
 			}
 
 			UsedRPCIds.Add(FuncInfo.RPCId, FuncInfo.Function.Identifier);
@@ -6661,7 +6625,7 @@ FUnrealFunctionDefinitionInfo& FHeaderParser::CompileFunctionDeclaration()
 		if (UFunction* OverriddenFunction = ::FindUField<UFunction>(SuperStruct, FuncDefFuncInfo.Function.Identifier))
 		{
 			// Native function overrides should be done in CPP text, not in a UFUNCTION() declaration (you can't change flags, and it'd otherwise be a burden to keep them identical)
-			UE_LOG_ERROR_UHT(TEXT("%s: Override of UFUNCTION in parent class (%s) cannot have a UFUNCTION() declaration above it; it will use the same parameters as the original declaration."), FuncDefFuncInfo.Function.Identifier, *OverriddenFunction->GetOuter()->GetName());
+			UE_LOG_ERROR_UHT(*this, TEXT("%s: Override of UFUNCTION in parent class (%s) cannot have a UFUNCTION() declaration above it; it will use the same parameters as the original declaration."), FuncDefFuncInfo.Function.Identifier, *OverriddenFunction->GetOuter()->GetName());
 		}
 	}
 
@@ -6671,7 +6635,7 @@ FUnrealFunctionDefinitionInfo& FHeaderParser::CompileFunctionDeclaration()
 		// we don't want to prevent compilation:
 		if (!bClassGeneratedFromBP)
 		{
-			UE_LOG_ERROR_UHT(TEXT("BlueprintPure specifier is not allowed for functions with no return value and no output parameters."));
+			UE_LOG_ERROR_UHT(*this, TEXT("BlueprintPure specifier is not allowed for functions with no return value and no output parameters."));
 		}
 	}
 
@@ -6682,7 +6646,7 @@ FUnrealFunctionDefinitionInfo& FHeaderParser::CompileFunctionDeclaration()
 		if (FuncDef.HasAnyFunctionFlags(FUNC_Native))
 		{
 			// @TODO: UCREMOVAL Reconsider?
-			//FError::Throwf(TEXT("'const' may only be used for native functions"));
+			//FUHTException::Throwf(*this, TEXT("'const' may only be used for native functions"));
 		}
 
 		FuncDef.SetFunctionFlags(FUNC_Const);
@@ -6723,7 +6687,7 @@ FUnrealFunctionDefinitionInfo& FHeaderParser::CompileFunctionDeclaration()
 
 		if (!bGotZero)
 		{
-			FError::Throwf(TEXT("Expected 0 to indicate function is abstract"));
+			FUHTException::Throwf(*this, TEXT("Expected 0 to indicate function is abstract"));
 		}
 	}
 
@@ -6735,11 +6699,11 @@ FUnrealFunctionDefinitionInfo& FHeaderParser::CompileFunctionDeclaration()
 		FuncDef.SetFunctionExportFlags(FUNCEXPORT_Final);
 		if (GetCurrentClass()->HasAnyClassFlags(CLASS_Interface))
 		{
-			FError::Throwf(TEXT("Interface functions cannot be declared 'final'"));
+			FUHTException::Throwf(*this, TEXT("Interface functions cannot be declared 'final'"));
 		}
 		else if (FuncDef.HasAnyFunctionFlags(FUNC_BlueprintEvent))
 		{
-			FError::Throwf(TEXT("Blueprint events cannot be declared 'final'"));
+			FUHTException::Throwf(*this, TEXT("Blueprint events cannot be declared 'final'"));
 		}
 	}
 
@@ -6748,7 +6712,7 @@ FUnrealFunctionDefinitionInfo& FHeaderParser::CompileFunctionDeclaration()
 	{
 		if ((FuncDef.GetFunctionFlags() & FUNC_NetFuncFlags) != (SuperFuncDef->GetFunctionFlags() & FUNC_NetFuncFlags))
 		{
-			FError::Throwf(TEXT("Overridden function '%s': Cannot specify different replication flags when overriding a function."), *FuncDef.GetName());
+			FUHTException::Throwf(*this, TEXT("Overridden function '%s': Cannot specify different replication flags when overriding a function."), *FuncDef.GetName());
 		}
 	}
 
@@ -6756,7 +6720,7 @@ FUnrealFunctionDefinitionInfo& FHeaderParser::CompileFunctionDeclaration()
 	// this is required because the networking code only checks the class for RPCs when initializing network data, not any states within it
 	if (FuncDef.HasAnyFunctionFlags(FUNC_Net) && SuperFuncDef == nullptr && UHTCast<FUnrealClassDefinitionInfo>(FuncDef.GetOuter()) == nullptr)
 	{
-		FError::Throwf(TEXT("Function '%s': Base implementation of RPCs cannot be in a state. Add a stub outside state scope."), *FuncDef.GetName());
+		FUHTException::Throwf(*this, TEXT("Function '%s': Base implementation of RPCs cannot be in a state. Add a stub outside state scope."), *FuncDef.GetName());
 	}
 
 	if (FuncDef.HasAnyFunctionFlags(FUNC_BlueprintCallable | FUNC_BlueprintEvent))
@@ -6765,14 +6729,14 @@ FUnrealFunctionDefinitionInfo& FHeaderParser::CompileFunctionDeclaration()
 		{
 			if (PropertyDef->IsStaticArray())
 			{
-				FError::Throwf(TEXT("Static array cannot be exposed to blueprint. Function: %s Parameter %s\n"), *FuncDef.GetName(), *PropertyDef->GetName());
+				FUHTException::Throwf(*this, TEXT("Static array cannot be exposed to blueprint. Function: %s Parameter %s\n"), *FuncDef.GetName(), *PropertyDef->GetName());
 			}
 
 			if (!FPropertyTraits::IsSupportedByBlueprint(*PropertyDef, false))
 			{
 				FString ExtendedCPPType;
 				FString CPPType = PropertyDef->GetCPPType(&ExtendedCPPType);
-				UE_LOG_ERROR_UHT(TEXT("Type '%s%s' is not supported by blueprint. %s.%s"), *CPPType, *ExtendedCPPType, *FuncDef.GetName(), *PropertyDef->GetName());
+				UE_LOG_ERROR_UHT(*this, TEXT("Type '%s%s' is not supported by blueprint. %s.%s"), *CPPType, *ExtendedCPPType, *FuncDef.GetName(), *PropertyDef->GetName());
 			}
 		}
 	}
@@ -6784,7 +6748,7 @@ FUnrealFunctionDefinitionInfo& FHeaderParser::CompileFunctionDeclaration()
 	FToken Token;
 	if (!GetToken(Token))
 	{
-		FError::Throwf(TEXT("Unexpected end of file"));
+		FUHTException::Throwf(*this, TEXT("Unexpected end of file"));
 	}
 
 	// Optionally consume a semicolon
@@ -6824,7 +6788,7 @@ void FHeaderParser::ParseFieldMetaData(TMap<FName, FString>& MetaData, const TCH
 		RequireSymbol( TEXT('('), ErrorMessageGetter );
 		if (!GetRawTokenRespectingQuotes(PropertyMetaData, TCHAR(')')))
 		{
-			FError::Throwf(TEXT("'%s': No metadata specified"), FieldName);
+			FUHTException::Throwf(*this, TEXT("'%s': No metadata specified"), FieldName);
 		}
 		RequireSymbol( TEXT(')'), ErrorMessageGetter);
 	}
@@ -6937,7 +6901,7 @@ void FHeaderParser::ValidatePropertyIsDeprecatedIfNecessary(const FPropertyBase&
 	// check to see if we have a FClassProperty using a deprecated class
 	if (VarProperty.MetaClassDef != nullptr && VarProperty.MetaClassDef->HasAnyClassFlags(CLASS_Deprecated))
 	{
-		UE_LOG_ERROR_UHT(TEXT("Property is using a deprecated class: %s.  Property should be marked deprecated as well."), *VarProperty.MetaClassDef->GetPathName());
+		UE_LOG_ERROR_UHT(*this, TEXT("Property is using a deprecated class: %s.  Property should be marked deprecated as well."), *VarProperty.MetaClassDef->GetPathName());
 	}
 
 	// check to see if we have a FObjectProperty using a deprecated class.
@@ -6947,7 +6911,7 @@ void FHeaderParser::ValidatePropertyIsDeprecatedIfNecessary(const FPropertyBase&
 	{
 		if (VarProperty.ClassDef->HasAnyClassFlags(CLASS_Deprecated))	// and the object class being used has been deprecated
 		{
-			UE_LOG_ERROR_UHT(TEXT("Property is using a deprecated class: %s.  Property should be marked deprecated as well."), *VarProperty.ClassDef->GetPathName());
+			UE_LOG_ERROR_UHT(*this, TEXT("Property is using a deprecated class: %s.  Property should be marked deprecated as well."), *VarProperty.ClassDef->GetPathName());
 		}
 	}
 }
@@ -6976,7 +6940,7 @@ bool FHeaderParser::ValidateScriptStructOkForNet(const FString& OriginStructName
 			if (!PropertyDef->HasAnyPropertyFlags(CPF_RepSkip))
 			{
 				bIsStructValid = false;
-				UE_LOG_ERROR_UHT(TEXT("Sets are not supported for Replication or RPCs.  Set %s in %s.  Origin %s"), *PropertyDef->GetName(), *PropertyDef->GetOuter()->GetName(), *OriginStructName);
+				UE_LOG_ERROR_UHT(*this, TEXT("Sets are not supported for Replication or RPCs.  Set %s in %s.  Origin %s"), *PropertyDef->GetName(), *PropertyDef->GetOuter()->GetName(), *OriginStructName);
 			}
 		}
 		else if (PropertyDef->IsMap())
@@ -6984,7 +6948,7 @@ bool FHeaderParser::ValidateScriptStructOkForNet(const FString& OriginStructName
 			if (!PropertyDef->HasAnyPropertyFlags(CPF_RepSkip))
 			{
 				bIsStructValid = false;
-				UE_LOG_ERROR_UHT(TEXT("Maps are not supported for Replication or RPCs.  Map %s in %s.  Origin %s"), *PropertyDef->GetName(), *PropertyDef->GetOuter()->GetName(), *OriginStructName);
+				UE_LOG_ERROR_UHT(*this, TEXT("Maps are not supported for Replication or RPCs.  Map %s in %s.  Origin %s"), *PropertyDef->GetName(), *PropertyDef->GetOuter()->GetName(), *OriginStructName);
 			}
 		}
 		else if (PropertyDef->IsStructOrStructStaticArray())
@@ -7062,7 +7026,7 @@ void FHeaderParser::CompileVariableDeclaration(FUnrealStructDefinitionInfo& Stru
 		}
 		else
 		{
-			UE_LOG_ERROR_UHT(TEXT("An explicit Category specifier is required for any property exposed to the editor or Blueprints in an Engine module."));
+			UE_LOG_ERROR_UHT(*this, TEXT("An explicit Category specifier is required for any property exposed to the editor or Blueprints in an Engine module."));
 		}
 	}
 
@@ -7071,13 +7035,13 @@ void FHeaderParser::CompileVariableDeclaration(FUnrealStructDefinitionInfo& Stru
 	{
 		// Get the name of the type, removing the asterisk representing the pointer
 		FString TypeName = FString(TypeRange.Count, Input + TypeRange.StartIndex).TrimStartAndEnd().LeftChop(1).TrimEnd();
-		FError::Throwf(TEXT("UPROPERTY pointers cannot be interfaces - did you mean TScriptInterface<%s>?"), *TypeName);
+		FUHTException::Throwf(*this, TEXT("UPROPERTY pointers cannot be interfaces - did you mean TScriptInterface<%s>?"), *TypeName);
 	}
 
 	// If the category was specified explicitly, it wins
 	if (Category && !(OriginalPropertyBase.PropertyFlags & (CPF_Edit|CPF_BlueprintVisible|CPF_BlueprintAssignable|CPF_BlueprintCallable)))
 	{
-		UE_LOG_WARNING_UHT(TEXT("Property has a Category set but is not exposed to the editor or Blueprints with EditAnywhere, BlueprintReadWrite, VisibleAnywhere, BlueprintReadOnly, BlueprintAssignable, BlueprintCallable keywords.\r\n"));
+		UE_LOG_WARNING_UHT(*this, TEXT("Property has a Category set but is not exposed to the editor or Blueprints with EditAnywhere, BlueprintReadWrite, VisibleAnywhere, BlueprintReadOnly, BlueprintAssignable, BlueprintCallable keywords.\r\n"));
 	}
 
 	// Make sure that editblueprint variables are editable
@@ -7085,26 +7049,26 @@ void FHeaderParser::CompileVariableDeclaration(FUnrealStructDefinitionInfo& Stru
 	{
 		if (OriginalPropertyBase.PropertyFlags & CPF_DisableEditOnInstance)
 		{
-			UE_LOG_ERROR_UHT(TEXT("Property cannot have 'DisableEditOnInstance' without being editable"));
+			UE_LOG_ERROR_UHT(*this, TEXT("Property cannot have 'DisableEditOnInstance' without being editable"));
 		}
 
 		if (OriginalPropertyBase.PropertyFlags & CPF_DisableEditOnTemplate)
 		{
-			UE_LOG_ERROR_UHT(TEXT("Property cannot have 'DisableEditOnTemplate' without being editable"));
+			UE_LOG_ERROR_UHT(*this, TEXT("Property cannot have 'DisableEditOnTemplate' without being editable"));
 		}
 	}
 
 	// Validate.
 	if (OriginalPropertyBase.PropertyFlags & CPF_ParmFlags)
 	{
-		FError::Throwf(TEXT("Illegal type modifiers in member variable declaration") );
+		FUHTException::Throwf(*this, TEXT("Illegal type modifiers in member variable declaration") );
 	}
 
 	if (FString* ExposeOnSpawnValue = OriginalPropertyBase.MetaData.Find(NAME_ExposeOnSpawn))
 	{
 		if ((*ExposeOnSpawnValue == TEXT("true")) && !FExposeOnSpawnValidator::IsSupported(OriginalPropertyBase))
 		{
-			UE_LOG_ERROR_UHT(TEXT("ExposeOnSpawn - Property cannot be exposed"));
+			UE_LOG_ERROR_UHT(*this, TEXT("ExposeOnSpawn - Property cannot be exposed"));
 		}
 	}
 
@@ -7134,7 +7098,7 @@ void FHeaderParser::CompileVariableDeclaration(FUnrealStructDefinitionInfo& Stru
 				int32 BitfieldSize = 0;
 				if (!GetConstInt(/*out*/ BitfieldSize) || (BitfieldSize != 1))
 				{
-					FError::Throwf(TEXT("Bad or missing bitfield size for '%s', must be 1."), *NewPropDef.GetName());
+					FUHTException::Throwf(*this, TEXT("Bad or missing bitfield size for '%s', must be 1."), *NewPropDef.GetName());
 				}
 			}
 		}
@@ -7146,7 +7110,7 @@ void FHeaderParser::CompileVariableDeclaration(FUnrealStructDefinitionInfo& Stru
 		{
 			if (NewProperties.Num())
 			{
-				FError::Throwf(TEXT("Comma delimited properties cannot be converted %s.%s\n"), *StructDef.GetName(), *NewPropDef.GetName());
+				FUHTException::Throwf(*this, TEXT("Comma delimited properties cannot be converted %s.%s\n"), *StructDef.GetName(), *NewPropDef.GetName());
 			}
 		}
 
@@ -7167,19 +7131,19 @@ void FHeaderParser::CompileVariableDeclaration(FUnrealStructDefinitionInfo& Stru
 		{
 			if (Struct->IsA<UScriptStruct>() && !StructDef.GetBoolMetaDataHierarchical(FHeaderParserNames::NAME_BlueprintType))
 			{
-				UE_LOG_ERROR_UHT(TEXT("Cannot expose property to blueprints in a struct that is not a BlueprintType. %s.%s"), *StructDef.GetName(), *NewPropDef.GetName());
+				UE_LOG_ERROR_UHT(*this, TEXT("Cannot expose property to blueprints in a struct that is not a BlueprintType. %s.%s"), *StructDef.GetName(), *NewPropDef.GetName());
 			}
 
 			if (NewPropDef.IsStaticArray())
 			{
-				UE_LOG_ERROR_UHT(TEXT("Static array cannot be exposed to blueprint %s.%s"), *StructDef.GetName(), *NewPropDef.GetName());
+				UE_LOG_ERROR_UHT(*this, TEXT("Static array cannot be exposed to blueprint %s.%s"), *StructDef.GetName(), *NewPropDef.GetName());
 			}
 
 			if (!FPropertyTraits::IsSupportedByBlueprint(NewPropDef, true))
 			{
 				FString ExtendedCPPType;
 				FString CPPType = NewPropDef.GetCPPType(&ExtendedCPPType);
-				UE_LOG_ERROR_UHT(TEXT("Type '%s%s' is not supported by blueprint. %s.%s"), *CPPType, *ExtendedCPPType, *StructDef.GetName(), *NewPropDef.GetName());
+				UE_LOG_ERROR_UHT(*this, TEXT("Type '%s%s' is not supported by blueprint. %s.%s"), *CPPType, *ExtendedCPPType, *StructDef.GetName(), *NewPropDef.GetName());
 			}
 		}
 
@@ -7295,7 +7259,7 @@ bool FHeaderParser::CompileStatement(TArray<FUnrealFunctionDefinitionInfo*>& Del
 	}
 	else if (!CompileDeclaration(DelegatesToFixup, Token))
 	{
-		FError::Throwf(TEXT("'%s': Bad command or expression"), Token.Identifier );
+		FUHTException::Throwf(*this, TEXT("'%s': Bad command or expression"), Token.Identifier );
 	}
 	return true;
 }
@@ -7337,11 +7301,11 @@ void FHeaderParser::SkipStatements( int32 NestCount, const TCHAR* ErrorTag  )
 
 	if( NestCount > 0 )
 	{
-		FError::Throwf(TEXT("Unexpected end of file at end of %s"), ErrorTag );
+		FUHTException::Throwf(*this, TEXT("Unexpected end of file at end of %s"), ErrorTag );
 	}
 	else if ( NestCount < 0 )
 	{
-		FError::Throwf(TEXT("Extraneous closing brace found in %s"), ErrorTag);
+		FUHTException::Throwf(*this, TEXT("Extraneous closing brace found in %s"), ErrorTag);
 	}
 }
 
@@ -7353,14 +7317,8 @@ void FHeaderParser::SkipStatements( int32 NestCount, const TCHAR* ErrorTag  )
 // Parses the header associated with the specified class.
 // Returns result enumeration.
 //
-ECompilationResult::Type FHeaderParser::ParseHeader()
+void FHeaderParser::ParseHeader()
 {
-	// Early-out if this class has previously failed some aspect of parsing
-	if (FailedFilesAnnotation.Get(&SourceFile))
-	{
-		return ECompilationResult::OtherCompilationError;
-	}
-
 	// Reset the parser to begin a new class
 	bEncounteredNewStyleClass_UnmatchedBrackets = false;
 	bSpottedAutogeneratedHeaderInclude          = false;
@@ -7369,8 +7327,6 @@ ECompilationResult::Type FHeaderParser::ParseHeader()
 	bClassHasGeneratedUInterfaceBody            = false;
 	bClassHasGeneratedIInterfaceBody            = false;
 
-	ECompilationResult::Type Result = ECompilationResult::OtherCompilationError;
-
 	// Message.
 	UE_LOG(LogCompile, Verbose, TEXT("Parsing %s"), *SourceFile.GetFilename());
 
@@ -7391,118 +7347,86 @@ ECompilationResult::Type FHeaderParser::ParseHeader()
 	// Tells if this header defines no-export classes only.
 	bool bNoExportClassesOnly = true;
 
-#if !PLATFORM_EXCEPTIONS_DISABLED
-	try
-#endif
+	// Parse entire program.
+	TArray<FUnrealFunctionDefinitionInfo*> DelegatesToFixup;
+	while (CompileStatement(DelegatesToFixup))
 	{
-		// Parse entire program.
-		TArray<FUnrealFunctionDefinitionInfo*> DelegatesToFixup;
-		while (CompileStatement(DelegatesToFixup))
-		{
-			bEmptyFile = false;
+		bEmptyFile = false;
 
-			// Clear out the previous comment in anticipation of the next statement.
-			ClearComment();
-			StatementsParsed++;
-		}
+		// Clear out the previous comment in anticipation of the next statement.
+		ClearComment();
+		StatementsParsed++;
+	}
 
-		PopNest(ENestType::GlobalScope, TEXT("Global scope"));
+	PopNest(ENestType::GlobalScope, TEXT("Global scope"));
 
-		// now validate all delegate variables declared in the class
-		{
-			auto ScopeTypeIterator = SourceFile.GetScope()->GetTypeIterator();
-			while (ScopeTypeIterator.MoveNext())
-			{
-				FUnrealFieldDefinitionInfo* TypeDef = *ScopeTypeIterator;
-				if (TypeDef->AsScriptStruct() != nullptr || TypeDef->AsClass() != nullptr)
-				{
-					TMap<FName, FUnrealFunctionDefinitionInfo*> DelegateCache;
-					FixupDelegateProperties(TypeDef->AsStructChecked(), *TypeDef->GetScope(), DelegateCache);
-				}
-			}
-		}
-
-		// Fix up any delegates themselves, if they refer to other delegates
+	// now validate all delegate variables declared in the class
+	{
+		auto ScopeTypeIterator = SourceFile.GetScope()->GetTypeIterator();
+		while (ScopeTypeIterator.MoveNext())
 		{
-			TMap<FName, FUnrealFunctionDefinitionInfo*> DelegateCache;
-			for (FUnrealFunctionDefinitionInfo* DelegateDef : DelegatesToFixup)
+			FUnrealFieldDefinitionInfo* TypeDef = *ScopeTypeIterator;
+			if (TypeDef->AsScriptStruct() != nullptr || TypeDef->AsClass() != nullptr)
 			{
-				FixupDelegateProperties(*DelegateDef, SourceFile.GetScope().Get(), DelegateCache);
+				TMap<FName, FUnrealFunctionDefinitionInfo*> DelegateCache;
+				FixupDelegateProperties(TypeDef->AsStructChecked(), *TypeDef->GetScope(), DelegateCache);
 			}
 		}
+	}
 
-		// Precompute info for runtime optimization.
-		LinesParsed += InputLine;
-
-		if (RPCsNeedingHookup.Num() > 0)
+	// Fix up any delegates themselves, if they refer to other delegates
+	{
+		TMap<FName, FUnrealFunctionDefinitionInfo*> DelegateCache;
+		for (FUnrealFunctionDefinitionInfo* DelegateDef : DelegatesToFixup)
 		{
-			FString ErrorMsg(TEXT("Request functions missing response pairs:\r\n"));
-			for (TMap<int32, FString>::TConstIterator It(RPCsNeedingHookup); It; ++It)
-			{
-				ErrorMsg += FString::Printf(TEXT("%s missing id %d\r\n"), *It.Value(), It.Key());
-			}
-
-			RPCsNeedingHookup.Empty();
-			FError::Throwf(*ErrorMsg);
+			FixupDelegateProperties(*DelegateDef, SourceFile.GetScope().Get(), DelegateCache);
 		}
+	}
 
-		// Make sure the compilation ended with valid nesting.
-		if (bEncounteredNewStyleClass_UnmatchedBrackets)
-		{
-			FError::Throwf(TEXT("Missing } at end of class") );
-		}
+	// Precompute info for runtime optimization.
+	LinesParsed += InputLine;
 
-		if (NestLevel == 1)
-		{
-			FError::Throwf(TEXT("Internal nest inconsistency") );
-		}
-		else if (NestLevel > 2)
+	if (RPCsNeedingHookup.Num() > 0)
+	{
+		FString ErrorMsg(TEXT("Request functions missing response pairs:\r\n"));
+		for (TMap<int32, FString>::TConstIterator It(RPCsNeedingHookup); It; ++It)
 		{
-			FError::Throwf(TEXT("Unexpected end of script in '%s' block"), NestTypeName(TopNest->NestType) );
+			ErrorMsg += FString::Printf(TEXT("%s missing id %d\r\n"), *It.Value(), It.Key());
 		}
 
-		// First-pass success.
-		Result = ECompilationResult::Succeeded;
-
-		for (TSharedRef<FUnrealTypeDefinitionInfo>& TypeDef : SourceFile.GetDefinedClasses())
-		{
-			FUnrealClassDefinitionInfo& ClassDef = UHTCastChecked<FUnrealClassDefinitionInfo>(TypeDef);
-			UClass* Class = ClassDef.GetClass();
-			PostParsingClassSetup(ClassDef);
+		RPCsNeedingHookup.Empty();
+		FUHTException::Throwf(*this, TEXT("%s"), *ErrorMsg);
+	}
 
-			bNoExportClassesOnly = bNoExportClassesOnly && Class->HasAnyClassFlags(CLASS_NoExport);
-		}
+	// Make sure the compilation ended with valid nesting.
+	if (bEncounteredNewStyleClass_UnmatchedBrackets)
+	{
+		FUHTException::Throwf(*this, TEXT("Missing } at end of class") );
+	}
 
-		if (!bSpottedAutogeneratedHeaderInclude && !bEmptyFile && !bNoExportClassesOnly)
-		{
-			const FString& ExpectedHeaderName = SourceFile.GetGeneratedHeaderFilename();
-			FError::Throwf(TEXT("Expected an include at the top of the header: '#include \"%s\"'"), *ExpectedHeaderName);
-		}
+	if (NestLevel == 1)
+	{
+		FUHTException::Throwf(*this, TEXT("Internal nest inconsistency") );
 	}
-#if !PLATFORM_EXCEPTIONS_DISABLED
-	catch( TCHAR* ErrorMsg )
+	else if (NestLevel > 2)
 	{
-		if (NestLevel == 0)
-		{
-			// Pushing nest so there is a file context for this error.
-			PushNest(ENestType::GlobalScope, nullptr, &SourceFile);
-		}
-
-		// Handle compiler error.
-		{
-			TGuardValue<ELogTimes::Type> DisableLogTimes(GPrintLogTimes, ELogTimes::None);
-			FString FormattedErrorMessageWithContext = FString::Printf(TEXT("%s: Error: %s"), *GetContext(), ErrorMsg);
+		FUHTException::Throwf(*this, TEXT("Unexpected end of script in '%s' block"), NestTypeName(TopNest->NestType) );
+	}
 
-			UE_LOG(LogCompile, Log,  TEXT("%s"), *FormattedErrorMessageWithContext );
-			Warn->Log(ELogVerbosity::Error, *FString::Printf(TEXT("Error: %s"), ErrorMsg));
-		}
+	for (TSharedRef<FUnrealTypeDefinitionInfo>& TypeDef : SourceFile.GetDefinedClasses())
+	{
+		FUnrealClassDefinitionInfo& ClassDef = UHTCastChecked<FUnrealClassDefinitionInfo>(TypeDef);
+		UClass* Class = ClassDef.GetClass();
+		PostParsingClassSetup(ClassDef);
 
-		FailedFilesAnnotation.Set(&SourceFile);
-		Result = ECompilationResult::OtherCompilationError;
+		bNoExportClassesOnly = bNoExportClassesOnly && Class->HasAnyClassFlags(CLASS_NoExport);
 	}
-#endif
 
-	return Result; //@TODO: UCREMOVAL: This function is always returning succeeded even on a compiler error; should this continue?
+	if (!bSpottedAutogeneratedHeaderInclude && !bEmptyFile && !bNoExportClassesOnly)
+	{
+		const FString& ExpectedHeaderName = SourceFile.GetGeneratedHeaderFilename();
+		FUHTException::Throwf(*this, TEXT("Expected an include at the top of the header: '#include \"%s\"'"), *ExpectedHeaderName);
+	}
 }
 
 /*-----------------------------------------------------------------------------
@@ -7512,12 +7436,10 @@ ECompilationResult::Type FHeaderParser::ParseHeader()
 // Parse Class's annotated headers and optionally its child classes.
 static const FString ObjectHeader(TEXT("NoExportTypes.h"));
 
-FHeaderParser::FHeaderParser(FUnrealPackageDefinitionInfo& InPackageDef, FFeedbackContext* InWarn, FUnrealSourceFile& InSourceFile)
-	: FBaseParser()
-	, Warn(InWarn)
+FHeaderParser::FHeaderParser(FUnrealPackageDefinitionInfo& InPackageDef, FUnrealSourceFile& InSourceFile)
+	: FBaseParser(InSourceFile)
 	, Filename(IFileManager::Get().ConvertToAbsolutePathForExternalAppForRead(*InSourceFile.GetFilename()))
 	, PackageDef(InPackageDef)
-	, SourceFile(InSourceFile)
 {
 
 	const FManifestModule& Module = PackageDef.GetModule();
@@ -7553,28 +7475,27 @@ FHeaderParser::FHeaderParser(FUnrealPackageDefinitionInfo& InPackageDef, FFeedba
 }
 
 // Throws if a specifier value wasn't provided
-void FHeaderParser::RequireSpecifierValue(const FPropertySpecifier& Specifier, bool bRequireExactlyOne)
+void FHeaderParser::RequireSpecifierValue(const FUHTExceptionContext& Context, const FPropertySpecifier& Specifier, bool bRequireExactlyOne)
 {
 	if (Specifier.Values.Num() == 0)
 	{
-		FError::Throwf(TEXT("The specifier '%s' must be given a value"), *Specifier.Key);
+		FUHTException::Throwf(Context, TEXT("The specifier '%s' must be given a value"), *Specifier.Key);
 	}
 	else if ((Specifier.Values.Num() != 1) && bRequireExactlyOne)
 	{
-		FError::Throwf(TEXT("The specifier '%s' must be given exactly one value"), *Specifier.Key);
+		FUHTException::Throwf(Context, TEXT("The specifier '%s' must be given exactly one value"), *Specifier.Key);
 	}
 }
 
 // Throws if a specifier value wasn't provided
-FString FHeaderParser::RequireExactlyOneSpecifierValue(const FPropertySpecifier& Specifier)
+FString FHeaderParser::RequireExactlyOneSpecifierValue(const FUHTExceptionContext& Context, const FPropertySpecifier& Specifier)
 {
-	RequireSpecifierValue(Specifier, /*bRequireExactlyOne*/ true);
+	RequireSpecifierValue(Context, Specifier, /*bRequireExactlyOne*/ true);
 	return Specifier.Values[0];
 }
 
-ECompilationResult::Type FHeaderParser::Parse(
+void FHeaderParser::Parse(
 	FUnrealPackageDefinitionInfo& PackageDef,
-	FFeedbackContext* Warn,
 	FUnrealSourceFile& SourceFile)
 {
 	SCOPE_SECONDS_COUNTER_UHT(Parse);
@@ -7585,77 +7506,55 @@ ECompilationResult::Type FHeaderParser::Parse(
 	TGuardValue<bool> AutoRestoreVerifyObjectRefsFlag(GVerifyObjectReferencesOnly, true);
 	// Create the header parser and register it as the warning context.
 	// Note: This must be declared outside the try block, since the catch block will log into it.
-	FHeaderParser HeaderParser(PackageDef, Warn, SourceFile);
-	Warn->SetContext(&HeaderParser);
+	FHeaderParser HeaderParser(PackageDef, SourceFile);
 
 	// Hierarchically parse all classes.
-	ECompilationResult::Type Result = ECompilationResult::Succeeded;
-#if !PLATFORM_EXCEPTIONS_DISABLED
-	try
-#endif
-	{
-		TArray<FUnrealSourceFile*> SourceFilesRequired;
-		for (FHeaderProvider& Include : SourceFile.GetIncludes())
-		{
-			if (Include.GetId() == ObjectHeader)
-			{
-				continue;
-			}
-
-			if (FUnrealSourceFile* DepFile = Include.Resolve(SourceFile))
-			{
-				SourceFilesRequired.Add(DepFile);
-			}
-		}
 
-		for (const TSharedRef<FUnrealTypeDefinitionInfo>& ClassDataPair : SourceFile.GetDefinedClasses())
+	TArray<FUnrealSourceFile*> SourceFilesRequired;
+	for (FHeaderProvider& Include : SourceFile.GetIncludes())
+	{
+		if (Include.GetId() == ObjectHeader)
 		{
-			FUnrealClassDefinitionInfo& ClassDef = ClassDataPair->AsClassChecked();
-			for (FUnrealClassDefinitionInfo* ParentClassDef = ClassDef.GetSuperClass(); ParentClassDef; ParentClassDef = ParentClassDef->GetSuperClass())
-			{
-				if (ParentClassDef->HasAnyClassFlags(CLASS_Parsed | CLASS_Intrinsic))
-				{
-					break;
-				}
-				SourceFilesRequired.Add(&ParentClassDef->GetUnrealSourceFile());
-			}
+			continue;
 		}
 
-		for (FUnrealSourceFile* RequiredFile : SourceFilesRequired)
+		if (FUnrealSourceFile* DepFile = Include.Resolve(SourceFile))
 		{
-			SourceFile.GetScope()->IncludeScope(&RequiredFile->GetScope().Get());
+			SourceFilesRequired.Add(DepFile);
 		}
+	}
 
-		// Parse the file
+	for (const TSharedRef<FUnrealTypeDefinitionInfo>& ClassDataPair : SourceFile.GetDefinedClasses())
+	{
+		FUnrealClassDefinitionInfo& ClassDef = ClassDataPair->AsClassChecked();
+		for (FUnrealClassDefinitionInfo* ParentClassDef = ClassDef.GetSuperClass(); ParentClassDef; ParentClassDef = ParentClassDef->GetSuperClass())
 		{
-			Result = HeaderParser.ParseHeader();
-
-			for (const TSharedRef<FUnrealTypeDefinitionInfo>& ClassDef : SourceFile.GetDefinedClasses())
+			if (ParentClassDef->HasAnyClassFlags(CLASS_Parsed | CLASS_Intrinsic))
 			{
-				ClassDef->AsClassChecked().SetClassFlags(CLASS_Parsed);
+				break;
 			}
+			SourceFilesRequired.Add(&ParentClassDef->GetUnrealSourceFile());
 		}
 	}
-#if !PLATFORM_EXCEPTIONS_DISABLED
-	catch (const FUHTException& Ex)
+
+	for (FUnrealSourceFile* RequiredFile : SourceFilesRequired)
 	{
-		FResults::LogError(Ex);
-		Result = ECompilationResult::OtherCompilationError;
+		SourceFile.GetScope()->IncludeScope(&RequiredFile->GetScope().Get());
 	}
-	catch (TCHAR* ErrorMsg)
+
+	// Parse the file
 	{
-		Warn->Log(ELogVerbosity::Error, ErrorMsg);
-		Result = ECompilationResult::OtherCompilationError;
+		HeaderParser.ParseHeader();
+
+		for (const TSharedRef<FUnrealTypeDefinitionInfo>& ClassDef : SourceFile.GetDefinedClasses())
+		{
+			ClassDef->AsClassChecked().SetClassFlags(CLASS_Parsed);
+		}
 	}
-#endif
 
 	// Remember stats about this file
 	SourceFile.SetLinesParsed(HeaderParser.LinesParsed);
 	SourceFile.SetStatementsParsed(HeaderParser.StatementsParsed);
-
-	// Unregister the header parser from the feedback context
-	Warn->SetContext(NULL);
-	return Result;
 }
 
 /** 
@@ -7846,7 +7745,7 @@ const TCHAR* GetBlockDirectiveTypeString(EBlockDirectiveType DirectiveType)
 // Performs a preliminary parse of the text in the specified buffer, pulling out useful information for the header generation process
 void FHeaderParser::SimplifiedClassParse(FUnrealSourceFile& SourceFile, const TCHAR* InBuffer, FStringOutputDevice& ClassHeaderTextStrippedOfCppText)
 {
-	FHeaderPreParser Parser;
+	FHeaderPreParser Parser(SourceFile);
 	FString StrLine;
 	FString ClassName;
 	FString BaseClassName;
@@ -8210,7 +8109,7 @@ void FHeaderParser::SimplifiedClassParse(FUnrealSourceFile& SourceFile, const TC
 						FUHTException::Throwf(SourceFile, CurrentLine, TEXT("Missing open parenthesis after UINTERFACE"));
 					}
 
-					TSharedRef<FUnrealTypeDefinitionInfo> ClassDecl = Parser.ParseClassDeclaration(SourceFile, StartOfLine + (UInterfaceMacroDecl - Str), CurrentLine, true, TEXT("UINTERFACE"));
+					TSharedRef<FUnrealTypeDefinitionInfo> ClassDecl = Parser.ParseClassDeclaration(StartOfLine + (UInterfaceMacroDecl - Str), CurrentLine, true, TEXT("UINTERFACE"));
 					bFoundExportedClasses |= !EnumHasAnyFlags(ClassDecl->AsClassChecked().GetParsedClassFlags(), CLASS_NoExport);
 					SourceFile.AddDefinedClass(MoveTemp(ClassDecl));
 				}
@@ -8225,7 +8124,7 @@ void FHeaderParser::SimplifiedClassParse(FUnrealSourceFile& SourceFile, const TC
 						FUHTException::Throwf(SourceFile, CurrentLine, TEXT("Missing open parenthesis after UCLASS"));
 					}
 
-					TSharedRef<FUnrealTypeDefinitionInfo> ClassDecl = Parser.ParseClassDeclaration(SourceFile, StartOfLine + (UClassMacroDecl - Str), CurrentLine, false, TEXT("UCLASS"));
+					TSharedRef<FUnrealTypeDefinitionInfo> ClassDecl = Parser.ParseClassDeclaration(StartOfLine + (UClassMacroDecl - Str), CurrentLine, false, TEXT("UCLASS"));
 					bFoundExportedClasses |= !EnumHasAnyFlags(ClassDecl->AsClassChecked().GetParsedClassFlags(), CLASS_NoExport);
 					SourceFile.AddDefinedClass(MoveTemp(ClassDecl));
 				}
@@ -8240,7 +8139,7 @@ void FHeaderParser::SimplifiedClassParse(FUnrealSourceFile& SourceFile, const TC
 						FUHTException::Throwf(SourceFile, CurrentLine, TEXT("Missing open parenthesis after UENUM"));
 					}
 
-					TSharedRef<FUnrealTypeDefinitionInfo> EnumDecl = Parser.ParseEnumDeclaration(SourceFile, StartOfLine + (UEnumMacroDecl - Str), CurrentLine);
+					TSharedRef<FUnrealTypeDefinitionInfo> EnumDecl = Parser.ParseEnumDeclaration(StartOfLine + (UEnumMacroDecl - Str), CurrentLine);
 					SourceFile.AddDefinedEnum(MoveTemp(EnumDecl));
 				}
 			}
@@ -8254,7 +8153,7 @@ void FHeaderParser::SimplifiedClassParse(FUnrealSourceFile& SourceFile, const TC
 						FUHTException::Throwf(SourceFile, CurrentLine, TEXT("Missing open parenthesis after USTRUCT"));
 					}
 
-					TSharedRef<FUnrealTypeDefinitionInfo> StructDecl = Parser.ParseStructDeclaration(SourceFile, StartOfLine + (UStructMacroDecl - Str), CurrentLine);
+					TSharedRef<FUnrealTypeDefinitionInfo> StructDecl = Parser.ParseStructDeclaration(StartOfLine + (UStructMacroDecl - Str), CurrentLine);
 					SourceFile.AddDefinedStruct(MoveTemp(StructDecl));
 				}
 			}
@@ -8263,14 +8162,14 @@ void FHeaderParser::SimplifiedClassParse(FUnrealSourceFile& SourceFile, const TC
 
 	if (bFoundExportedClasses && !bFoundGeneratedInclude)
 	{
-		FError::Throwf(TEXT("No #include found for the .generated.h file - the .generated.h file should always be the last #include in a header"));
+		FUHTException::Throwf(Parser, TEXT("No #include found for the .generated.h file - the .generated.h file should always be the last #include in a header"));
 	}
 }
 
 /////////////////////////////////////////////////////
 // FHeaderPreParser
 
-TSharedRef<FUnrealTypeDefinitionInfo> FHeaderPreParser::ParseClassDeclaration(FUnrealSourceFile& SourceFile, const TCHAR* InputText, int32 InLineNumber, bool bClassIsAnInterface, const TCHAR* StartingMatchID)
+TSharedRef<FUnrealTypeDefinitionInfo> FHeaderPreParser::ParseClassDeclaration(const TCHAR* InputText, int32 InLineNumber, bool bClassIsAnInterface, const TCHAR* StartingMatchID)
 {
 	const TCHAR* ErrorMsg = TEXT("Class declaration");
 
@@ -8304,7 +8203,7 @@ TSharedRef<FUnrealTypeDefinitionInfo> FHeaderPreParser::ParseClassDeclaration(FU
 	TSharedRef<FUnrealClassDefinitionInfo> ClassDef = MakeShareable(new FUnrealClassDefinitionInfo(SourceFile, InLineNumber, MoveTemp(ClassName), FName(ClassNameWithoutPrefixStr, FNAME_Add), bClassIsAnInterface));
 
 	// Parse the inheritance list
-	ParseInheritance(TEXT("class"), [this, &SourceFile, &ClassNameWithoutPrefixStr, &ClassDef = *ClassDef](const TCHAR* ClassName, bool bIsSuperClass)
+	ParseInheritance(TEXT("class"), [this, &ClassNameWithoutPrefixStr, &ClassDef = *ClassDef](const TCHAR* ClassName, bool bIsSuperClass)
 	{
 		FString ClassNameStr(ClassName);
 		SourceFile.AddClassIncludeIfNeeded(InputLine, ClassNameWithoutPrefixStr, ClassNameStr);
@@ -8327,7 +8226,7 @@ TSharedRef<FUnrealTypeDefinitionInfo> FHeaderPreParser::ParseClassDeclaration(FU
 //
 // Compile an enumeration definition.
 //
-TSharedRef<FUnrealTypeDefinitionInfo> FHeaderPreParser::ParseEnumDeclaration(FUnrealSourceFile& SourceFile, const TCHAR* InputText, int32 InLineNumber)
+TSharedRef<FUnrealTypeDefinitionInfo> FHeaderPreParser::ParseEnumDeclaration(const TCHAR* InputText, int32 InLineNumber)
 {
 	const TCHAR* ErrorMsg = TEXT("Enum declaration");
 
@@ -8346,7 +8245,7 @@ TSharedRef<FUnrealTypeDefinitionInfo> FHeaderPreParser::ParseEnumDeclaration(FUn
 	bool bReadEnumName = false;
 	if (!GetIdentifier(EnumToken))
 	{
-		FError::Throwf(TEXT("Missing identifier after UENUM()"));
+		FUHTException::Throwf(*this, TEXT("Missing identifier after UENUM()"));
 	}
 
 	if (EnumToken.Matches(TEXT("namespace"), ESearchCase::CaseSensitive))
@@ -8359,7 +8258,7 @@ TSharedRef<FUnrealTypeDefinitionInfo> FHeaderPreParser::ParseEnumDeclaration(FUn
 	{
 		if (!GetIdentifier(EnumToken))
 		{
-			FError::Throwf(TEXT("Missing identifier after enum"));
+			FUHTException::Throwf(*this, TEXT("Missing identifier after enum"));
 		}
 
 		if (EnumToken.Matches(TEXT("class"), ESearchCase::CaseSensitive) || EnumToken.Matches(TEXT("struct"), ESearchCase::CaseSensitive))
@@ -8377,13 +8276,13 @@ TSharedRef<FUnrealTypeDefinitionInfo> FHeaderPreParser::ParseEnumDeclaration(FUn
 	}
 	else
 	{
-		FError::Throwf(TEXT("UENUM() should be followed by \'enum\' or \'namespace\' keywords."));
+		FUHTException::Throwf(*this, TEXT("UENUM() should be followed by \'enum\' or \'namespace\' keywords."));
 	}
 
 	// Get enumeration name.
 	if (!bReadEnumName)
 	{
-		FError::Throwf(TEXT("Missing enumeration name"));
+		FUHTException::Throwf(*this, TEXT("Missing enumeration name"));
 	}
 
 	TSharedRef<FUnrealEnumDefinitionInfo> EnumDef = MakeShareable(new FUnrealEnumDefinitionInfo(SourceFile, InLineNumber, EnumToken.Identifier, FName(EnumToken.Identifier, FNAME_Add)));
@@ -8393,7 +8292,7 @@ TSharedRef<FUnrealTypeDefinitionInfo> FHeaderPreParser::ParseEnumDeclaration(FUn
 //
 // Compile an structure definition.
 //
-TSharedRef<FUnrealTypeDefinitionInfo> FHeaderPreParser::ParseStructDeclaration(FUnrealSourceFile& SourceFile, const TCHAR* InputText, int32 InLineNumber)
+TSharedRef<FUnrealTypeDefinitionInfo> FHeaderPreParser::ParseStructDeclaration(const TCHAR* InputText, int32 InLineNumber)
 {
 	const TCHAR* ErrorMsg = TEXT("Struct declaration");
 
@@ -8430,7 +8329,7 @@ TSharedRef<FUnrealTypeDefinitionInfo> FHeaderPreParser::ParseStructDeclaration(F
 	TSharedRef<FUnrealScriptStructDefinitionInfo> StructDef = MakeShareable(new FUnrealScriptStructDefinitionInfo(SourceFile, InLineNumber, MoveTemp(StructNameInScript), FName(StructNameStripped, FNAME_Add)));
 
 	// Parse the inheritance list
-	ParseInheritance(TEXT("struct"), [this, &SourceFile, &StructNameStripped, &StructDef = *StructDef](const TCHAR* StructName, bool bIsSuperClass)
+	ParseInheritance(TEXT("struct"), [this, &StructNameStripped, &StructDef = *StructDef](const TCHAR* StructName, bool bIsSuperClass)
 	{
 		FString StructNameStr(StructName);
 		SourceFile.AddScriptStructIncludeIfNeeded(InputLine, StructNameStripped, StructNameStr);
@@ -8689,12 +8588,12 @@ void FHeaderParser::PostPopNestClass(FUnrealClassDefinitionInfo& CurrentClassDef
 
 					if (InterfaceFunctionDef->HasAnyFunctionFlags(FUNC_Event) && !ClassFunctionDef->HasAnyFunctionFlags(FUNC_Event))
 					{
-						FError::Throwf(TEXT("Implementation of function '%s::%s' must be declared as 'event' to match declaration in interface '%s'"), *ClassFunctionDef->GetOuter()->GetName(), *ClassFunctionDef->GetName(), *InterfaceDef->GetName());
+						FUHTException::Throwf(*this, TEXT("Implementation of function '%s::%s' must be declared as 'event' to match declaration in interface '%s'"), *ClassFunctionDef->GetOuter()->GetName(), *ClassFunctionDef->GetName(), *InterfaceDef->GetName());
 					}
 
 					if (InterfaceFunctionDef->HasAnyFunctionFlags(FUNC_Delegate) && !ClassFunctionDef->HasAnyFunctionFlags(FUNC_Delegate))
 					{
-						FError::Throwf(TEXT("Implementation of function '%s::%s' must be declared as 'delegate' to match declaration in interface '%s'"), *ClassFunctionDef->GetOuter()->GetName(), *ClassFunctionDef->GetName(), *InterfaceDef->GetName());
+						FUHTException::Throwf(*this, TEXT("Implementation of function '%s::%s' must be declared as 'delegate' to match declaration in interface '%s'"), *ClassFunctionDef->GetOuter()->GetName(), *ClassFunctionDef->GetName(), *InterfaceDef->GetName());
 					}
 
 					// Making sure all the parameters match up correctly
@@ -8702,7 +8601,7 @@ void FHeaderParser::PostPopNestClass(FUnrealClassDefinitionInfo& CurrentClassDef
 
 					if (ClassFunctionDef->GetProperties().Num() != InterfaceFunctionDef->GetProperties().Num())
 					{
-						FError::Throwf(TEXT("Implementation of function '%s' conflicts with interface '%s' - different number of parameters (%i/%i)"), *InterfaceFunctionDef->GetName(), *InterfaceDef->GetName(), ClassFunctionDef->GetProperties().Num(), InterfaceFunctionDef->GetProperties().Num());
+						FUHTException::Throwf(*this, TEXT("Implementation of function '%s' conflicts with interface '%s' - different number of parameters (%i/%i)"), *InterfaceFunctionDef->GetName(), *InterfaceDef->GetName(), ClassFunctionDef->GetProperties().Num(), InterfaceFunctionDef->GetProperties().Num());
 					}
 
 					for (int32 Index = 0, End = ClassFunctionDef->GetProperties().Num(); Index != End; ++Index)
@@ -8713,11 +8612,11 @@ void FHeaderParser::PostPopNestClass(FUnrealClassDefinitionInfo& CurrentClassDef
 						{
 							if (InterfaceParamDef.HasAnyPropertyFlags(CPF_ReturnParm))
 							{
-								FError::Throwf(TEXT("Implementation of function '%s' conflicts only by return type with interface '%s'"), *InterfaceFunctionDef->GetName(), *InterfaceDef->GetName());
+								FUHTException::Throwf(*this, TEXT("Implementation of function '%s' conflicts only by return type with interface '%s'"), *InterfaceFunctionDef->GetName(), *InterfaceDef->GetName());
 							}
 							else
 							{
-								FError::Throwf(TEXT("Implementation of function '%s' conflicts with interface '%s' - parameter %i '%s'"), *InterfaceFunctionDef->GetName(), *InterfaceDef->GetName(), End, *InterfaceParamDef.GetName());
+								FUHTException::Throwf(*this, TEXT("Implementation of function '%s' conflicts with interface '%s' - parameter %i '%s'"), *InterfaceFunctionDef->GetName(), *InterfaceDef->GetName(), End, *InterfaceParamDef.GetName());
 							}
 						}
 					}
@@ -8735,7 +8634,7 @@ void FHeaderParser::PostPopNestClass(FUnrealClassDefinitionInfo& CurrentClassDef
 					&& !InterfaceFunctionDef->HasAnyFunctionFlags(FUNC_BlueprintEvent)
 					&& !InterfaceDef->HasMetaData(NAME_CannotImplementInterfaceInBlueprint))  // FBlueprintMetadata::MD_CannotImplementInterfaceInBlueprint
 				{
-					FError::Throwf(TEXT("Missing UFunction implementation of function '%s' from interface '%s'.  This function needs a UFUNCTION() declaration."), *InterfaceFunctionDef->GetName(), *InterfaceDef->GetName());
+					FUHTException::Throwf(*this, TEXT("Missing UFunction implementation of function '%s' from interface '%s'.  This function needs a UFUNCTION() declaration."), *InterfaceFunctionDef->GetName(), *InterfaceDef->GetName());
 				}
 			}
 		}
@@ -8762,7 +8661,7 @@ void FHeaderParser::PostPopNestInterface(FUnrealClassDefinitionInfo& CurrentInte
 	}
 }
 
-FDocumentationPolicy FHeaderParser::GetDocumentationPolicyFromName(const FString& PolicyName)
+FDocumentationPolicy FHeaderParser::GetDocumentationPolicyFromName(const FUHTExceptionContext& Context, const FString& PolicyName)
 {
 	FDocumentationPolicy DocumentationPolicy;
 	if (FCString::Strcmp(*PolicyName, TEXT("Strict")) == 0)
@@ -8775,7 +8674,7 @@ FDocumentationPolicy FHeaderParser::GetDocumentationPolicyFromName(const FString
 	}
 	else
 	{
-		FError::Throwf(TEXT("Documentation Policy '%s' not yet supported"), *PolicyName);
+		FUHTException::Throwf(Context, TEXT("Documentation Policy '%s' not yet supported"), *PolicyName);
 	}
 	return DocumentationPolicy;
 }
@@ -8789,7 +8688,7 @@ FDocumentationPolicy FHeaderParser::GetDocumentationPolicyForStruct(FUnrealStruc
 	FString DocumentationPolicyName;
 	if (StructDef.GetStringMetaDataHierarchical(NAME_DocumentationPolicy, &DocumentationPolicyName))
 	{
-		DocumentationPolicy = GetDocumentationPolicyFromName(DocumentationPolicyName);
+		DocumentationPolicy = GetDocumentationPolicyFromName(StructDef, DocumentationPolicyName);
 	}
 	return DocumentationPolicy;
 }
@@ -8806,13 +8705,13 @@ void FHeaderParser::CheckDocumentationPolicyForEnum(FUnrealEnumDefinitionInfo& E
 
 	check(!DocumentationPolicyName->IsEmpty());
 
-	FDocumentationPolicy DocumentationPolicy = GetDocumentationPolicyFromName(*DocumentationPolicyName);
+	FDocumentationPolicy DocumentationPolicy = GetDocumentationPolicyFromName(EnumDef, *DocumentationPolicyName);
 	if (DocumentationPolicy.bClassOrStructCommentRequired)
 	{
 		const FString* EnumToolTip = MetaData.Find(NAME_ToolTip);
 		if (EnumToolTip == nullptr)
 		{
-			UE_LOG_ERROR_UHT(TEXT("Enum '%s' does not provide a tooltip / comment (DocumentationPolicy)."), *EnumDef.GetName());
+			UE_LOG_ERROR_UHT(*this, TEXT("Enum '%s' does not provide a tooltip / comment (DocumentationPolicy)."), *EnumDef.GetName());
 		}
 	}
 
@@ -8828,14 +8727,14 @@ void FHeaderParser::CheckDocumentationPolicyForEnum(FUnrealEnumDefinitionInfo& E
 		const FString* ToolTip = Entry.Find(NAME_ToolTip);
 		if (ToolTip == nullptr)
 		{
-			UE_LOG_ERROR_UHT(TEXT("Enum entry '%s::%s' does not provide a tooltip / comment (DocumentationPolicy)."), *EnumDef.GetName(), *(*EntryName));
+			UE_LOG_ERROR_UHT(*this, TEXT("Enum entry '%s::%s' does not provide a tooltip / comment (DocumentationPolicy)."), *EnumDef.GetName(), *(*EntryName));
 			continue;
 		}
 
 		const FString* ExistingEntry = ToolTipToEntry.Find(*ToolTip);
 		if (ExistingEntry != nullptr)
 		{
-			UE_LOG_ERROR_UHT(TEXT("Enum entries '%s::%s' and '%s::%s' have identical tooltips / comments (DocumentationPolicy)."), *EnumDef.GetName(), *(*ExistingEntry), *EnumDef.GetName(), *(*EntryName));
+			UE_LOG_ERROR_UHT(*this, TEXT("Enum entries '%s::%s' and '%s::%s' have identical tooltips / comments (DocumentationPolicy)."), *EnumDef.GetName(), *(*ExistingEntry), *EnumDef.GetName(), *(*EntryName));
 		}
 		ToolTipToEntry.Add(*ToolTip, *EntryName);
 	}
@@ -8856,7 +8755,7 @@ void FHeaderParser::CheckDocumentationPolicyForStruct(FUnrealStructDefinitionInf
 
 		if (ClassTooltip.IsEmpty() || ClassTooltip.Equals(StructDef.GetName()))
 		{
-			UE_LOG_ERROR_UHT(TEXT("Struct '%s' does not provide a tooltip / comment (DocumentationPolicy)."), *StructDef.GetName());
+			UE_LOG_ERROR_UHT(*this, TEXT("Struct '%s' does not provide a tooltip / comment (DocumentationPolicy)."), *StructDef.GetName());
 		}
 	}
 
@@ -8868,13 +8767,13 @@ void FHeaderParser::CheckDocumentationPolicyForStruct(FUnrealStructDefinitionInf
 			FString ToolTip = PropertyDef->GetToolTipText().ToString();
 			if (ToolTip.IsEmpty() || ToolTip.Equals(PropertyDef->GetDisplayNameText().ToString()))
 			{
-				UE_LOG_ERROR_UHT(TEXT("Property '%s::%s' does not provide a tooltip / comment (DocumentationPolicy)."), *StructDef.GetName(), *PropertyDef->GetName());
+				UE_LOG_ERROR_UHT(*this, TEXT("Property '%s::%s' does not provide a tooltip / comment (DocumentationPolicy)."), *StructDef.GetName(), *PropertyDef->GetName());
 				continue;
 			}
 			const FName* ExistingPropertyName = ToolTipToPropertyName.Find(ToolTip);
 			if (ExistingPropertyName != nullptr)
 			{
-				UE_LOG_ERROR_UHT(TEXT("Property '%s::%s' and '%s::%s' are using identical tooltips (DocumentationPolicy)."), *StructDef.GetName(), *ExistingPropertyName->ToString(), *StructDef.GetName(), *PropertyDef->GetName());
+				UE_LOG_ERROR_UHT(*this, TEXT("Property '%s::%s' and '%s::%s' are using identical tooltips (DocumentationPolicy)."), *StructDef.GetName(), *ExistingPropertyName->ToString(), *StructDef.GetName(), *PropertyDef->GetName());
 			}
 			ToolTipToPropertyName.Add(MoveTemp(ToolTip), PropertyDef->GetFName());
 		}
@@ -8891,7 +8790,7 @@ void FHeaderParser::CheckDocumentationPolicyForStruct(FUnrealStructDefinitionInf
 
 				if(!CheckUIMinMaxRangeFromMetaData(UIMin, UIMax))
 				{
-					UE_LOG_ERROR_UHT(TEXT("Property '%s::%s' does not provide a valid UIMin / UIMax (DocumentationPolicy)."), *StructDef.GetName(), *PropertyDef->GetName());
+					UE_LOG_ERROR_UHT(*this, TEXT("Property '%s::%s' does not provide a valid UIMin / UIMax (DocumentationPolicy)."), *StructDef.GetName(), *PropertyDef->GetName());
 				}
 			}
 		}
@@ -8908,13 +8807,13 @@ void FHeaderParser::CheckDocumentationPolicyForStruct(FUnrealStructDefinitionInf
 				FString ToolTip = FunctionDef->GetToolTipText().ToString();
 				if (ToolTip.IsEmpty())
 				{
-					UE_LOG_ERROR_UHT(TEXT("Function '%s::%s' does not provide a tooltip / comment (DocumentationPolicy)."), *ClassDef->GetName(), *FunctionDef->GetName());
+					UE_LOG_ERROR_UHT(*this, TEXT("Function '%s::%s' does not provide a tooltip / comment (DocumentationPolicy)."), *ClassDef->GetName(), *FunctionDef->GetName());
 					continue;
 				}
 				const FName* ExistingFuncName = ToolTipToFunc.Find(ToolTip);
 				if (ExistingFuncName != nullptr)
 				{
-					UE_LOG_ERROR_UHT(TEXT("Functions '%s::%s' and '%s::%s' uses identical tooltips / comments (DocumentationPolicy)."), *ClassDef->GetName(), *(*ExistingFuncName).ToString(), *ClassDef->GetName(), *FunctionDef->GetName());
+					UE_LOG_ERROR_UHT(*this, TEXT("Functions '%s::%s' and '%s::%s' uses identical tooltips / comments (DocumentationPolicy)."), *ClassDef->GetName(), *(*ExistingFuncName).ToString(), *ClassDef->GetName(), *FunctionDef->GetName());
 				}
 				ToolTipToFunc.Add(MoveTemp(ToolTip), FunctionDef->GetFName());
 			}
@@ -8938,7 +8837,7 @@ void FHeaderParser::CheckDocumentationPolicyForFunc(FUnrealClassDefinitionInfo&
 		const FString* FunctionTooltip = FunctionDef.FindMetaData(NAME_ToolTip);
 		if (FunctionTooltip == nullptr)
 		{
-			UE_LOG_ERROR_UHT(TEXT("Function '%s::%s' does not provide a tooltip / comment (DocumentationPolicy)."), *ClassDef.GetName(), *FunctionDef.GetName());
+			UE_LOG_ERROR_UHT(*this, TEXT("Function '%s::%s' does not provide a tooltip / comment (DocumentationPolicy)."), *ClassDef.GetName(), *FunctionDef.GetName());
 		}
 	}
 
@@ -8947,7 +8846,7 @@ void FHeaderParser::CheckDocumentationPolicyForFunc(FUnrealClassDefinitionInfo&
 		const FString* FunctionComment = FunctionDef.FindMetaData(NAME_Comment);
 		if (FunctionComment == nullptr)
 		{
-			UE_LOG_ERROR_UHT(TEXT("Function '%s::%s' does not provide a comment (DocumentationPolicy)."), *ClassDef.GetName(), *FunctionDef.GetName());
+			UE_LOG_ERROR_UHT(*this, TEXT("Function '%s::%s' does not provide a comment (DocumentationPolicy)."), *ClassDef.GetName(), *FunctionDef.GetName());
 			return;
 		}
 		
@@ -8977,7 +8876,7 @@ void FHeaderParser::CheckDocumentationPolicyForFunc(FUnrealClassDefinitionInfo&
 				const FString* ParamToolTip = ParamToolTips.Find(ParamName);
 				if (ParamToolTip == nullptr)
 				{
-					UE_LOG_ERROR_UHT(TEXT("Function '%s::%s' doesn't provide a tooltip for parameter '%s' (DocumentationPolicy)."), *ClassDef.GetName(), *FunctionDef.GetName(), *ParamName.ToString());
+					UE_LOG_ERROR_UHT(*this, TEXT("Function '%s::%s' doesn't provide a tooltip for parameter '%s' (DocumentationPolicy)."), *ClassDef.GetName(), *FunctionDef.GetName(), *ParamName.ToString());
 				}
 				ExistingFields.Add(ParamName);
 			}
@@ -8992,7 +8891,7 @@ void FHeaderParser::CheckDocumentationPolicyForFunc(FUnrealClassDefinitionInfo&
 				}
 				if (!ExistingFields.Contains(ParamName))
 				{
-					UE_LOG_ERROR_UHT(TEXT("Function '%s::%s' provides a tooltip for an unknown parameter '%s' (DocumentationPolicy)."), *ClassDef.GetName(), *FunctionDef.GetName(), *Pair.Key.ToString());
+					UE_LOG_ERROR_UHT(*this, TEXT("Function '%s::%s' provides a tooltip for an unknown parameter '%s' (DocumentationPolicy)."), *ClassDef.GetName(), *FunctionDef.GetName(), *Pair.Key.ToString());
 				}
 			}
 
@@ -9008,7 +8907,7 @@ void FHeaderParser::CheckDocumentationPolicyForFunc(FUnrealClassDefinitionInfo&
 				const FName* ExistingParam = ToolTipToParam.Find(Pair.Value);
 				if (ExistingParam != nullptr)
 				{
-					UE_LOG_ERROR_UHT(TEXT("Function '%s::%s' uses identical tooltips for parameters '%s' and '%s' (DocumentationPolicy)."), *ClassDef.GetName(), *FunctionDef.GetName(), *ExistingParam->ToString(), *Pair.Key.ToString());
+					UE_LOG_ERROR_UHT(*this, TEXT("Function '%s::%s' uses identical tooltips for parameters '%s' and '%s' (DocumentationPolicy)."), *ClassDef.GetName(), *FunctionDef.GetName(), *ExistingParam->ToString(), *Pair.Key.ToString());
 				}
 				ToolTipToParam.Add(MoveTemp(Pair.Value), Pair.Key);
 			}
@@ -9063,12 +8962,12 @@ void FHeaderParser::ConditionalLogPointerUsage(EPointerMemberBehavior PointerMem
 	{
 		case EPointerMemberBehavior::Disallow:
 		{
-			UE_LOG(LogCompile, Error, TEXT("%s usage in member declaration detected in '%s', line %d [[%s]]"), PointerTypeDesc, *GetSourceFileContext(), InputLine, *PointerTypeDecl);
+			UE_LOG(LogCompile, Error, TEXT("%s usage in member declaration detected in '%s', line %d [[%s]]"), PointerTypeDesc, *Filename, InputLine, *PointerTypeDecl);
 		}
 		break;
 		case EPointerMemberBehavior::AllowAndLog:
 		{
-			UE_LOG(LogCompile, Log, TEXT("%s usage in member declaration detected in '%s', line %d [[%s]]"), PointerTypeDesc, *GetSourceFileContext(), InputLine, *PointerTypeDecl);
+			UE_LOG(LogCompile, Log, TEXT("%s usage in member declaration detected in '%s', line %d [[%s]]"), PointerTypeDesc, *Filename, InputLine, *PointerTypeDecl);
 		}
 		break;
 		case EPointerMemberBehavior::AllowSilently:
@@ -9096,14 +8995,13 @@ FUnrealFunctionDefinitionInfo& CreateFunctionImpl(FUnrealSourceFile& SourceFile,
 			FUnrealFieldDefinitionInfo* Type = *TypeIterator;
 			if (Type->GetFName() == FuncInfo.Function.Identifier)
 			{
-				FError::Throwf(TEXT("'%s' conflicts with '%s'"), FuncInfo.Function.Identifier, *Type->GetFullName());
+				FUHTException::Throwf(SourceFile, InputLine, TEXT("'%s' conflicts with '%s'"), FuncInfo.Function.Identifier, *Type->GetFullName());
 			}
 		}
 	}
 
 	TSharedRef<FUnrealFunctionDefinitionInfo> FuncDefRef = MakeShared<FUnrealFunctionDefinitionInfo>(SourceFile, InputLine, FString(FuncInfo.Function.Identifier), FName(FuncInfo.Function.Identifier, FNAME_Add), OuterDef, MoveTemp(FuncInfo), InFunctionType);
 	FUnrealFunctionDefinitionInfo& FuncDef = *FuncDefRef;
-
 	FuncDef.GetFunctionData().SetFunctionNames(FuncDef);
 
 	CurrentScope->AddType(FuncDef);
diff --git a/Engine/Source/Programs/UnrealHeaderTool/Private/HeaderParser.h b/Engine/Source/Programs/UnrealHeaderTool/Private/HeaderParser.h
--- a/Engine/Source/Programs/UnrealHeaderTool/Private/HeaderParser.h
+++ b/Engine/Source/Programs/UnrealHeaderTool/Private/HeaderParser.h
@@ -5,7 +5,6 @@
 #include "CoreMinimal.h"
 #include "ParserHelper.h"
 #include "BaseParser.h"
-#include "Misc/CompilationResult.h"
 #include "Scope.h"
 #include "UnrealTypeDefinitionInfo.h"
 #include "GeneratedCodeVersion.h"
@@ -129,7 +128,7 @@ struct FDocumentationPolicy
 // Header parser class.  Extracts metadata from annotated C++ headers and gathers enough
 // information to autogenerate additional headers and other boilerplate code.
 //
-class FHeaderParser : public FBaseParser, public FContextSupplier
+class FHeaderParser : public FBaseParser
 {
 public:
 	// Performs a preliminary parse of the text in the specified buffer, pulling out:
@@ -153,18 +152,14 @@ public:
 	 * Parse the given header.
 	 * 
 	 * @param PackDef The owning package of the source
-	 * @param Warn The feedback context to use for any logging
 	 * @param SourceFile The source file being parsed
 	 */
-	static ECompilationResult::Type Parse(FUnrealPackageDefinitionInfo& PackageDef, FFeedbackContext* Warn, FUnrealSourceFile& SourceFile);
+	static void Parse(FUnrealPackageDefinitionInfo& PackageDef, FUnrealSourceFile& SourceFile);
 
 protected:
 	friend struct FScriptLocation;
 	friend struct FNativeClassHeaderGenerator;
 
-	// For compiling messages and errors.
-	FFeedbackContext* Warn;
-
 	// Filename currently being parsed
 	const FString Filename;
 
@@ -328,7 +323,7 @@ protected:
 	TMap<int32, FString> RPCsNeedingHookup;
 
 	// Constructor.
-	explicit FHeaderParser(FUnrealPackageDefinitionInfo& InPackageDef, FFeedbackContext* InWarn, FUnrealSourceFile& InSourceFile);
+	explicit FHeaderParser(FUnrealPackageDefinitionInfo& InPackageDef, FUnrealSourceFile& InSourceFile);
 
 	// Returns true if the token is a dynamic delegate declaration
 	bool IsValidDelegateDeclaration(const FToken& Token) const;
@@ -341,8 +336,8 @@ protected:
 
 public:
 	// Throws if a specifier value wasn't provided
-	static void RequireSpecifierValue(const FPropertySpecifier& Specifier, bool bRequireExactlyOne = false);
-	static FString RequireExactlyOneSpecifierValue(const FPropertySpecifier& Specifier);
+	static void RequireSpecifierValue(const FUHTExceptionContext& Context, const FPropertySpecifier& Specifier, bool bRequireExactlyOne = false);
+	static FString RequireExactlyOneSpecifierValue(const FUHTExceptionContext& Context, const FPropertySpecifier& Specifier);
 
 	/**
 	* Find a field in the specified context.  Starts with the specified scope, then iterates
@@ -383,19 +378,13 @@ protected:
 	*/
 	static TMap<FName, FString> GetParameterToolTipsFromFunctionComment(const FString& Input);
 	
-	FString GetSourceFileContext() const;
-
-	// FContextSupplier interface.
-	virtual FString GetContext() override;
-	// End of FContextSupplier interface.
-
 	// High-level compiling functions.
 	/**
 	 * Parses given source file.
 	 *
 	 * @returns Compilation result enum.
 	 */
-	ECompilationResult::Type ParseHeader();
+	void ParseHeader();
 	void CompileDirective();
 	FUnrealEnumDefinitionInfo& CompileEnum();
 	FUnrealScriptStructDefinitionInfo& CompileStructDeclaration();
@@ -572,7 +561,7 @@ protected:
 	void InitScriptLocation( FScriptLocation& Retry );
 	void ReturnToLocation( const FScriptLocation& Retry, bool Binary=1, bool Text=1 );
 
-	static void ValidatePropertyIsDeprecatedIfNecessary(const FPropertyBase& VarProperty, const EPropertyFlags* OuterPropertyFlags);
+	void ValidatePropertyIsDeprecatedIfNecessary(const FPropertyBase& VarProperty, const EPropertyFlags* OuterPropertyFlags);
 
 	// Cache of ScriptStructs that have been validated for Net Replication and RPC
 	TSet<FUnrealScriptStructDefinitionInfo*> ScriptStructsValidForNet;
@@ -589,9 +578,6 @@ private:
 	// Definition of package being parsed
 	FUnrealPackageDefinitionInfo& PackageDef;
 
-	// Source file currently parsed by UHT.
-	FUnrealSourceFile& SourceFile;
-
 	// True if the module currently being parsed is part of the engine, as opposed to being part of a game
 	bool bIsCurrentModulePartOfEngine;
 
@@ -623,7 +609,7 @@ private:
 	void VerifyRepNotifyCallback(FUnrealPropertyDefinitionInfo& PropertyDef, FUnrealFunctionDefinitionInfo* TargetFuncDef);
 
 	// Constructs the policy from a string
-	static FDocumentationPolicy GetDocumentationPolicyFromName(const FString& PolicyName);
+	static FDocumentationPolicy GetDocumentationPolicyFromName(const FUHTExceptionContext& Context, const FString& PolicyName);
 
 	// Constructs the policy for documentation checks for a given struct
 	static FDocumentationPolicy GetDocumentationPolicyForStruct(FUnrealStructDefinitionInfo& StructDef);
@@ -690,12 +676,12 @@ public:
 class FHeaderPreParser : public FBaseParser
 {
 public:
-	FHeaderPreParser()
+	FHeaderPreParser(FUnrealSourceFile& InSourceFile)
+		: FBaseParser(InSourceFile)
 	{
 	}
 
 	TSharedRef<FUnrealTypeDefinitionInfo> ParseClassDeclaration(
-		FUnrealSourceFile& SourceFile,
 		const TCHAR* InputText,
 		int32 InLineNumber,
 		bool bClassIsAnInterface,
@@ -703,13 +689,11 @@ public:
 	);
 
 	TSharedRef<FUnrealTypeDefinitionInfo> ParseEnumDeclaration(
-		FUnrealSourceFile& SourceFile,
 		const TCHAR* InputText,
 		int32 InLineNumber
 	);
 
 	TSharedRef<FUnrealTypeDefinitionInfo> ParseStructDeclaration(
-		FUnrealSourceFile& SourceFile,
 		const TCHAR* InputText,
 		int32 InLineNumber
 	);
diff --git a/Engine/Source/Programs/UnrealHeaderTool/Private/IScriptGeneratorPluginInterface.cpp b/Engine/Source/Programs/UnrealHeaderTool/Private/IScriptGeneratorPluginInterface.cpp
--- a/Engine/Source/Programs/UnrealHeaderTool/Private/IScriptGeneratorPluginInterface.cpp
+++ b/Engine/Source/Programs/UnrealHeaderTool/Private/IScriptGeneratorPluginInterface.cpp
@@ -1,10 +1,12 @@
 // Copyright Epic Games, Inc. All Rights Reserved.
 
 #include "IScriptGeneratorPluginInterface.h"
+#include "Exceptions.h"
 #include "UnrealHeaderTool.h"
-#include "UObject/ErrorException.h"
 #include "Algo/FindSortedStringCaseInsensitive.h"
 
+extern FString GManifestFilename; // This is a bit of a hack...
+
 EBuildModuleType::Type EBuildModuleType::Parse(const TCHAR* Value)
 {
 	static const TCHAR* AlphabetizedTypes[] = {
@@ -24,7 +26,7 @@ EBuildModuleType::Type EBuildModuleType::Parse(const TCHAR* Value)
 	int32 TypeIndex = Algo::FindSortedStringCaseInsensitive(Value, AlphabetizedTypes);
 	if (TypeIndex < 0)
 	{
-		FError::Throwf(TEXT("Unrecognized EBuildModuleType name: %s"), Value);
+		FUHTException::Throwf(FString(GManifestFilename), 1, TEXT("Unrecognized EBuildModuleType name: %s"), Value);
 	}
 
 	static EBuildModuleType::Type AlphabetizedValues[] = {
@@ -58,7 +60,7 @@ EPackageOverrideType::Type EPackageOverrideType::Parse(const TCHAR* Value)
 	int32 TypeIndex = Algo::FindSortedStringCaseInsensitive(Value, AlphabetizedTypes);
 	if (TypeIndex < 0)
 	{
-		FError::Throwf(TEXT("Unrecognized EPackageOverrideType name: %s"), Value);
+		FUHTException::Throwf(FString(GManifestFilename), 1, TEXT("Unrecognized EPackageOverrideType name: %s"), Value);
 	}
 
 	static EPackageOverrideType::Type AlphabetizedValues[] = {
@@ -71,4 +73,4 @@ EPackageOverrideType::Type EPackageOverrideType::Parse(const TCHAR* Value)
 	};
 
 	return AlphabetizedValues[TypeIndex];
-}
\ No newline at end of file
+}
diff --git a/Engine/Source/Programs/UnrealHeaderTool/Private/Manifest.cpp b/Engine/Source/Programs/UnrealHeaderTool/Private/Manifest.cpp
--- a/Engine/Source/Programs/UnrealHeaderTool/Private/Manifest.cpp
+++ b/Engine/Source/Programs/UnrealHeaderTool/Private/Manifest.cpp
@@ -1,18 +1,21 @@
 // Copyright Epic Games, Inc. All Rights Reserved.
 
 #include "Manifest.h"
+#include "Exceptions.h"
+#include "IScriptGeneratorPluginInterface.h"
 #include "UnrealHeaderTool.h"
 #include "Misc/DateTime.h"
 #include "Logging/LogMacros.h"
 #include "HAL/FileManager.h"
 #include "Misc/FileHelper.h"
 #include "Misc/Paths.h"
-#include "UObject/ErrorException.h"
 #include "Misc/PackageName.h"
 #include "UnrealHeaderToolGlobals.h"
 #include "Serialization/JsonTypes.h"
 #include "Serialization/JsonSerializer.h"
 
+FString GManifestFilename;
+
 namespace
 {
 	template <typename T> struct TJsonFieldType;
@@ -27,7 +30,7 @@ namespace
 	{
 		if (JsonValue->Type != TJsonFieldType<T>::Value)
 		{
-			FError::Throwf(TEXT("'%s' is the wrong type"), Outer);
+			FUHTException::Throwf(FString(GManifestFilename), 1, TEXT("'%s' is the wrong type"), Outer);
 		}
 
 		JsonValue->AsArgumentType(OutVal);
@@ -40,12 +43,12 @@ namespace
 
 		if (!JsonValue)
 		{
-			FError::Throwf(TEXT("Unable to find field '%s' in '%s'"), FieldName, Outer);
+			FUHTException::Throwf(FString(GManifestFilename), 1, TEXT("Unable to find field '%s' in '%s'"), FieldName, Outer);
 		}
 
 		if ((*JsonValue)->Type != TJsonFieldType<T>::Value)
 		{
-			FError::Throwf(TEXT("Field '%s' in '%s' is the wrong type"), Outer);
+			FUHTException::Throwf(FString(GManifestFilename), 1, TEXT("Field '%s' in '%s' is the wrong type"), Outer);
 		}
 
 		(*JsonValue)->AsArgumentType(OutVal);
@@ -63,13 +66,14 @@ namespace
 
 FManifest FManifest::LoadFromFile(const FString& Filename)
 {
+	GManifestFilename = Filename;
 	FManifest Result;
 	FString FilenamePath = FPaths::GetPath(Filename);
 	FString Json;
 
 	if (!FFileHelper::LoadFileToString(Json, *Filename))
 	{
-		FError::Throwf(TEXT("Unable to load manifest: %s"), *Filename);
+		FUHTException::Throwf(FString(Filename), 1, TEXT("Unable to load manifest: %s"), *Filename);
 	}
 
 	TSharedPtr<FJsonObject> RootObject = TSharedPtr<FJsonObject>();
@@ -77,7 +81,7 @@ FManifest FManifest::LoadFromFile(const FString& Filename)
 
 	if (!FJsonSerializer::Deserialize(Reader, RootObject))
 	{
-		FError::Throwf(TEXT("Manifest is malformed: %s"), *Filename);
+		FUHTException::Throwf(FString(Filename), 1, TEXT("Manifest is malformed: %s"), *Filename);
 	}
 
 	TArray<TSharedPtr<FJsonValue>> ModulesArray;
diff --git a/Engine/Source/Programs/UnrealHeaderTool/Private/ParserHelper.cpp b/Engine/Source/Programs/UnrealHeaderTool/Private/ParserHelper.cpp
--- a/Engine/Source/Programs/UnrealHeaderTool/Private/ParserHelper.cpp
+++ b/Engine/Source/Programs/UnrealHeaderTool/Private/ParserHelper.cpp
@@ -355,7 +355,7 @@ void FFuncInfo::SetFunctionNames(FUnrealFunctionDefinitionInfo& FunctionDef)
 			}
 			else if (CppImplName == FunctionName)
 			{
-				FError::Throwf(TEXT("Native implementation function must be different than original function name."));
+				FUHTException::Throwf(FunctionDef, TEXT("Native implementation function must be different than original function name."));
 			}
 
 			if (CppValidationImplName.IsEmpty() && FunctionDef.HasAllFunctionFlags(FUNC_NetValidate))
@@ -364,7 +364,7 @@ void FFuncInfo::SetFunctionNames(FUnrealFunctionDefinitionInfo& FunctionDef)
 			}
 			else if (CppValidationImplName == FunctionName)
 			{
-				FError::Throwf(TEXT("Validation function must be different than original function name."));
+				FUHTException::Throwf(FunctionDef, TEXT("Validation function must be different than original function name."));
 			}
 		}
 	}
diff --git a/Engine/Source/Programs/UnrealHeaderTool/Private/ParserHelper.h b/Engine/Source/Programs/UnrealHeaderTool/Private/ParserHelper.h
--- a/Engine/Source/Programs/UnrealHeaderTool/Private/ParserHelper.h
+++ b/Engine/Source/Programs/UnrealHeaderTool/Private/ParserHelper.h
@@ -3,7 +3,6 @@
 #pragma once
 
 #include "CoreMinimal.h"
-#include "UObject/ErrorException.h"
 #include "UObject/Package.h"
 #include "UObject/MetaData.h"
 #include "UObject/TextProperty.h"
diff --git a/Engine/Source/Programs/UnrealHeaderTool/Private/PropertyTypes.cpp b/Engine/Source/Programs/UnrealHeaderTool/Private/PropertyTypes.cpp
--- a/Engine/Source/Programs/UnrealHeaderTool/Private/PropertyTypes.cpp
+++ b/Engine/Source/Programs/UnrealHeaderTool/Private/PropertyTypes.cpp
@@ -162,65 +162,66 @@ namespace
 	 *
 	 * NOTE: This method does not support container or enum types.
 	 */
-	template <template<typename PropertyTraits> typename FuncDispatch, typename RetValue, typename ... Args>
-	RetValue PropertyTypeDispatch(EPropertyType PropertyType, Args&& ... args)
-	{
-		switch (PropertyType)
-		{
-		case CPT_Byte:						return FuncDispatch<FPropertyTypeTraitsByte>()(std::forward<Args>(args)...);
-		case CPT_Int8:						return FuncDispatch<FPropertyTypeTraitsInt8>()(std::forward<Args>(args)...);
-		case CPT_Int16:						return FuncDispatch<FPropertyTypeTraitsInt16>()(std::forward<Args>(args)...);
-		case CPT_Int:						return FuncDispatch<FPropertyTypeTraitsInt>()(std::forward<Args>(args)...);
-		case CPT_Int64:						return FuncDispatch<FPropertyTypeTraitsInt64>()(std::forward<Args>(args)...);
-		case CPT_UInt16:					return FuncDispatch<FPropertyTypeTraitsUInt16>()(std::forward<Args>(args)...);
-		case CPT_UInt32:					return FuncDispatch<FPropertyTypeTraitsUInt32>()(std::forward<Args>(args)...);
-		case CPT_UInt64:					return FuncDispatch<FPropertyTypeTraitsUInt64>()(std::forward<Args>(args)...);
-		case CPT_Bool:						return FuncDispatch<FPropertyTypeTraitsBool>()(std::forward<Args>(args)...);
-		case CPT_Bool8:						return FuncDispatch<FPropertyTypeTraitsBool8>()(std::forward<Args>(args)...);
-		case CPT_Bool16:					return FuncDispatch<FPropertyTypeTraitsBool16>()(std::forward<Args>(args)...);
-		case CPT_Bool32:					return FuncDispatch<FPropertyTypeTraitsBool32>()(std::forward<Args>(args)...);
-		case CPT_Bool64:					return FuncDispatch<FPropertyTypeTraitsBool64>()(std::forward<Args>(args)...);
-		case CPT_Float:						return FuncDispatch<FPropertyTypeTraitsFloat>()(std::forward<Args>(args)...);
-		case CPT_Double:					return FuncDispatch<FPropertyTypeTraitsDouble>()(std::forward<Args>(args)...);
-		case CPT_FLargeWorldCoordinatesReal:return FuncDispatch<FPropertyTypeTraitsLargeWorldCoordinatesReal>()(std::forward<Args>(args)...);
-		case CPT_ObjectReference:			return FuncDispatch<FPropertyTypeTraitsObjectReference>()(std::forward<Args>(args)...);
-		case CPT_WeakObjectReference:		return FuncDispatch<FPropertyTypeTraitsWeakObjectReference>()(std::forward<Args>(args)...);
-		case CPT_LazyObjectReference:		return FuncDispatch<FPropertyTypeTraitsLazyObjectReference>()(std::forward<Args>(args)...);
-		case CPT_ObjectPtrReference:		return FuncDispatch<FPropertyTypeTraitsObjectPtrReference>()(std::forward<Args>(args)...);
-		case CPT_SoftObjectReference:		return FuncDispatch<FPropertyTypeTraitsSoftObjectReference>()(std::forward<Args>(args)...);
-		case CPT_Interface:					return FuncDispatch<FPropertyTypeTraitsInterface>()(std::forward<Args>(args)...);
-		case CPT_Name:						return FuncDispatch<FPropertyTypeTraitsName>()(std::forward<Args>(args)...);
-		case CPT_String:					return FuncDispatch<FPropertyTypeTraitsString>()(std::forward<Args>(args)...);
-		case CPT_Text:						return FuncDispatch<FPropertyTypeTraitsText>()(std::forward<Args>(args)...);
-		case CPT_Struct:					return FuncDispatch<FPropertyTypeTraitsStruct>()(std::forward<Args>(args)...);
-		case CPT_Delegate:					return FuncDispatch<FPropertyTypeTraitsDelegate>()(std::forward<Args>(args)...);
-		case CPT_MulticastDelegate:			return FuncDispatch<FPropertyTypeTraitsMulticastDelegate>()(std::forward<Args>(args)...);
-		case CPT_FieldPath:					return FuncDispatch<FPropertyTypeTraitsFieldPath>()(std::forward<Args>(args)...);
-		default:							FError::Throwf(TEXT("Unknown property type %i"), (uint8)PropertyType);
+	template <template<typename PropertyTraits> typename FuncDispatch, typename RetValue, typename PropDefType, typename ... Args>
+	RetValue SimplePropertyTypeDispatch(PropDefType& PropDef, Args&& ... args)
+	{
+		switch (PropDef.GetPropertyBase().Type)
+		{
+		case CPT_Byte:						return FuncDispatch<FPropertyTypeTraitsByte>()(PropDef, std::forward<Args>(args)...);
+		case CPT_Int8:						return FuncDispatch<FPropertyTypeTraitsInt8>()(PropDef, std::forward<Args>(args)...);
+		case CPT_Int16:						return FuncDispatch<FPropertyTypeTraitsInt16>()(PropDef, std::forward<Args>(args)...);
+		case CPT_Int:						return FuncDispatch<FPropertyTypeTraitsInt>()(PropDef, std::forward<Args>(args)...);
+		case CPT_Int64:						return FuncDispatch<FPropertyTypeTraitsInt64>()(PropDef, std::forward<Args>(args)...);
+		case CPT_UInt16:					return FuncDispatch<FPropertyTypeTraitsUInt16>()(PropDef, std::forward<Args>(args)...);
+		case CPT_UInt32:					return FuncDispatch<FPropertyTypeTraitsUInt32>()(PropDef, std::forward<Args>(args)...);
+		case CPT_UInt64:					return FuncDispatch<FPropertyTypeTraitsUInt64>()(PropDef, std::forward<Args>(args)...);
+		case CPT_Bool:						return FuncDispatch<FPropertyTypeTraitsBool>()(PropDef, std::forward<Args>(args)...);
+		case CPT_Bool8:						return FuncDispatch<FPropertyTypeTraitsBool8>()(PropDef, std::forward<Args>(args)...);
+		case CPT_Bool16:					return FuncDispatch<FPropertyTypeTraitsBool16>()(PropDef, std::forward<Args>(args)...);
+		case CPT_Bool32:					return FuncDispatch<FPropertyTypeTraitsBool32>()(PropDef, std::forward<Args>(args)...);
+		case CPT_Bool64:					return FuncDispatch<FPropertyTypeTraitsBool64>()(PropDef, std::forward<Args>(args)...);
+		case CPT_Float:						return FuncDispatch<FPropertyTypeTraitsFloat>()(PropDef, std::forward<Args>(args)...);
+		case CPT_Double:					return FuncDispatch<FPropertyTypeTraitsDouble>()(PropDef, std::forward<Args>(args)...);
+		case CPT_FLargeWorldCoordinatesReal:return FuncDispatch<FPropertyTypeTraitsLargeWorldCoordinatesReal>()(PropDef, std::forward<Args>(args)...);
+		case CPT_ObjectReference:			return FuncDispatch<FPropertyTypeTraitsObjectReference>()(PropDef, std::forward<Args>(args)...);
+		case CPT_WeakObjectReference:		return FuncDispatch<FPropertyTypeTraitsWeakObjectReference>()(PropDef, std::forward<Args>(args)...);
+		case CPT_LazyObjectReference:		return FuncDispatch<FPropertyTypeTraitsLazyObjectReference>()(PropDef, std::forward<Args>(args)...);
+		case CPT_ObjectPtrReference:		return FuncDispatch<FPropertyTypeTraitsObjectPtrReference>()(PropDef, std::forward<Args>(args)...);
+		case CPT_SoftObjectReference:		return FuncDispatch<FPropertyTypeTraitsSoftObjectReference>()(PropDef, std::forward<Args>(args)...);
+		case CPT_Interface:					return FuncDispatch<FPropertyTypeTraitsInterface>()(PropDef, std::forward<Args>(args)...);
+		case CPT_Name:						return FuncDispatch<FPropertyTypeTraitsName>()(PropDef, std::forward<Args>(args)...);
+		case CPT_String:					return FuncDispatch<FPropertyTypeTraitsString>()(PropDef, std::forward<Args>(args)...);
+		case CPT_Text:						return FuncDispatch<FPropertyTypeTraitsText>()(PropDef, std::forward<Args>(args)...);
+		case CPT_Struct:					return FuncDispatch<FPropertyTypeTraitsStruct>()(PropDef, std::forward<Args>(args)...);
+		case CPT_Delegate:					return FuncDispatch<FPropertyTypeTraitsDelegate>()(PropDef, std::forward<Args>(args)...);
+		case CPT_MulticastDelegate:			return FuncDispatch<FPropertyTypeTraitsMulticastDelegate>()(PropDef, std::forward<Args>(args)...);
+		case CPT_FieldPath:					return FuncDispatch<FPropertyTypeTraitsFieldPath>()(PropDef, std::forward<Args>(args)...);
+		default:							FUHTException::Throwf(PropDef, TEXT("Unknown property type %i"), (uint8)PropDef.GetPropertyBase().Type);
 		}
 	}
 
 	/**
 	 * Given just the property type, invoke the dispatch functor.
 	 */
-	template <template <typename PropertyTraits> typename FuncDispatch, bool bHandleContainers, typename RetValue, typename ... Args>
-	RetValue PropertyTypeDispatch(const FPropertyBase& VarProperty, Args&& ... args)
+	template <template <typename PropertyTraits> typename FuncDispatch, bool bHandleContainers, typename RetValue, typename PropDefType, typename ... Args>
+	RetValue ComplexPropertyTypeDispatch(PropDefType& PropDef, Args&& ... args)
 	{
+		const FPropertyBase& VarProperty = PropDef.GetPropertyBase();
 		if constexpr (bHandleContainers)
 		{
 			switch (VarProperty.ArrayType)
 			{
 			case EArrayType::Static:
-				return FuncDispatch<FPropertyTypeTraitsStaticArray>()(std::forward<Args>(args)...);
+				return FuncDispatch<FPropertyTypeTraitsStaticArray>()(PropDef, std::forward<Args>(args)...);
 			case EArrayType::Dynamic:
-				return FuncDispatch<FPropertyTypeTraitsDynamicArray>()(std::forward<Args>(args)...);
+				return FuncDispatch<FPropertyTypeTraitsDynamicArray>()(PropDef, std::forward<Args>(args)...);
 			case EArrayType::Set:
-				return FuncDispatch<FPropertyTypeTraitsSet>()(std::forward<Args>(args)...);
+				return FuncDispatch<FPropertyTypeTraitsSet>()(PropDef, std::forward<Args>(args)...);
 			}
 
 			if (VarProperty.MapKeyProp.IsValid())
 			{
-				return FuncDispatch<FPropertyTypeTraitsMap>()(std::forward<Args>(args)...);
+				return FuncDispatch<FPropertyTypeTraitsMap>()(PropDef, std::forward<Args>(args)...);
 			}
 		}
 
@@ -228,10 +229,10 @@ namespace
 		// NOTE: VarProperty.EnumDef is a union and might not be an enum
 		if (VarProperty.IsEnum())
 		{
-			return FuncDispatch<FPropertyTypeTraitsEnum>()(std::forward<Args>(args)...);
+			return FuncDispatch<FPropertyTypeTraitsEnum>()(PropDef, std::forward<Args>(args)...);
 		}
 
-		return PropertyTypeDispatch<FuncDispatch, RetValue>(VarProperty.Type, std::forward<Args>(args)...);
+		return SimplePropertyTypeDispatch<FuncDispatch, RetValue>(PropDef, std::forward<Args>(args)...);
 	}
 
 	//----------------------------------------------------------------------------------------------------------------------------------------------------------
@@ -312,18 +313,18 @@ namespace
 	template <bool bHandleContainers>
 	void CreatePropertyHelper(FUnrealPropertyDefinitionInfo& PropDef)
 	{
-		return PropertyTypeDispatch<CreatePropertyDispatch, bHandleContainers, void>(PropDef.GetPropertyBase(), std::ref(PropDef));
+		return ComplexPropertyTypeDispatch<CreatePropertyDispatch, bHandleContainers, void>(PropDef);
 	}
 
 	template <bool bHandleContainers>
 	FProperty* CreateEngineTypeHelper(FUnrealPropertyDefinitionInfo& PropDef, FFieldVariant Scope, const FName& Name, EObjectFlags ObjectFlags)
 	{
-		return PropertyTypeDispatch<CreateEngineTypeDispatch, bHandleContainers, FProperty*>(PropDef.GetPropertyBase(), std::ref(PropDef), Scope, std::ref(Name), ObjectFlags);
+		return ComplexPropertyTypeDispatch<CreateEngineTypeDispatch, bHandleContainers, FProperty*>(PropDef, Scope, std::ref(Name), ObjectFlags);
 	}
 
 	bool IsSupportedByBlueprintSansContainers(const FUnrealPropertyDefinitionInfo& PropDef, bool bMemberVariable)
 	{
-		return PropertyTypeDispatch<IsSupportedByBlueprintDispatch, false, bool>(PropDef.GetPropertyBase(), std::ref(PropDef), bMemberVariable);
+		return ComplexPropertyTypeDispatch<IsSupportedByBlueprintDispatch, false, bool>(PropDef, bMemberVariable);
 	}
 }
 
@@ -971,7 +972,7 @@ struct FPropertyTypeTraitsEnum : public FPropertyTypeTraitsBase
 		}
 		if (Enum->HasMetaData(TEXT("Hidden"), EnumEntryIndex))
 		{
-			FError::Throwf(TEXT("Hidden enum entries cannot be used as default values: %s \"%s\" "), *PropDef.GetName(), *CppForm);
+			FUHTException::Throwf(PropDef, TEXT("Hidden enum entries cannot be used as default values: %s \"%s\" "), *PropDef.GetName(), *CppForm);
 		}
 		return true;
 	}
@@ -1673,7 +1674,7 @@ struct FPropertyTypeTraitsText : public FPropertyTypeTraitsBase
 		// These should be replaced with INVTEXT as FText::FromString can produce inconsistent keys
 		if (FDefaultValueHelper::StringFromCppString(CppForm, TEXT("FText::FromString"), OutForm))
 		{
-			UE_LOG_WARNING_UHT(TEXT("FText::FromString should be replaced with INVTEXT for default parameter values"));
+			UE_LOG_WARNING_UHT(PropDef, TEXT("FText::FromString should be replaced with INVTEXT for default parameter values"));
 			return true;
 		}
 
@@ -1699,7 +1700,7 @@ struct FPropertyTypeTraitsText : public FPropertyTypeTraitsBase
 			{
 				if (ParsedTextNamespace.GetValue().Equals(UHTDummyNamespace))
 				{
-					FError::Throwf(TEXT("LOCTEXT default parameter values are not supported; use NSLOCTEXT instead: %s \"%s\" "), *PropDef.GetName(), *CppForm);
+					FUHTException::Throwf(PropDef, TEXT("LOCTEXT default parameter values are not supported; use NSLOCTEXT instead: %s \"%s\" "), *PropDef.GetName(), *CppForm);
 				}
 			}
 		}
@@ -2244,17 +2245,17 @@ struct FPropertyTypeTraitsStaticArray : public FPropertyTypeTraitsBase
 
 	static FString GetEngineClassName(const FUnrealPropertyDefinitionInfo& PropDef)
 	{
-		return PropertyTypeDispatch<GetEngineClassNameDispatch, false, FString>(PropDef.GetPropertyBase(), std::ref(PropDef));
+		return ComplexPropertyTypeDispatch<GetEngineClassNameDispatch, false, FString>(PropDef);
 	}
 
 	static FString GetCPPType(const FUnrealPropertyDefinitionInfo& PropDef, FString* ExtendedTypeText, uint32 CPPExportFlags)
 	{
-		return PropertyTypeDispatch<GetCPPTypeDispatch, false, FString>(PropDef.GetPropertyBase(), std::ref(PropDef), ExtendedTypeText, CPPExportFlags);
+		return ComplexPropertyTypeDispatch<GetCPPTypeDispatch, false, FString>(PropDef, ExtendedTypeText, CPPExportFlags);
 	}
 
 	static FString GetCPPTypeForwardDeclaration(const FUnrealPropertyDefinitionInfo& PropDef)
 	{
-		return PropertyTypeDispatch<GetCPPTypeForwardDeclarationDispatch, false, FString>(PropDef.GetPropertyBase(), std::ref(PropDef));
+		return ComplexPropertyTypeDispatch<GetCPPTypeForwardDeclarationDispatch, false, FString>(PropDef);
 	}
 };
 
@@ -2504,7 +2505,7 @@ bool FPropertyTraits::DefaultValueStringCppFormatToInnerFormat(const FUnrealProp
 		return false;
 	}
 
-	return PropertyTypeDispatch<DefaultValueStringCppFormatToInnerFormatDispatch, false, bool>(PropDef.GetPropertyBase(), std::ref(PropDef), std::ref(CppForm), std::ref(OutForm));
+	return ComplexPropertyTypeDispatch<DefaultValueStringCppFormatToInnerFormatDispatch, false, bool>(PropDef, std::ref(CppForm), std::ref(OutForm));
 }
 
 TSharedRef<FUnrealPropertyDefinitionInfo> FPropertyTraits::CreateProperty(const FPropertyBase& VarProperty, FUnrealTypeDefinitionInfo& Outer, const FName& Name, EVariableCategory VariableCategory, EAccessSpecifier AccessSpecifier, const TCHAR* Dimensions, FUnrealSourceFile& SourceFile, int LineNumber, int ParsePosition)
@@ -2520,7 +2521,7 @@ TSharedRef<FUnrealPropertyDefinitionInfo> FPropertyTraits::CreateProperty(const
 	// Technically, it should be added to the outermost property.
 	if (UHTCast<FUnrealPropertyDefinitionInfo>(Outer) == nullptr)
 	{
-		FUHTMetaData::RemapMetaData(PropDef.GetPropertyBase().MetaData);
+		FUHTMetaData::RemapMetaData(PropDef, PropDef.GetPropertyBase().MetaData);
 		PropDef.ValidateMetaDataFormat(PropDef.GetPropertyBase().MetaData);
 	}
 	return PropDefRef;
@@ -2586,22 +2587,22 @@ FProperty* FPropertyTraits::CreateEngineType(TSharedRef<FUnrealPropertyDefinitio
 
 bool FPropertyTraits::IsSupportedByBlueprint(const FUnrealPropertyDefinitionInfo& PropDef, bool bMemberVariable)
 {
-	return PropertyTypeDispatch<IsSupportedByBlueprintDispatch, true, bool>(PropDef.GetPropertyBase(), std::ref(PropDef), bMemberVariable);
+	return ComplexPropertyTypeDispatch<IsSupportedByBlueprintDispatch, true, bool>(PropDef, bMemberVariable);
 }
 
 FString FPropertyTraits::GetEngineClassName(const FUnrealPropertyDefinitionInfo& PropDef)
 {
-	return PropertyTypeDispatch<GetEngineClassNameDispatch, true, FString>(PropDef.GetPropertyBase(), std::ref(PropDef));
+	return ComplexPropertyTypeDispatch<GetEngineClassNameDispatch, true, FString>(PropDef);
 }
 
 FString FPropertyTraits::GetCPPType(const FUnrealPropertyDefinitionInfo& PropDef, FString* ExtendedTypeText/* = nullptr */, uint32 CPPExportFlags/* = 0 */)
 {
-	return PropertyTypeDispatch<GetCPPTypeDispatch, true, FString>(PropDef.GetPropertyBase(), std::ref(PropDef), ExtendedTypeText, CPPExportFlags);
+	return ComplexPropertyTypeDispatch<GetCPPTypeDispatch, true, FString>(PropDef, ExtendedTypeText, CPPExportFlags);
 }
 
 FString FPropertyTraits::GetCPPTypeForwardDeclaration(const FUnrealPropertyDefinitionInfo& PropDef)
 {
-	return PropertyTypeDispatch<GetCPPTypeForwardDeclarationDispatch, true, FString>(PropDef.GetPropertyBase(), std::ref(PropDef));
+	return ComplexPropertyTypeDispatch<GetCPPTypeForwardDeclarationDispatch, true, FString>(PropDef);
 }
 
 bool FPropertyTraits::SameType(const FUnrealPropertyDefinitionInfo& Lhs, const FUnrealPropertyDefinitionInfo& Rhs)
diff --git a/Engine/Source/Programs/UnrealHeaderTool/Private/Scope.cpp b/Engine/Source/Programs/UnrealHeaderTool/Private/Scope.cpp
--- a/Engine/Source/Programs/UnrealHeaderTool/Private/Scope.cpp
+++ b/Engine/Source/Programs/UnrealHeaderTool/Private/Scope.cpp
@@ -2,7 +2,6 @@
 
 #include "Scope.h"
 #include "UnrealHeaderTool.h"
-#include "UObject/ErrorException.h"
 #include "ParserHelper.h"
 #include "UnrealTypeDefinitionInfo.h"
 #include "ClassMaps.h"
diff --git a/Engine/Source/Programs/UnrealHeaderTool/Private/UnrealHeaderToolGlobals.h b/Engine/Source/Programs/UnrealHeaderTool/Private/UnrealHeaderToolGlobals.h
--- a/Engine/Source/Programs/UnrealHeaderTool/Private/UnrealHeaderToolGlobals.h
+++ b/Engine/Source/Programs/UnrealHeaderTool/Private/UnrealHeaderToolGlobals.h
@@ -8,9 +8,3 @@
 #include "UnrealHeaderToolStats.h"
 
 DECLARE_LOG_CATEGORY_EXTERN(LogCompile, Log, All);
-
-extern bool GUHTWarningLogged;
-extern bool GUHTErrorLogged;
-
-#define UE_LOG_WARNING_UHT(Format, ...) { GUHTWarningLogged = true; UE_LOG(LogCompile, Warning, Format, ##__VA_ARGS__); }
-#define UE_LOG_ERROR_UHT(Format, ...) { GUHTErrorLogged = true; UE_LOG(LogCompile, Error, Format, ##__VA_ARGS__); }
diff --git a/Engine/Source/Programs/UnrealHeaderTool/Private/UnrealTypeDefinitionInfo.cpp b/Engine/Source/Programs/UnrealHeaderTool/Private/UnrealTypeDefinitionInfo.cpp
--- a/Engine/Source/Programs/UnrealHeaderTool/Private/UnrealTypeDefinitionInfo.cpp
+++ b/Engine/Source/Programs/UnrealHeaderTool/Private/UnrealTypeDefinitionInfo.cpp
@@ -17,7 +17,6 @@
 #include "Math/UnitConversion.h"
 #include "Misc/ConfigCacheIni.h"
 #include "Misc/PackageName.h"
-#include "UObject/ErrorException.h"
 #include "UObject/Interface.h"
 #include "UObject/ObjectRedirector.h"
 
@@ -153,7 +152,7 @@ namespace
 	}
 }
 
-void FUHTMetaData::RemapMetaData(TMap<FName, FString>& MetaData)
+void FUHTMetaData::RemapMetaData(FUnrealTypeDefinitionInfo& TypeDef, TMap<FName, FString>& MetaData)
 {
 	// Evaluate any key redirects on the passed in pairs
 	for (TPair<FName, FString>& Pair : MetaData)
@@ -163,7 +162,7 @@ void FUHTMetaData::RemapMetaData(TMap<FName, FString>& MetaData)
 
 		if (NewKey != NAME_None)
 		{
-			UE_LOG_WARNING_UHT(TEXT("Remapping old metadata key '%s' to new key '%s', please update the declaration."), *CurrentKey.ToString(), *NewKey.ToString());
+			UE_LOG_WARNING_UHT(TypeDef, TEXT("Remapping old metadata key '%s' to new key '%s', please update the declaration."), *CurrentKey.ToString(), *NewKey.ToString());
 			CurrentKey = NewKey;
 		}
 	}
@@ -292,11 +291,16 @@ FUnrealClassDefinitionInfo* FUnrealTypeDefinitionInfo::AsClass()
 	return nullptr;
 }
 
+FString FUnrealTypeDefinitionInfo::GetFilename() const
+{
+	return HasSource() ? SourceFile->GetFilename() : FString(TEXT("UnknownSource"));
+}
+
 TSharedRef<FScope> FUnrealTypeDefinitionInfo::GetScope()
 {
 	if (!HasSource())
 	{
-		FError::Throwf(TEXT("Attempt to fetch the scope for type \"%s\" when it doesn't implement the method or there is no source file associated with the type."), *GetNameCPP());
+		FUHTException::Throwf(*this, TEXT("Attempt to fetch the scope for type \"%s\" when it doesn't implement the method or there is no source file associated with the type."), *GetNameCPP());
 	}
 	return GetUnrealSourceFile().GetScope();
 }
@@ -310,7 +314,7 @@ uint32 FUnrealTypeDefinitionInfo::GetHash(bool bIncludeNoExport) const
 {
 	if (Hash == 0)
 	{
-		FError::Throwf(TEXT("Attempt to fetch the generated hash for type \"%s\" before it has been generated.  Include dependencies, topological sort, or job graph is in error."), *GetNameCPP());
+		FUHTException::Throwf(*this, TEXT("Attempt to fetch the generated hash for type \"%s\" before it has been generated.  Include dependencies, topological sort, or job graph is in error."), *GetNameCPP());
 	}
 	return Hash;
 }
@@ -387,7 +391,7 @@ void FUnrealTypeDefinitionInfo::ValidateMetaDataFormat(const FName InKey, ECheck
 	{
 		if (!InValue.IsNumeric())
 		{
-			FError::Throwf(TEXT("Metadata value for '%s' is non-numeric : '%s'"), *InKey.ToString(), *InValue);
+			FUHTException::Throwf(*this, TEXT("Metadata value for '%s' is non-numeric : '%s'"), *InKey.ToString(), *InValue);
 		}
 	}
 	break;
@@ -406,7 +410,7 @@ void FUnrealTypeDefinitionInfo::ValidateMetaDataFormat(const FName InKey, ECheck
 
 				if (ClassDef != nullptr && ClassDef->GetName() == TEXT("BlueprintFunctionLibrary"))
 				{
-					FError::Throwf(TEXT("%s doesn't make sense on static method '%s' in a blueprint function library"), *InKey.ToString(), *FuncDef->GetName());
+					FUHTException::Throwf(*this, TEXT("%s doesn't make sense on static method '%s' in a blueprint function library"), *InKey.ToString(), *FuncDef->GetName());
 				}
 			}
 		}
@@ -450,7 +454,7 @@ void FUnrealTypeDefinitionInfo::ValidateMetaDataFormat(const FName InKey, ECheck
 
 			if (!bGoodParams)
 			{
-				UE_LOG_ERROR_UHT(TEXT("Commutative associative binary operators must have exactly 2 parameters of the same type and a return value."));
+				UE_LOG_ERROR_UHT(*FuncDef, TEXT("Commutative associative binary operators must have exactly 2 parameters of the same type and a return value."));
 			}
 		}
 	}
@@ -481,7 +485,7 @@ void FUnrealTypeDefinitionInfo::ValidateMetaDataFormat(const FName InKey, ECheck
 					FUnrealPropertyDefinitionInfo* FoundFieldDef = FHeaderParser::FindProperty(*FuncDef, *Entry, false);
 					if (!FoundFieldDef)
 					{
-						UE_LOG_ERROR_UHT(TEXT("Function does not have a parameter named '%s'"), *Entry);
+						UE_LOG_ERROR_UHT(*FuncDef, TEXT("Function does not have a parameter named '%s'"), *Entry);
 					}
 					else
 					{
@@ -495,7 +499,7 @@ void FUnrealTypeDefinitionInfo::ValidateMetaDataFormat(const FName InKey, ECheck
 							}
 							else
 							{
-								UE_LOG_ERROR_UHT(TEXT("Function already specified an ExpandEnumAsExec input (%s), but '%s' is also an input parameter. Only one is permitted."), *FirstInputDef->GetName(), *Entry);
+								UE_LOG_ERROR_UHT(*FuncDef, TEXT("Function already specified an ExpandEnumAsExec input (%s), but '%s' is also an input parameter. Only one is permitted."), *FirstInputDef->GetName(), *Entry);
 							}
 						}
 					}
@@ -511,7 +515,7 @@ void FUnrealTypeDefinitionInfo::ValidateMetaDataFormat(const FName InKey, ECheck
 		const FString ExperimentalValue(TEXT("Experimental"));
 		if ((InValue != EarlyAccessValue) && (InValue != ExperimentalValue))
 		{
-			FError::Throwf(TEXT("'%s' metadata was '%s' but it must be %s or %s"), *InKey.ToString(), *InValue, *ExperimentalValue, *EarlyAccessValue);
+			FUHTException::Throwf(*this, TEXT("'%s' metadata was '%s' but it must be %s or %s"), *InKey.ToString(), *InValue, *ExperimentalValue, *EarlyAccessValue);
 		}
 	}
 	break;
@@ -523,13 +527,13 @@ void FUnrealTypeDefinitionInfo::ValidateMetaDataFormat(const FName InKey, ECheck
 		{
 			if (!PropDef->IsNumericOrNumericStaticArray() && !PropDef->IsStructOrStructStaticArray())
 			{
-				FError::Throwf(TEXT("'Units' meta data can only be applied to numeric and struct properties"));
+				FUHTException::Throwf(*this, TEXT("'Units' meta data can only be applied to numeric and struct properties"));
 			}
 		}
 
 		if (!FUnitConversion::UnitFromString(*InValue))
 		{
-			FError::Throwf(TEXT("Unrecognized units (%s) specified for property '%s'"), *InValue, *GetFullName());
+			FUHTException::Throwf(*this, TEXT("Unrecognized units (%s) specified for property '%s'"), *InValue, *GetFullName());
 		}
 	}
 	break;
@@ -539,7 +543,7 @@ void FUnrealTypeDefinitionInfo::ValidateMetaDataFormat(const FName InKey, ECheck
 		const TCHAR* StrictValue = TEXT("Strict");
 		if (InValue != StrictValue)
 		{
-			FError::Throwf(TEXT("'%s' metadata was '%s' but it must be %s"), *InKey.ToString(), *InValue, *StrictValue);
+			FUHTException::Throwf(*this, TEXT("'%s' metadata was '%s' but it must be %s"), *InKey.ToString(), *InValue, *StrictValue);
 		}
 	}
 	break;
@@ -1359,7 +1363,7 @@ FUnrealClassDefinitionInfo* FUnrealClassDefinitionInfo::FindClass(const TCHAR* C
 	return ClassDef;
 }
 
-FUnrealClassDefinitionInfo* FUnrealClassDefinitionInfo::FindScriptClassOrThrow(const FString& InClassName)
+FUnrealClassDefinitionInfo* FUnrealClassDefinitionInfo::FindScriptClassOrThrow(const FHeaderParser& Parser, const FString& InClassName)
 {
 	FString ErrorMsg;
 	if (FUnrealClassDefinitionInfo* ResultDef = FindScriptClass(InClassName, &ErrorMsg))
@@ -1367,9 +1371,9 @@ FUnrealClassDefinitionInfo* FUnrealClassDefinitionInfo::FindScriptClassOrThrow(c
 		return ResultDef;
 	}
 
-	FError::Throwf(*ErrorMsg);
+	FUHTException::Throwf(Parser, MoveTemp(ErrorMsg));
 
-	// Unreachable, but compiler will warn otherwise because FError::Throwf isn't declared noreturn
+	// Unreachable, but compiler will warn otherwise because FUHTException::Throwf isn't declared noreturn
 	return nullptr;
 }
 
@@ -1438,7 +1442,7 @@ void FUnrealClassDefinitionInfo::PostParseFinalizeInternal()
 	{
 		FString UName = GetNameCPP();
 		FString IName = TEXT("I") + UName.RightChop(1);
-		FError::Throwf(TEXT("UInterface '%s' parsed without a corresponding '%s'"), *UName, *IName);
+		FUHTException::Throwf(*this, TEXT("UInterface '%s' parsed without a corresponding '%s'"), *UName, *IName);
 	}
 
 	FHeaderParser::CheckSparseClassData(*this);
@@ -1500,7 +1504,7 @@ void FUnrealClassDefinitionInfo::ParseClassProperties(TArray<FPropertySpecifier>
 
 		case EClassMetadataSpecifier::Within:
 
-			ClassWithinStr = FHeaderParser::RequireExactlyOneSpecifierValue(PropSpecifier);
+			ClassWithinStr = FHeaderParser::RequireExactlyOneSpecifierValue(*this, PropSpecifier);
 			break;
 
 		case EClassMetadataSpecifier::EditInlineNew:
@@ -1541,7 +1545,7 @@ void FUnrealClassDefinitionInfo::ParseClassProperties(TArray<FPropertySpecifier>
 
 		case EClassMetadataSpecifier::DependsOn:
 
-			FError::Throwf(TEXT("The dependsOn specifier is deprecated. Please use #include \"ClassHeaderFilename.h\" instead."));
+			FUHTException::Throwf(*this, TEXT("The dependsOn specifier is deprecated. Please use #include \"ClassHeaderFilename.h\" instead."));
 			break;
 
 		case EClassMetadataSpecifier::MinimalAPI:
@@ -1600,7 +1604,7 @@ void FUnrealClassDefinitionInfo::ParseClassProperties(TArray<FPropertySpecifier>
 		case EClassMetadataSpecifier::Config:
 
 			// Class containing config properties - parse the name of the config file to use
-			ConfigName = FHeaderParser::RequireExactlyOneSpecifierValue(PropSpecifier);
+			ConfigName = FHeaderParser::RequireExactlyOneSpecifierValue(*this, PropSpecifier);
 			break;
 
 		case EClassMetadataSpecifier::DefaultConfig:
@@ -1623,12 +1627,12 @@ void FUnrealClassDefinitionInfo::ParseClassProperties(TArray<FPropertySpecifier>
 
 		case EClassMetadataSpecifier::EditorConfig:
 			// Save EditorConfig properties to the given JSON file.
-			MetaData.Add(NAME_EditorConfig, FHeaderParser::RequireExactlyOneSpecifierValue(PropSpecifier));
+			MetaData.Add(NAME_EditorConfig, FHeaderParser::RequireExactlyOneSpecifierValue(*this, PropSpecifier));
 			break;
 
 		case EClassMetadataSpecifier::ShowCategories:
 
-			FHeaderParser::RequireSpecifierValue(PropSpecifier);
+			FHeaderParser::RequireSpecifierValue(*this, PropSpecifier);
 
 			for (FString& Value : PropSpecifier.Values)
 			{
@@ -1638,7 +1642,7 @@ void FUnrealClassDefinitionInfo::ParseClassProperties(TArray<FPropertySpecifier>
 
 		case EClassMetadataSpecifier::HideCategories:
 
-			FHeaderParser::RequireSpecifierValue(PropSpecifier);
+			FHeaderParser::RequireSpecifierValue(*this, PropSpecifier);
 
 			for (FString& Value : PropSpecifier.Values)
 			{
@@ -1648,7 +1652,7 @@ void FUnrealClassDefinitionInfo::ParseClassProperties(TArray<FPropertySpecifier>
 
 		case EClassMetadataSpecifier::ShowFunctions:
 
-			FHeaderParser::RequireSpecifierValue(PropSpecifier);
+			FHeaderParser::RequireSpecifierValue(*this, PropSpecifier);
 
 			for (const FString& Value : PropSpecifier.Values)
 			{
@@ -1658,7 +1662,7 @@ void FUnrealClassDefinitionInfo::ParseClassProperties(TArray<FPropertySpecifier>
 
 		case EClassMetadataSpecifier::HideFunctions:
 
-			FHeaderParser::RequireSpecifierValue(PropSpecifier);
+			FHeaderParser::RequireSpecifierValue(*this, PropSpecifier);
 
 			for (FString& Value : PropSpecifier.Values)
 			{
@@ -1669,12 +1673,12 @@ void FUnrealClassDefinitionInfo::ParseClassProperties(TArray<FPropertySpecifier>
 			// Currently some code only handles a single sidecar data structure so we enforce that here
 		case EClassMetadataSpecifier::SparseClassDataTypes:
 
-			SparseClassDataTypes.AddUnique(FHeaderParser::RequireExactlyOneSpecifierValue(PropSpecifier));
+			SparseClassDataTypes.AddUnique(FHeaderParser::RequireExactlyOneSpecifierValue(*this, PropSpecifier));
 			break;
 
 		case EClassMetadataSpecifier::ClassGroup:
 
-			FHeaderParser::RequireSpecifierValue(PropSpecifier);
+			FHeaderParser::RequireSpecifierValue(*this, PropSpecifier);
 
 			for (FString& Value : PropSpecifier.Values)
 			{
@@ -1684,7 +1688,7 @@ void FUnrealClassDefinitionInfo::ParseClassProperties(TArray<FPropertySpecifier>
 
 		case EClassMetadataSpecifier::AutoExpandCategories:
 
-			FHeaderParser::RequireSpecifierValue(PropSpecifier);
+			FHeaderParser::RequireSpecifierValue(*this, PropSpecifier);
 
 			for (FString& Value : PropSpecifier.Values)
 			{
@@ -1695,7 +1699,7 @@ void FUnrealClassDefinitionInfo::ParseClassProperties(TArray<FPropertySpecifier>
 
 		case EClassMetadataSpecifier::AutoCollapseCategories:
 
-			FHeaderParser::RequireSpecifierValue(PropSpecifier);
+			FHeaderParser::RequireSpecifierValue(*this, PropSpecifier);
 
 			for (FString& Value : PropSpecifier.Values)
 			{
@@ -1706,7 +1710,7 @@ void FUnrealClassDefinitionInfo::ParseClassProperties(TArray<FPropertySpecifier>
 
 		case EClassMetadataSpecifier::DontAutoCollapseCategories:
 
-			FHeaderParser::RequireSpecifierValue(PropSpecifier);
+			FHeaderParser::RequireSpecifierValue(*this, PropSpecifier);
 
 			for (const FString& Value : PropSpecifier.Values)
 			{
@@ -1743,7 +1747,7 @@ void FUnrealClassDefinitionInfo::ParseClassProperties(TArray<FPropertySpecifier>
 			break;
 
 		default:
-			FError::Throwf(TEXT("Unknown class specifier '%s'"), *PropSpecifier.Key);
+			FUHTException::Throwf(*this, TEXT("Unknown class specifier '%s'"), *PropSpecifier.Key);
 		}
 	}
 }
@@ -1845,7 +1849,7 @@ void FUnrealClassDefinitionInfo::MergeAndValidateClassFlags(const FString& Decla
 	{
 		if (!(Class->ClassFlags & CLASS_NotPlaceable))
 		{
-			FError::Throwf(TEXT("The 'placeable' specifier is only allowed on classes which have a base class that's marked as not placeable. Classes are assumed to be placeable by default."));
+			FUHTException::Throwf(*this, TEXT("The 'placeable' specifier is only allowed on classes which have a base class that's marked as not placeable. Classes are assumed to be placeable by default."));
 		}
 		Class->ClassFlags &= ~CLASS_NotPlaceable;
 		bWantsToBePlaceable = false; // Reset this flag after it's been merged
@@ -1863,21 +1867,21 @@ void FUnrealClassDefinitionInfo::MergeAndValidateClassFlags(const FString& Decla
 		// don't allow actor classes to be declared editinlinenew
 		if (IsActorClass(Class))
 		{
-			FError::Throwf(TEXT("Invalid class attribute: Creating actor instances via the property window is not allowed"));
+			FUHTException::Throwf(*this, TEXT("Invalid class attribute: Creating actor instances via the property window is not allowed"));
 		}
 	}
 
 	// Make sure both RequiredAPI and MinimalAPI aren't specified
 	if (Class->HasAllClassFlags(CLASS_MinimalAPI | CLASS_RequiredAPI))
 	{
-		FError::Throwf(TEXT("MinimalAPI cannot be specified when the class is fully exported using a MODULENAME_API macro"));
+		FUHTException::Throwf(*this, TEXT("MinimalAPI cannot be specified when the class is fully exported using a MODULENAME_API macro"));
 	}
 
 	// All classes must start with a valid Unreal prefix
 	const FString ExpectedClassName = GetNameWithPrefix();
 	if (DeclaredClassName != ExpectedClassName)
 	{
-		FError::Throwf(TEXT("Class name '%s' is invalid, should be identified as '%s'"), *DeclaredClassName, *ExpectedClassName);
+		FUHTException::Throwf(*this, TEXT("Class name '%s' is invalid, should be identified as '%s'"), *DeclaredClassName, *ExpectedClassName);
 	}
 
 	if ((Class->ClassFlags & CLASS_NoExport))
@@ -1888,7 +1892,7 @@ void FUnrealClassDefinitionInfo::MergeAndValidateClassFlags(const FString& Decla
 		if (!(Class->ClassFlags & CLASS_Intrinsic) && (PreviousClassFlags & CLASS_NoExport) == 0 &&
 			(PreviousClassFlags & CLASS_Native) != 0)	// a new native class (one that hasn't been compiled into C++ yet) won't have this set
 		{
-			FError::Throwf(TEXT("'noexport': Must include CLASS_NoExport in native class declaration"));
+			FUHTException::Throwf(*this, TEXT("'noexport': Must include CLASS_NoExport in native class declaration"));
 		}
 	}
 
@@ -1896,12 +1900,12 @@ void FUnrealClassDefinitionInfo::MergeAndValidateClassFlags(const FString& Decla
 	{
 		if (Class->HasAnyClassFlags(CLASS_NoExport))
 		{
-			FError::Throwf(TEXT("'abstract': NoExport class missing abstract keyword from class declaration (must change C++ version first)"));
+			FUHTException::Throwf(*this, TEXT("'abstract': NoExport class missing abstract keyword from class declaration (must change C++ version first)"));
 			Class->ClassFlags |= CLASS_Abstract;
 		}
 		else if (IsNative())
 		{
-			FError::Throwf(TEXT("'abstract': missing abstract keyword from class declaration - class will no longer be exported as abstract"));
+			FUHTException::Throwf(*this, TEXT("'abstract': missing abstract keyword from class declaration - class will no longer be exported as abstract"));
 		}
 	}
 }
@@ -1918,12 +1922,12 @@ void FUnrealClassDefinitionInfo::SetAndValidateConfigName()
 			UClass* SuperClass = Class->GetSuperClass();
 			if (!SuperClass)
 			{
-				FError::Throwf(TEXT("Cannot inherit config filename: %s has no super class"), *Class->GetName());
+				FUHTException::Throwf(*this, TEXT("Cannot inherit config filename: %s has no super class"), *Class->GetName());
 			}
 
 			if (SuperClass->ClassConfigName == NAME_None)
 			{
-				FError::Throwf(TEXT("Cannot inherit config filename: parent class %s is not marked config."), *SuperClass->GetPathName());
+				FUHTException::Throwf(*this, TEXT("Cannot inherit config filename: parent class %s is not marked config."), *SuperClass->GetPathName());
 			}
 		}
 		else
@@ -1949,12 +1953,12 @@ void FUnrealClassDefinitionInfo::SetAndValidateWithinClass()
 		FUnrealClassDefinitionInfo* RequiredWithinClassDef = FUnrealClassDefinitionInfo::FindClass(*ClassWithinStr);
 		if (!RequiredWithinClassDef)
 		{
-			FError::Throwf(TEXT("Within class '%s' not found."), *ClassWithinStr);
+			FUHTException::Throwf(*this, TEXT("Within class '%s' not found."), *ClassWithinStr);
 		}
 		UClass* RequiredWithinClass = RequiredWithinClassDef->GetClass();
 		if (RequiredWithinClass->IsChildOf(UInterface::StaticClass()))
 		{
-			FError::Throwf(TEXT("Classes cannot be 'within' interfaces"));
+			FUHTException::Throwf(*this, TEXT("Classes cannot be 'within' interfaces"));
 		}
 		else if (Class->ClassWithin == NULL || Class->ClassWithin == UObject::StaticClass() || RequiredWithinClass->IsChildOf(Class->ClassWithin))
 		{
@@ -1962,7 +1966,7 @@ void FUnrealClassDefinitionInfo::SetAndValidateWithinClass()
 		}
 		else if (Class->ClassWithin != RequiredWithinClass)
 		{
-			FError::Throwf(TEXT("%s must be within %s, not %s"), *Class->GetPathName(), *Class->ClassWithin->GetPathName(), *RequiredWithinClass->GetPathName());
+			FUHTException::Throwf(*this, TEXT("%s must be within %s, not %s"), *Class->GetPathName(), *Class->ClassWithin->GetPathName(), *RequiredWithinClass->GetPathName());
 		}
 	}
 	else
@@ -1979,7 +1983,7 @@ void FUnrealClassDefinitionInfo::SetAndValidateWithinClass()
 
 	if (!Class->ClassWithin->IsChildOf(ExpectedWithin))
 	{
-		FError::Throwf(TEXT("Parent class declared within '%s'.  Cannot override within class with '%s' since it isn't a child"), *ExpectedWithin->GetName(), *Class->ClassWithin->GetName());
+		FUHTException::Throwf(*this, TEXT("Parent class declared within '%s'.  Cannot override within class with '%s' since it isn't a child"), *ExpectedWithin->GetName(), *Class->ClassWithin->GetName());
 	}
 
 	ClassWithin = &GTypeDefinitionInfoMap.FindChecked<FUnrealClassDefinitionInfo>(Class->ClassWithin);
diff --git a/Engine/Source/Programs/UnrealHeaderTool/Private/UnrealTypeDefinitionInfo.h b/Engine/Source/Programs/UnrealHeaderTool/Private/UnrealTypeDefinitionInfo.h
--- a/Engine/Source/Programs/UnrealHeaderTool/Private/UnrealTypeDefinitionInfo.h
+++ b/Engine/Source/Programs/UnrealHeaderTool/Private/UnrealTypeDefinitionInfo.h
@@ -13,6 +13,7 @@
 #define UHT_ENABLE_ENGINE_TYPE_CHECKS 1
 
 // Forward declarations.
+class FHeaderParser;
 class FScope;
 class FUnrealSourceFile;
 struct FManifestModule;
@@ -90,12 +91,12 @@ class FUHTMetaData
 {
 public:
 
-	static void RemapMetaData(TMap<FName, FString>& MetaData);
+	static void RemapMetaData(FUnrealTypeDefinitionInfo& TypeDef, TMap<FName, FString>& MetaData);
 
 	template <typename Def>
 	static void RemapAndAddMetaData(Def& TypeDef, TMap<FName, FString>&& InMetaData)
 	{
-		RemapMetaData(InMetaData);
+		RemapMetaData(TypeDef, InMetaData);
 		TypeDef.ValidateMetaDataFormat(InMetaData);
 		TypeDef.AddMetaData(MoveTemp(InMetaData));
 	}
@@ -103,7 +104,7 @@ public:
 	template <typename Def>
 	static void RemapAndAddMetaData(Def& TypeDef, TMap<FName, FString>& InMetaData)
 	{
-		RemapMetaData(InMetaData);
+		RemapMetaData(TypeDef, InMetaData);
 		TypeDef.ValidateMetaDataFormat(InMetaData);
 		TypeDef.AddMetaData(MoveTemp(InMetaData));
 	}
@@ -204,7 +205,9 @@ private:
 /**
  * Class that stores information about type definitions.
  */
-class FUnrealTypeDefinitionInfo : public TSharedFromThis<FUnrealTypeDefinitionInfo>
+class FUnrealTypeDefinitionInfo 
+	: public TSharedFromThis<FUnrealTypeDefinitionInfo>
+	, public FUHTExceptionContext
 {
 private:
 	enum class EFinalizeState : uint8
@@ -429,12 +432,17 @@ public:
 	/**
 	 * Gets the line number in source file this type was defined in.
 	 */
-	int32 GetLineNumber() const
+	virtual int32 GetLineNumber() const override
 	{
 		check(HasSource());
 		return LineNumber;
 	}
 
+	/** 
+	 * Returns the filename from the source file
+	 */
+	virtual FString GetFilename() const override;
+
 	/**
 	 * Sets the input line in the rare case where the definition is created before fully parsed (sparse delegates)
 	 */
@@ -1915,14 +1923,6 @@ public:
 
 	struct FDefinitionRange
 	{
-		void Validate()
-		{
-			if (End <= Start)
-			{
-				FError::Throwf(TEXT("The class definition range is invalid. Most probably caused by previous parsing error."));
-			}
-		}
-
 		const TCHAR* Start = nullptr;
 		const TCHAR* End = nullptr;
 	};
@@ -2158,6 +2158,17 @@ public:
 		bHasGeneratedBody = true;
 	}
 
+	/**
+	 * Verify we have a valid definition range
+	 */
+	void ValidateDefinitionRange()
+	{
+		if (DefinitionRange.End <= DefinitionRange.Start)
+		{
+			FUHTException::Throwf(*this, TEXT("The class definition range is invalid. Most probably caused by previous parsing error."));
+		}
+	}
+
 	/**
 	 * Return the definition range of the structure
 	 */
@@ -2557,18 +2568,19 @@ public:
 	 * the prefix of the given name while searching. Throws an exception with the script error
 	 * if the class could not be found.
 	 *
-	 * @param   InClassName  Name w/ Unreal prefix to use when searching for a class
-	 * @return               The found class.
+	 * @param Parser Calling context for exceptions
+	 * @param InClassName Name w/ Unreal prefix to use when searching for a class
+	 * @return The found class.
 	 */
-	static FUnrealClassDefinitionInfo* FindScriptClassOrThrow(const FString& InClassName);
+	static FUnrealClassDefinitionInfo* FindScriptClassOrThrow(const FHeaderParser& Parser, const FString& InClassName);
 
 	/**
 	 * Attempts to find a script class based on the given name. Will attempt to strip
 	 * the prefix of the given name while searching. Optionally returns script errors when appropriate.
 	 *
-	 * @param   InClassName  Name w/ Unreal prefix to use when searching for a class
-	 * @param   OutErrorMsg  Error message (if any) giving the caller flexibility in how they present an error
-	 * @return               The found class, or NULL if the class was not found.
+	 * @param InClassName  Name w/ Unreal prefix to use when searching for a class
+	 * @param OutErrorMsg  Error message (if any) giving the caller flexibility in how they present an error
+	 * @return The found class, or NULL if the class was not found.
 	 */
 	static FUnrealClassDefinitionInfo* FindScriptClass(const FString& InClassName, FString* OutErrorMsg = nullptr);
 
