Enumerations are constructed post parse

#rnx

diff --git a/Engine/Source/Programs/UnrealHeaderTool/Private/ClassMaps.h b/Engine/Source/Programs/UnrealHeaderTool/Private/ClassMaps.h
--- a/Engine/Source/Programs/UnrealHeaderTool/Private/ClassMaps.h
+++ b/Engine/Source/Programs/UnrealHeaderTool/Private/ClassMaps.h
@@ -39,15 +39,33 @@ struct FTypeDefinitionInfoMap : public FFreezableContainer
 	void Add(UObject* Object, TSharedRef<FUnrealTypeDefinitionInfo>&& Definition)
 	{
 		check(!bFrozen);
-		DefinitionsByUObject.Add(Object, Definition);
-		DefinitionsByName.Add(Object->GetFName(), MoveTemp(Definition));
+		AddNameLookup(UHTCastChecked<FUnrealObjectDefinitionInfo>(Definition));
+		AddObjectLookup(Object, MoveTemp(Definition));
 	}
 
 	void Add(UObject* Object, TSharedRef<FUnrealTypeDefinitionInfo>& Definition)
 	{
 		check(!bFrozen);
+		AddNameLookup(UHTCastChecked<FUnrealObjectDefinitionInfo>(Definition));
+		AddObjectLookup(Object, Definition);
+	}
+
+	void AddObjectLookup(UObject* Object, TSharedRef<FUnrealTypeDefinitionInfo>&& Definition)
+	{
+		//check(!bFrozen);
+		DefinitionsByUObject.Add(Object, Definition);
+	}
+
+	void AddObjectLookup(UObject* Object, TSharedRef<FUnrealTypeDefinitionInfo>& Definition)
+	{
+		//check(!bFrozen);
 		DefinitionsByUObject.Add(Object, Definition);
-		DefinitionsByName.Add(Object->GetFName(), Definition);
+	}
+
+	void AddNameLookup(FUnrealObjectDefinitionInfo& Definition)
+	{
+		check(!bFrozen);
+		DefinitionsByName.Add(Definition.GetFName(), Definition.AsShared());
 	}
 
 	bool Contains(const UObject* Object)
diff --git a/Engine/Source/Programs/UnrealHeaderTool/Private/CodeGenerator.cpp b/Engine/Source/Programs/UnrealHeaderTool/Private/CodeGenerator.cpp
--- a/Engine/Source/Programs/UnrealHeaderTool/Private/CodeGenerator.cpp
+++ b/Engine/Source/Programs/UnrealHeaderTool/Private/CodeGenerator.cpp
@@ -77,7 +77,7 @@ static bool bWriteContents = false;
 static bool bVerifyContents = false;
 
 UClass* ProcessParsedClass(FUnrealClassDefinitionInfo& ClassDef);
-UEnum* ProcessParsedEnum(FUnrealEnumDefinitionInfo& EnumDef);
+void ProcessParsedEnum(FUnrealEnumDefinitionInfo& EnumDef);
 UScriptStruct* ProcessParsedStruct(FUnrealScriptStructDefinitionInfo& ScriptStructDef);
 
 // Array of all the temporary header async file tasks so we can ensure they have completed before issuing our timings
@@ -89,7 +89,6 @@ bool HasIdentifierExactMatch(const TCHAR* StringBegin, const TCHAR* StringEnd, c
 FUnrealClassDefinitionInfo* GUObjectDef = nullptr;
 FUnrealClassDefinitionInfo* GUClassDef = nullptr;
 FUnrealClassDefinitionInfo* GUInterfaceDef = nullptr;
-FUnrealClassDefinitionInfo* GUPackageDef = nullptr;
 
 namespace
 {
@@ -1565,7 +1564,7 @@ void FNativeClassHeaderGenerator::OutputProperty(FOutputDevice& DeclOut, FOutput
 
 	case EUHTPropertyType::Enum:
 	{
-		if (PropertyBase.EnumDef->GetEnum()->GetCppForm() != UEnum::ECppForm::EnumClass)
+		if (PropertyBase.EnumDef->GetCppForm() != UEnum::ECppForm::EnumClass)
 		{
 			OutputByteProperty();
 		}
@@ -6451,9 +6450,8 @@ void DefineTypes(TArray<FUnrealPackageDefinitionInfo*>& PackageDefs)
 				for (TSharedRef<FUnrealTypeDefinitionInfo>& TypeDef : SourceFile.GetDefinedEnums())
 				{
 					FUnrealEnumDefinitionInfo& EnumDef = TypeDef->AsEnumChecked();
-					UEnum* ResultEnum = ProcessParsedEnum(EnumDef);
-					EnumDef.SetObject(ResultEnum);
-					GTypeDefinitionInfoMap.Add(ResultEnum, TypeDef);
+					ProcessParsedEnum(EnumDef);
+					GTypeDefinitionInfoMap.AddNameLookup(UHTCastChecked<FUnrealObjectDefinitionInfo>(TypeDef));
 				}
 
 				for (TSharedRef<FUnrealTypeDefinitionInfo>& TypeDef : SourceFile.GetDefinedStructs())
@@ -6485,7 +6483,6 @@ void ResolveParents(TArray<FUnrealPackageDefinitionInfo*>& PackageDefs)
 	GUObjectDef = &GTypeDefinitionInfoMap.FindChecked<FUnrealClassDefinitionInfo>(UObject::StaticClass());
 	GUClassDef = &GTypeDefinitionInfoMap.FindChecked<FUnrealClassDefinitionInfo>(UClass::StaticClass());
 	GUInterfaceDef = &GTypeDefinitionInfoMap.FindChecked<FUnrealClassDefinitionInfo>(UInterface::StaticClass());
-	GUPackageDef = &GTypeDefinitionInfoMap.FindChecked<FUnrealClassDefinitionInfo>(UPackage::StaticClass());
 
 	for (FUnrealPackageDefinitionInfo* PackageDef : PackageDefs)
 	{
@@ -6580,13 +6577,98 @@ void TopologicalSort(TArray<FUnrealSourceFile*>& OrderedSourceFiles)
 
 void ParseSourceFiles(TArray<FUnrealSourceFile*>& OrderedSourceFiles)
 {
+	// Disable loading of objects outside of this package (or more exactly, objects which aren't UFields, CDO, or templates)
+	TGuardValue<bool> AutoRestoreVerifyObjectRefsFlag(GVerifyObjectReferencesOnly, true);
+
+// Concurrent parsing disabled
+#if 0 
+	/**
+	 * For every FUnrealSourceFile being processed, an instance of this class represents the data associated with generating the new output.
+	 */
+	struct FParseCPP
+	{
+		FParseCPP(FUnrealPackageDefinitionInfo& InPackageDef, FUnrealSourceFile& InSourceFile)
+			: PackageDef(InPackageDef)
+			, SourceFile(InSourceFile)
+		{}
+
+		/**
+		 * The package definition being exported
+		 */
+		FUnrealPackageDefinitionInfo& PackageDef;
+
+		/**
+		 * The source file being exported
+		 */
+		FUnrealSourceFile& SourceFile;
+
+		/**
+		 * This task represents the task that parses the source
+		 */
+		FGraphEventRef ParseTaskRef;
+	};
+
+	TArray<FParseCPP> ParsedCPPs;
+	ParsedCPPs.Reserve(OrderedSourceFiles.Num());
+	for (FUnrealSourceFile* SourceFile : OrderedSourceFiles)
+	{
+		ParsedCPPs.Emplace(SourceFile->GetPackageDef(), *SourceFile);
+	}
+
+	TSet<FUnrealSourceFile*> Includes;
+	Includes.Reserve(ParsedCPPs.Num());
+	FGraphEventArray TempTasks;
+	TempTasks.Reserve(ParsedCPPs.Num());
+	FGraphEventArray ParsedSourceTasks;
+	ParsedSourceTasks.Reserve(ParsedCPPs.Num());
+	for (FParseCPP& ParsedCPP : ParsedCPPs)
+	{
+		const FManifestModule& Module = ParsedCPP.PackageDef.GetModule();
+		FUnrealSourceFile& SourceFile = ParsedCPP.SourceFile;
+
+		FString ModuleRelativeFilename = SourceFile.GetFilename();
+		ConvertToBuildIncludePath(Module, ModuleRelativeFilename);
+
+		auto ParseSource = [&ParsedCPP]()
+		{
+			FResults::TryAlways([&ParsedCPP]()
+			{
+					FHeaderParser::Parse(ParsedCPP.PackageDef, ParsedCPP.SourceFile);
+			});
+		};
+
+		Includes.Reset();
+		for (FHeaderProvider& Header : SourceFile.GetIncludes())
+		{
+			if (FUnrealSourceFile* Include = Header.Resolve(SourceFile))
+			{
+				Includes.Add(Include);
+			}
+		}
+
+		// Our generation must wait on all of our includes generation to complete
+		TempTasks.Reset();
+		for (FUnrealSourceFile* Include : Includes)
+		{
+			FParseCPP& IncludeCPP = ParsedCPPs[Include->GetOrderedIndex()];
+			TempTasks.Add(IncludeCPP.ParseTaskRef);
+
+		}
+		ParsedCPP.ParseTaskRef = FFunctionGraphTask::CreateAndDispatchWhenReady(MoveTemp(ParseSource), TStatId(), &TempTasks);
+		ParsedSourceTasks.Add(ParsedCPP.ParseTaskRef);
+	}
+
+	// Wait for the results
+	FTaskGraphInterface::Get().WaitUntilTasksComplete(ParsedSourceTasks);
+	FResults::WaitForErrorTasks();
+#else
 	for (FUnrealSourceFile* SourceFile : OrderedSourceFiles)
 	{
 		FUnrealPackageDefinitionInfo& PackageDef = SourceFile->GetPackageDef();
 		FScopedDurationTimer SourceTimer(SourceFile->GetTime(ESourceFileTime::Parse));
-
 		FResults::TryAlways([&PackageDef, SourceFile]() { FHeaderParser::Parse(PackageDef, *SourceFile); });
 	}
+#endif
 }
 
 void PrepareTypesForExport(TArray<FUnrealPackageDefinitionInfo*>& PackageDefs)
@@ -6946,14 +7028,14 @@ UClass* ProcessParsedClass(FUnrealClassDefinitionInfo& ClassDef)
 	return ResultClass;
 }
 
-UEnum* ProcessParsedEnum(FUnrealEnumDefinitionInfo& EnumDef)
+void ProcessParsedEnum(FUnrealEnumDefinitionInfo& EnumDef)
 {
 	UPackage* Package = EnumDef.GetPackageDef().GetPackage();
 	const FString& EnumName = EnumDef.GetNameCPP();
 
-	if (UEnum* Existing = FEngineAPI::FindObject<UEnum>(ANY_PACKAGE, *EnumName))
+	if (TSharedRef<FUnrealTypeDefinitionInfo>* Existing = GTypeDefinitionInfoMap.FindByName(*EnumName))
 	{
-		FUHTException::Throwf(EnumDef, TEXT("Duplicate enum name: %s also exists in file %s"), *EnumName, *Existing->GetOutermost()->GetName());
+		FUHTException::Throwf(EnumDef, TEXT("Duplicate enum name: %s also exists in file %s"), *EnumName, *(*Existing)->GetFilename());
 	}
 
 	// Check if the enum name is using a reserved keyword
@@ -6961,10 +7043,6 @@ UEnum* ProcessParsedEnum(FUnrealEnumDefinitionInfo& EnumDef)
 	{
 		FUHTException::Throwf(EnumDef, TEXT("enum: '%s' uses a reserved type name."), *EnumName);
 	}
-
-	// Create enum definition.
-	UEnum* Enum = new(EC_InternalUseOnlyConstructor, Package, FName(EnumName), RF_Public) UEnum(FObjectInitializer());
-	return Enum;
 }
 
 UScriptStruct* ProcessParsedStruct(FUnrealScriptStructDefinitionInfo& ScriptStructDef)
diff --git a/Engine/Source/Programs/UnrealHeaderTool/Private/HeaderParser.cpp b/Engine/Source/Programs/UnrealHeaderTool/Private/HeaderParser.cpp
--- a/Engine/Source/Programs/UnrealHeaderTool/Private/HeaderParser.cpp
+++ b/Engine/Source/Programs/UnrealHeaderTool/Private/HeaderParser.cpp
@@ -61,7 +61,6 @@ namespace
 	static const FName NAME_ReturnValue(TEXT("ReturnValue"));
 	static const FName NAME_CppFromBpEvent(TEXT("CppFromBpEvent"));
 	static const FName NAME_CustomThunk(TEXT("CustomThunk"));
-	static const FName NAME_ArraySizeEnum(TEXT("ArraySizeEnum"));
 	static const FName NAME_EditInline(TEXT("EditInline"));
 	static const FName NAME_IncludePath(TEXT("IncludePath"));
 	static const FName NAME_ModuleRelativePath(TEXT("ModuleRelativePath"));
@@ -1265,9 +1264,11 @@ FUnrealEnumDefinitionInfo& FHeaderParser::CompileEnum()
 	}
 
 	// Register the list of enum names.
-	if (!EnumDef.SetEnums(EnumNames, CppForm, Flags, false))
+	if (!EnumDef.SetEnums(EnumNames, CppForm, Flags))
 	{
-		const FName MaxEnumItem      = *(EnumDef.GenerateEnumPrefix() + TEXT("_MAX"));
+		//@TODO - Two issues with this, first the SetEnums always returns true.  That was the case even with the original code.
+		// Second issue is that the code to generate the max string doesn't work in many cases (i.e. EPixelFormat)
+		const FName MaxEnumItem = *(EnumDef.GenerateEnumPrefix() + TEXT("_MAX"));
 		const int32 MaxEnumItemIndex = EnumDef.GetIndexByName(MaxEnumItem);
 		if (MaxEnumItemIndex != INDEX_NONE)
 		{
@@ -1279,20 +1280,16 @@ FUnrealEnumDefinitionInfo& FHeaderParser::CompileEnum()
 
 	CheckDocumentationPolicyForEnum(EnumDef, EnumValueMetaData, EntryMetaData);
 
-	// Add the metadata gathered for the enum to the package
-	if (EnumValueMetaData.Num() > 0)
-	{
-		UMetaData* PackageMetaData = EnumDef.GetPackageDef().GetPackage()->GetMetaData();
-		checkSlow(PackageMetaData);
-
-		PackageMetaData->SetObjectValues(EnumDef.GetEnum(), MoveTemp(EnumValueMetaData));
-	}
-
 	if (!EnumDef.IsValidEnumValue(0) && EnumMetaData.Contains(FHeaderParserNames::NAME_BlueprintType))
 	{
 		UE_LOG_WARNING_UHT(EnumDef, TEXT("'%s' does not have a 0 entry! (This is a problem when the enum is initalized by default)"), *EnumDef.GetName());
 	}
 
+	// Add the metadata gathered for the enum to the package
+	if (EnumValueMetaData.Num() > 0)
+	{
+		EnumDef.AddMetaData(MoveTemp(EnumValueMetaData));
+	}
 	return EnumDef;
 }
 
@@ -2018,7 +2015,7 @@ FUnrealScriptStructDefinitionInfo& FHeaderParser::CompileStructDeclaration()
 
 					if ( ch==0 )
 					{
-						FUHTException::Throwf(*this, TEXT("Unexpected end of struct definition %s"), *Struct->GetName());
+						FUHTException::Throwf(*this, TEXT("Unexpected end of struct definition %s"), *StructDef.GetName());
 					}
 					else if ( ch=='{' || (ch=='#' && (PeekIdentifier(TEXT("if"), ESearchCase::CaseSensitive) || PeekIdentifier(TEXT("ifdef"), ESearchCase::CaseSensitive))) )
 					{
@@ -2507,7 +2504,7 @@ void FHeaderParser::CheckSparseClassData(const FUnrealStructDefinitionInfo& Stru
 				UScriptStruct* ParentSparseClassDataStruct = ParentSparseClassDataStructDef->GetScriptStruct();
 				if (!SparseClassDataStructDef->GetScriptStruct()->IsChildOf(ParentSparseClassDataStruct))
 				{
-					FUHTException::Throwf(StructDef, TEXT("Class %s is a child of %s but its sparse class data struct, %s, does not inherit from %s."), *ClassDef->GetName(), *ParentClassDef->GetName(), *SparseClassDataStructDef->GetName(), *ParentSparseClassDataStruct->GetName());
+					FUHTException::Throwf(StructDef, TEXT("Class %s is a child of %s but its sparse class data struct, %s, does not inherit from %s."), *ClassDef->GetName(), *ParentClassDef->GetName(), *SparseClassDataStructDef->GetName(), *ParentSparseClassDataStructDef->GetName());
 				}
 			}
 		}
@@ -4596,70 +4593,6 @@ FUnrealPropertyDefinitionInfo& FHeaderParser::GetVarNameAndDim
 		// Only static arrays are declared with [].  Dynamic arrays use TArray<> instead.
 		VarProperty.ArrayType = EArrayType::Static;
 
-		UEnum* Enum = nullptr;
-
-		if (*Dimensions.String)
-		{
-			FString Temp = Dimensions.String;
-
-			bool bAgain;
-			do
-			{
-				bAgain = false;
-
-				// Remove any casts
-				static const TCHAR* Casts[] = {
-					TEXT("(uint32)"),
-					TEXT("(int32)"),
-					TEXT("(uint16)"),
-					TEXT("(int16)"),
-					TEXT("(uint8)"),
-					TEXT("(int8)"),
-					TEXT("(int)"),
-					TEXT("(unsigned)"),
-					TEXT("(signed)"),
-					TEXT("(unsigned int)"),
-					TEXT("(signed int)")
-				};
-
-				// Remove any brackets
-				if (Temp[0] == TEXT('('))
-				{
-					int32 TempLen      = Temp.Len();
-					int32 ClosingParen = FindMatchingClosingParenthesis(Temp);
-					if (ClosingParen == TempLen - 1)
-					{
-						Temp.MidInline(1, TempLen - 2, false);
-						bAgain = true;
-					}
-				}
-
-				for (const TCHAR* Cast : Casts)
-				{
-					if (Temp.StartsWith(Cast, ESearchCase::CaseSensitive))
-					{
-						Temp.RightChopInline(FCString::Strlen(Cast), false);
-						bAgain = true;
-					}
-				}
-			}
-			while (bAgain);
-
-			UEnum::LookupEnumNameSlow(*Temp, &Enum);
-		}
-
-		if (!Enum)
-		{
-			// If the enum wasn't declared in this scope, then try to find it anywhere we can
-			Enum = FEngineAPI::FindObject<UEnum>(ANY_PACKAGE, Dimensions.String);
-		}
-
-		if (Enum)
-		{
-			// set the ArraySizeEnum if applicable
-			VarProperty.MetaData.Add(NAME_ArraySizeEnum, Enum->GetPathName());
-		}
-
 		if (LayoutMacroType == ELayoutMacroType::None)
 		{
 			MatchSymbol(TEXT(']'));
@@ -6994,8 +6927,7 @@ struct FExposeOnSpawnValidator
 
 		if (!ProperNativeType && (CPT_Struct == Property.Type) && Property.ScriptStructDef)
 		{
-			UScriptStruct* Struct = Property.ScriptStructDef->GetScriptStruct();
-			ProperNativeType |= Struct->GetBoolMetaData(FHeaderParserNames::NAME_BlueprintType);
+			ProperNativeType |= Property.ScriptStructDef->GetBoolMetaData(FHeaderParserNames::NAME_BlueprintType);
 		}
 
 		return ProperNativeType;
@@ -7313,6 +7245,9 @@ void FHeaderParser::SkipStatements( int32 NestCount, const TCHAR* ErrorTag  )
 	Main script compiling routine.
 -----------------------------------------------------------------------------*/
 
+// Parse Class's annotated headers and optionally its child classes.
+static const FString ObjectHeader(TEXT("NoExportTypes.h"));
+
 //
 // Parses the header associated with the specified class.
 // Returns result enumeration.
@@ -7330,6 +7265,41 @@ void FHeaderParser::ParseHeader()
 	// Message.
 	UE_LOG(LogCompile, Verbose, TEXT("Parsing %s"), *SourceFile.GetFilename());
 
+	// Make sure that all of the requried include files are added to the outer scope
+	{
+		TArray<FUnrealSourceFile*> SourceFilesRequired;
+		for (FHeaderProvider& Include : SourceFile.GetIncludes())
+		{
+			if (Include.GetId() == ObjectHeader)
+			{
+				continue;
+			}
+
+			if (FUnrealSourceFile* DepFile = Include.Resolve(SourceFile))
+			{
+				SourceFilesRequired.Add(DepFile);
+			}
+		}
+
+		for (const TSharedRef<FUnrealTypeDefinitionInfo>& ClassDataPair : SourceFile.GetDefinedClasses())
+		{
+			FUnrealClassDefinitionInfo& ClassDef = ClassDataPair->AsClassChecked();
+			for (FUnrealClassDefinitionInfo* ParentClassDef = ClassDef.GetSuperClass(); ParentClassDef; ParentClassDef = ParentClassDef->GetSuperClass())
+			{
+				if (ParentClassDef->HasAnyClassFlags(CLASS_Parsed | CLASS_Intrinsic))
+				{
+					break;
+				}
+				SourceFilesRequired.Add(&ParentClassDef->GetUnrealSourceFile());
+			}
+		}
+
+		for (FUnrealSourceFile* RequiredFile : SourceFilesRequired)
+		{
+			SourceFile.GetScope()->IncludeScope(&RequiredFile->GetScope().Get());
+		}
+	}
+
 	// Init compiler variables.
 	ResetParser(*SourceFile.GetContent());
 
@@ -7427,15 +7397,21 @@ void FHeaderParser::ParseHeader()
 		const FString& ExpectedHeaderName = SourceFile.GetGeneratedHeaderFilename();
 		FUHTException::Throwf(*this, TEXT("Expected an include at the top of the header: '#include \"%s\"'"), *ExpectedHeaderName);
 	}
+
+	for (const TSharedRef<FUnrealTypeDefinitionInfo>& ClassDef : SourceFile.GetDefinedClasses())
+	{
+		ClassDef->AsClassChecked().SetClassFlags(CLASS_Parsed);
+	}
+
+	// Remember stats about this file
+	SourceFile.SetLinesParsed(LinesParsed);
+	SourceFile.SetStatementsParsed(StatementsParsed);
 }
 
 /*-----------------------------------------------------------------------------
 	Global functions.
 -----------------------------------------------------------------------------*/
 
-// Parse Class's annotated headers and optionally its child classes.
-static const FString ObjectHeader(TEXT("NoExportTypes.h"));
-
 FHeaderParser::FHeaderParser(FUnrealPackageDefinitionInfo& InPackageDef, FUnrealSourceFile& InSourceFile)
 	: FBaseParser(InSourceFile)
 	, Filename(IFileManager::Get().ConvertToAbsolutePathForExternalAppForRead(*InSourceFile.GetFilename()))
@@ -7499,62 +7475,7 @@ void FHeaderParser::Parse(
 	FUnrealSourceFile& SourceFile)
 {
 	SCOPE_SECONDS_COUNTER_UHT(Parse);
-
-	const FManifestModule& Module = PackageDef.GetModule();
-
-	// Disable loading of objects outside of this package (or more exactly, objects which aren't UFields, CDO, or templates)
-	TGuardValue<bool> AutoRestoreVerifyObjectRefsFlag(GVerifyObjectReferencesOnly, true);
-	// Create the header parser and register it as the warning context.
-	// Note: This must be declared outside the try block, since the catch block will log into it.
-	FHeaderParser HeaderParser(PackageDef, SourceFile);
-
-	// Hierarchically parse all classes.
-
-	TArray<FUnrealSourceFile*> SourceFilesRequired;
-	for (FHeaderProvider& Include : SourceFile.GetIncludes())
-	{
-		if (Include.GetId() == ObjectHeader)
-		{
-			continue;
-		}
-
-		if (FUnrealSourceFile* DepFile = Include.Resolve(SourceFile))
-		{
-			SourceFilesRequired.Add(DepFile);
-		}
-	}
-
-	for (const TSharedRef<FUnrealTypeDefinitionInfo>& ClassDataPair : SourceFile.GetDefinedClasses())
-	{
-		FUnrealClassDefinitionInfo& ClassDef = ClassDataPair->AsClassChecked();
-		for (FUnrealClassDefinitionInfo* ParentClassDef = ClassDef.GetSuperClass(); ParentClassDef; ParentClassDef = ParentClassDef->GetSuperClass())
-		{
-			if (ParentClassDef->HasAnyClassFlags(CLASS_Parsed | CLASS_Intrinsic))
-			{
-				break;
-			}
-			SourceFilesRequired.Add(&ParentClassDef->GetUnrealSourceFile());
-		}
-	}
-
-	for (FUnrealSourceFile* RequiredFile : SourceFilesRequired)
-	{
-		SourceFile.GetScope()->IncludeScope(&RequiredFile->GetScope().Get());
-	}
-
-	// Parse the file
-	{
-		HeaderParser.ParseHeader();
-
-		for (const TSharedRef<FUnrealTypeDefinitionInfo>& ClassDef : SourceFile.GetDefinedClasses())
-		{
-			ClassDef->AsClassChecked().SetClassFlags(CLASS_Parsed);
-		}
-	}
-
-	// Remember stats about this file
-	SourceFile.SetLinesParsed(HeaderParser.LinesParsed);
-	SourceFile.SetStatementsParsed(HeaderParser.StatementsParsed);
+	FHeaderParser(PackageDef, SourceFile).ParseHeader();
 }
 
 /** 
diff --git a/Engine/Source/Programs/UnrealHeaderTool/Private/PropertyTypes.cpp b/Engine/Source/Programs/UnrealHeaderTool/Private/PropertyTypes.cpp
--- a/Engine/Source/Programs/UnrealHeaderTool/Private/PropertyTypes.cpp
+++ b/Engine/Source/Programs/UnrealHeaderTool/Private/PropertyTypes.cpp
@@ -3,6 +3,7 @@
 #include "PropertyTypes.h"
 #include "BaseParser.h"
 #include "ClassMaps.h"
+#include "EngineAPI.h"
 #include "HeaderParser.h"
 #include "UnrealHeaderTool.h"
 #include "UnrealTypeDefinitionInfo.h"
@@ -86,6 +87,8 @@ struct FPropertyTypeTraitsEnum;
 
 namespace
 {
+	static const FName NAME_ArraySizeEnum(TEXT("ArraySizeEnum"));
+
 	void PropagateFlagsFromInnerAndHandlePersistentInstanceMetadata(EPropertyFlags& DestFlags, const TMap<FName, FString>& InMetaData, FUnrealPropertyDefinitionInfo& InnerDef)
 	{
 		// Copy some of the property flags to the container property.
@@ -433,9 +436,8 @@ struct FPropertyTypeTraitsByte : public FPropertyTypeTraitsNumericBase
 	static FProperty* CreateEngineType(FUnrealPropertyDefinitionInfo& PropDef, FFieldVariant Scope, const FName& Name, EObjectFlags ObjectFlags)
 	{
 		const FPropertyBase& VarProperty = PropDef.GetPropertyBase();
-		UEnum* Enum = VarProperty.EnumDef ? VarProperty.EnumDef->GetEnum() : nullptr;
 		FByteProperty* Result = new FByteProperty(Scope, Name, ObjectFlags);
-		Result->Enum = Enum;
+		Result->Enum = VarProperty.EnumDef ? VarProperty.EnumDef->GetEnum() : nullptr;
 		check(VarProperty.IntType == EIntType::Sized);
 		return Result;
 	}
@@ -454,18 +456,17 @@ struct FPropertyTypeTraitsByte : public FPropertyTypeTraitsNumericBase
 	{
 		if (FUnrealEnumDefinitionInfo* EnumDef = PropDef.GetPropertyBase().AsEnum())
 		{
-			UEnum* Enum = EnumDef->GetEnum();
 			const bool bEnumClassForm = EnumDef->GetCppForm() == UEnum::ECppForm::EnumClass;
-			const bool bNonNativeEnum = Enum->GetClass() != UEnum::StaticClass(); // cannot use RF_Native flag, because in UHT the flag is not set
+			const bool bNonNativeEnum = false; // Enum->GetClass() != UEnum::StaticClass(); // cannot use RF_Native flag, because in UHT the flag is not set
 			const bool bRawParam = (CPPExportFlags & CPPF_ArgumentOrReturnValue)
 				&& ((PropDef.HasAnyPropertyFlags(CPF_ReturnParm) || !PropDef.HasAnyPropertyFlags(CPF_OutParm))
 					|| bNonNativeEnum);
 			const bool bConvertedCode = (CPPExportFlags & CPPF_BlueprintCppBackend) && bNonNativeEnum;
 
 			FString FullyQualifiedEnumName;
-			if (!Enum->CppType.IsEmpty())
+			if (!EnumDef->GetCppType().IsEmpty())
 			{
-				FullyQualifiedEnumName = Enum->CppType;
+				FullyQualifiedEnumName = EnumDef->GetCppType();
 			}
 			else
 			{
@@ -473,12 +474,11 @@ struct FPropertyTypeTraitsByte : public FPropertyTypeTraitsNumericBase
 				// been set, but we do this here in case existing code relies on it... somehow.
 				if ((CPPExportFlags & CPPF_BlueprintCppBackend) && bNonNativeEnum)
 				{
-					ensure(Enum->CppType.IsEmpty());
-					FullyQualifiedEnumName = ::UnicodeToCPPIdentifier(Enum->GetName(), false, TEXT("E__"));
+					FullyQualifiedEnumName = ::UnicodeToCPPIdentifier(EnumDef->GetName(), false, TEXT("E__"));
 				}
 				else
 				{
-					FullyQualifiedEnumName = Enum->GetName();
+					FullyQualifiedEnumName = EnumDef->GetName();
 				}
 			}
 
@@ -962,15 +962,15 @@ struct FPropertyTypeTraitsEnum : public FPropertyTypeTraitsBase
 {
 	static bool DefaultValueStringCppFormatToInnerFormat(const FUnrealPropertyDefinitionInfo& PropDef, const FString& CppForm, FString& OutForm)
 	{
-		UEnum* Enum = PropDef.GetPropertyBase().EnumDef->GetEnum();
+		FUnrealEnumDefinitionInfo* EnumDef = PropDef.GetPropertyBase().EnumDef;
 		OutForm = FDefaultValueHelper::GetUnqualifiedEnumValue(FDefaultValueHelper::RemoveWhitespaces(CppForm));
 
-		const int32 EnumEntryIndex = Enum->GetIndexByName(*OutForm);
+		const int32 EnumEntryIndex = EnumDef->GetIndexByName(*OutForm);
 		if (EnumEntryIndex == INDEX_NONE)
 		{
 			return false;
 		}
-		if (Enum->HasMetaData(TEXT("Hidden"), EnumEntryIndex))
+		if (EnumDef->HasMetaData(TEXT("Hidden"), EnumEntryIndex))
 		{
 			FUHTException::Throwf(PropDef, TEXT("Hidden enum entries cannot be used as default values: %s \"%s\" "), *PropDef.GetName(), *CppForm);
 		}
@@ -981,9 +981,7 @@ struct FPropertyTypeTraitsEnum : public FPropertyTypeTraitsBase
 	{
 		const FPropertyBase& VarProperty = PropDef.GetPropertyBase();
 
-		UEnum* Enum = VarProperty.EnumDef->GetEnum();
-
-		if (Enum->GetCppForm() != UEnum::ECppForm::EnumClass)
+		if (VarProperty.EnumDef->GetCppForm() != UEnum::ECppForm::EnumClass)
 		{
 			check(VarProperty.Type == EPropertyType::CPT_Byte);
 			return FPropertyTypeTraitsByte::CreateProperty(PropDef);
@@ -1026,9 +1024,7 @@ struct FPropertyTypeTraitsEnum : public FPropertyTypeTraitsBase
 	{
 		const FPropertyBase& VarProperty = PropDef.GetPropertyBase();
 
-		UEnum* Enum = VarProperty.EnumDef->GetEnum();
-
-		if (Enum->GetCppForm() != UEnum::ECppForm::EnumClass)
+		if (VarProperty.EnumDef->GetCppForm() != UEnum::ECppForm::EnumClass)
 		{
 			check(VarProperty.Type == EPropertyType::CPT_Byte);
 			return FPropertyTypeTraitsByte::CreateEngineType(PropDef, Scope, Name, ObjectFlags);
@@ -1037,7 +1033,7 @@ struct FPropertyTypeTraitsEnum : public FPropertyTypeTraitsBase
 		FEnumProperty* Result = new FEnumProperty(Scope, Name, ObjectFlags);
 		PropDef.SetProperty(Result);
 		Result->UnderlyingProp = CastFieldChecked<FNumericProperty>(FPropertyTraits::CreateEngineType(PropDef.GetValuePropDefRef()));
-		Result->Enum = Enum;
+		Result->Enum = VarProperty.EnumDef->GetEnum();
 		return Result;
 	}
 
@@ -1049,8 +1045,7 @@ struct FPropertyTypeTraitsEnum : public FPropertyTypeTraitsBase
 	static FString GetEngineClassName(const FUnrealPropertyDefinitionInfo& PropDef)
 	{
 		const FPropertyBase& VarProperty = PropDef.GetPropertyBase();
-		UEnum* Enum = VarProperty.EnumDef->GetEnum();
-		if (Enum->GetCppForm() != UEnum::ECppForm::EnumClass)
+		if (VarProperty.EnumDef->GetCppForm() != UEnum::ECppForm::EnumClass)
 		{
 			return FPropertyTypeTraitsByte::GetEngineClassName(PropDef);
 		}
@@ -1063,27 +1058,26 @@ struct FPropertyTypeTraitsEnum : public FPropertyTypeTraitsBase
 	static FString GetCPPType(const FUnrealPropertyDefinitionInfo& PropDef, FString* ExtendedTypeText, uint32 CPPExportFlags)
 	{
 		const FPropertyBase& VarProperty = PropDef.GetPropertyBase();
-		UEnum* Enum = VarProperty.EnumDef->GetEnum();
-		if (Enum->GetCppForm() != UEnum::ECppForm::EnumClass)
+		FUnrealEnumDefinitionInfo* EnumDef = VarProperty.EnumDef;
+		if (EnumDef->GetCppForm() != UEnum::ECppForm::EnumClass)
 		{
 			return FPropertyTypeTraitsByte::GetCPPType(PropDef, ExtendedTypeText, CPPExportFlags);
 		}
 		else
 		{
-			const bool bNonNativeEnum = Enum->GetClass() != UEnum::StaticClass(); // cannot use RF_Native flag, because in UHT the flag is not set
+			const bool bNonNativeEnum = false; // Enum->GetClass() != UEnum::StaticClass(); // cannot use RF_Native flag, because in UHT the flag is not set
 
-			if (!Enum->CppType.IsEmpty())
+			if (!EnumDef->GetCppType().IsEmpty())
 			{
-				return Enum->CppType;
+				return EnumDef->GetCppType();
 			}
 
-			FString EnumName = Enum->GetName();
+			FString EnumName = EnumDef->GetName();
 
 			// This would give the wrong result if it's a namespaced type and the CppType hasn't
 			// been set, but we do this here in case existing code relies on it... somehow.
 			if ((CPPExportFlags & CPPF_BlueprintCppBackend) && bNonNativeEnum)
 			{
-				ensure(Enum->CppType.IsEmpty());
 				FString Result = ::UnicodeToCPPIdentifier(EnumName, false, TEXT("E__"));
 				return Result;
 			}
@@ -1095,13 +1089,12 @@ struct FPropertyTypeTraitsEnum : public FPropertyTypeTraitsBase
 	static FString GetCPPTypeForwardDeclaration(const FUnrealPropertyDefinitionInfo& PropDef)
 	{
 		const FPropertyBase& VarProperty = PropDef.GetPropertyBase();
-		UEnum* Enum = VarProperty.EnumDef->GetEnum();
-		if (Enum->GetCppForm() != UEnum::ECppForm::EnumClass)
+		if (VarProperty.EnumDef->GetCppForm() != UEnum::ECppForm::EnumClass)
 		{
 			return FPropertyTypeTraitsByte::GetCPPTypeForwardDeclaration(PropDef);
 		}
 
-		return FString::Printf(TEXT("enum class %s : %s;"), *Enum->GetName(), *PropDef.GetValuePropDef().GetCPPType());
+		return FString::Printf(TEXT("enum class %s : %s;"), *VarProperty.EnumDef->GetName(), *PropDef.GetValuePropDef().GetCPPType());
 	}
 };
 
@@ -1950,13 +1943,13 @@ struct FPropertyTypeTraitsStruct : public FPropertyTypeTraitsBase
 	static void CreateProperty(FUnrealPropertyDefinitionInfo& PropDef)
 	{
 		FPropertyBase& VarProperty = PropDef.GetPropertyBase();
-		UScriptStruct* Struct = PropDef.GetPropertyBase().ScriptStructDef->GetScriptStruct();
+		FUnrealScriptStructDefinitionInfo* ScriptStructDef = VarProperty.ScriptStructDef;
 
 #if UHT_ENABLE_VALUE_PROPERTY_TAG
 		PropDef.GetUnrealSourceFile().AddTypeDefIncludeIfNeeded(VarProperty.ScriptStructDef);
 #endif
 
-		if (Struct->StructFlags & STRUCT_HasInstancedReference)
+		if (ScriptStructDef->HasAnyStructFlags(STRUCT_HasInstancedReference))
 		{
 			VarProperty.PropertyFlags |= CPF_ContainsInstancedReference;
 		}
@@ -1973,8 +1966,7 @@ struct FPropertyTypeTraitsStruct : public FPropertyTypeTraitsBase
 	static bool IsSupportedByBlueprint(const FUnrealPropertyDefinitionInfo& PropDef, bool bMemberVariable)
 	{
 		const FPropertyBase& VarProperty = PropDef.GetPropertyBase();
-		UScriptStruct* Struct = PropDef.GetPropertyBase().ScriptStructDef->GetScriptStruct();
-		return Struct->GetBoolMetaDataHierarchical(FHeaderParserNames::NAME_BlueprintType);
+		return VarProperty.ScriptStructDef->GetBoolMetaDataHierarchical(FHeaderParserNames::NAME_BlueprintType);
 	}
 
 	static FString GetEngineClassName(const FUnrealPropertyDefinitionInfo& PropDef)
@@ -2233,6 +2225,70 @@ struct FPropertyTypeTraitsStaticArray : public FPropertyTypeTraitsBase
 
 	static FProperty* CreateEngineType(FUnrealPropertyDefinitionInfo& PropDef, FFieldVariant Scope, const FName& Name, EObjectFlags ObjectFlags)
 	{
+
+		if (PropDef.GetArrayDimensions())
+		{
+			UEnum* Enum = nullptr;
+
+			FString Temp = PropDef.GetArrayDimensions();
+
+			bool bAgain;
+			do
+			{
+				bAgain = false;
+
+				// Remove any casts
+				static const TCHAR* Casts[] = {
+					TEXT("(uint32)"),
+					TEXT("(int32)"),
+					TEXT("(uint16)"),
+					TEXT("(int16)"),
+					TEXT("(uint8)"),
+					TEXT("(int8)"),
+					TEXT("(int)"),
+					TEXT("(unsigned)"),
+					TEXT("(signed)"),
+					TEXT("(unsigned int)"),
+					TEXT("(signed int)")
+				};
+
+				// Remove any brackets
+				if (Temp[0] == TEXT('('))
+				{
+					int32 TempLen = Temp.Len();
+					int32 ClosingParen = FindMatchingClosingParenthesis(Temp);
+					if (ClosingParen == TempLen - 1)
+					{
+						Temp.MidInline(1, TempLen - 2, false);
+						bAgain = true;
+					}
+				}
+
+				for (const TCHAR* Cast : Casts)
+				{
+					if (Temp.StartsWith(Cast, ESearchCase::CaseSensitive))
+					{
+						Temp.RightChopInline(FCString::Strlen(Cast), false);
+						bAgain = true;
+					}
+				}
+			} while (bAgain);
+
+			UEnum::LookupEnumNameSlow(*Temp, &Enum);
+
+			if (!Enum)
+			{
+				// If the enum wasn't declared in this scope, then try to find it anywhere we can
+				Enum = FEngineAPI::FindObject<UEnum>(ANY_PACKAGE, PropDef.GetArrayDimensions());
+			}
+
+			if (Enum)
+			{
+				// set the ArraySizeEnum if applicable
+				PropDef.GetPropertyBase().MetaData.Add(NAME_ArraySizeEnum, Enum->GetPathName());
+			}
+		}
+
 		FProperty* Property = CreateEngineTypeHelper<false>(PropDef, Scope, Name, ObjectFlags);
 		Property->ArrayDim = 2;
 		return Property;
diff --git a/Engine/Source/Programs/UnrealHeaderTool/Private/UnrealTypeDefinitionInfo.cpp b/Engine/Source/Programs/UnrealHeaderTool/Private/UnrealTypeDefinitionInfo.cpp
--- a/Engine/Source/Programs/UnrealHeaderTool/Private/UnrealTypeDefinitionInfo.cpp
+++ b/Engine/Source/Programs/UnrealHeaderTool/Private/UnrealTypeDefinitionInfo.cpp
@@ -24,7 +24,6 @@
 extern FUnrealClassDefinitionInfo* GUObjectDef;
 extern FUnrealClassDefinitionInfo* GUClassDef;
 extern FUnrealClassDefinitionInfo* GUInterfaceDef;
-extern FUnrealClassDefinitionInfo* GUPackageDef;
 
 namespace
 {
@@ -1180,6 +1179,260 @@ FUnrealEnumDefinitionInfo::FUnrealEnumDefinitionInfo(FUnrealSourceFile& InSource
 	: FUnrealFieldDefinitionInfo(InSourceFile, InLineNumber, MoveTemp(InNameCPP), InName, InSourceFile.GetPackageDef())
 { }
 
+FString FUnrealEnumDefinitionInfo::GenerateEnumPrefix() const
+{
+	FString Prefix;
+	if (Names.Num() > 0)
+	{
+		Names[0].Key.ToString(Prefix);
+
+		// For each item in the enumeration, trim the prefix as much as necessary to keep it a prefix.
+		// This ensures that once all items have been processed, a common prefix will have been constructed.
+		// This will be the longest common prefix since as little as possible is trimmed at each step.
+		for (int32 NameIdx = 1; NameIdx < Names.Num(); NameIdx++)
+		{
+			FString EnumItemName = Names[NameIdx].Key.ToString();
+
+			// Find the length of the longest common prefix of Prefix and EnumItemName.
+			int32 PrefixIdx = 0;
+			while (PrefixIdx < Prefix.Len() && PrefixIdx < EnumItemName.Len() && Prefix[PrefixIdx] == EnumItemName[PrefixIdx])
+			{
+				PrefixIdx++;
+			}
+
+			// Trim the prefix to the length of the common prefix.
+			Prefix.LeftInline(PrefixIdx, false);
+		}
+
+		// Find the index of the rightmost underscore in the prefix.
+		int32 UnderscoreIdx = Prefix.Find(TEXT("_"), ESearchCase::CaseSensitive, ESearchDir::FromEnd);
+
+		// If an underscore was found, trim the prefix so only the part before the rightmost underscore is included.
+		if (UnderscoreIdx > 0)
+		{
+			Prefix.LeftInline(UnderscoreIdx, false);
+		}
+		else
+		{
+			// no underscores in the common prefix - this probably indicates that the names
+			// for this enum are not using Epic's notation, so just empty the prefix so that
+			// the max item will use the full name of the enum
+			Prefix.Empty();
+		}
+	}
+
+	// If no common prefix was found, or if the enum does not contain any entries,
+	// use the name of the enumeration instead.
+	if (Prefix.Len() == 0)
+	{
+		Prefix = GetName();
+	}
+
+#if UHT_ENABLE_ENGINE_TYPE_CHECKS
+	check(GetEnumSafe() == nullptr || GetEnumSafe()->GenerateEnumPrefix() == Prefix);
+#endif
+	return Prefix;
+}
+
+FString FUnrealEnumDefinitionInfo::GenerateFullEnumName(const TCHAR* InEnumName) const
+{
+	if (GetCppForm() == UEnum::ECppForm::Regular || UEnum::IsFullEnumName(InEnumName))
+	{
+		return InEnumName;
+	}
+
+	return FString::Printf(TEXT("%s::%s"), *GetName(), InEnumName);
+}
+
+bool FUnrealEnumDefinitionInfo::ContainsExistingMax() const
+{
+	if (GetIndexByName(*GenerateFullEnumName(TEXT("MAX")), EGetByNameFlags::CaseSensitive) != INDEX_NONE)
+	{
+#if UHT_ENABLE_ENGINE_TYPE_CHECKS
+		check(GetEnumSafe() == nullptr || GetEnumSafe()->ContainsExistingMax() == true);
+#endif
+		return true;
+	}
+
+	FName MaxEnumItem = *GenerateFullEnumName(*(GenerateEnumPrefix() + TEXT("_MAX")));
+	if (GetIndexByName(MaxEnumItem, EGetByNameFlags::CaseSensitive) != INDEX_NONE)
+	{
+#if UHT_ENABLE_ENGINE_TYPE_CHECKS
+		check(GetEnumSafe() == nullptr || GetEnumSafe()->ContainsExistingMax() == true);
+#endif
+		return true;
+	}
+
+#if UHT_ENABLE_ENGINE_TYPE_CHECKS
+	check(GetEnumSafe() == nullptr || GetEnumSafe()->ContainsExistingMax() == false);
+#endif
+	return false;
+}
+
+int64 FUnrealEnumDefinitionInfo::GetMaxEnumValue() const
+{
+	int32 NamesNum = Names.Num();
+	if (NamesNum == 0)
+	{
+		return 0;
+	}
+
+	int64 MaxValue = Names[0].Value;
+	for (int32 i = 0; i < NamesNum; ++i)
+	{
+		int64 CurrentValue = Names[i].Value;
+		if (CurrentValue > MaxValue)
+		{
+			MaxValue = CurrentValue;
+		}
+	}
+
+#if UHT_ENABLE_ENGINE_TYPE_CHECKS
+	check(GetEnumSafe() == nullptr || GetEnumSafe()->GetMaxEnumValue() == MaxValue);
+#endif
+	return MaxValue;
+}
+
+FName FUnrealEnumDefinitionInfo::GetNameByIndex(int32 Index) const
+{
+	if (Names.IsValidIndex(Index))
+	{
+#if UHT_ENABLE_ENGINE_TYPE_CHECKS
+		check(GetEnumSafe() == nullptr || GetEnumSafe()->GetNameByIndex(Index) == Names[Index].Key);
+#endif
+		return Names[Index].Key;
+	}
+
+#if UHT_ENABLE_ENGINE_TYPE_CHECKS
+	check(GetEnumSafe() == nullptr || GetEnumSafe()->GetNameByIndex(Index) == NAME_None);
+#endif
+	return NAME_None;
+}
+
+bool FUnrealEnumDefinitionInfo::IsValidEnumValue(int64 InValue) const
+{
+	int32 NamesNum = Names.Num();
+	for (int32 i = 0; i < NamesNum; ++i)
+	{
+		int64 CurrentValue = Names[i].Value;
+		if (CurrentValue == InValue)
+		{
+#if UHT_ENABLE_ENGINE_TYPE_CHECKS
+			check(GetEnumSafe() == nullptr || GetEnumSafe()->IsValidEnumValue(InValue) == true);
+#endif
+			return true;
+		}
+	}
+
+#if UHT_ENABLE_ENGINE_TYPE_CHECKS
+	check(GetEnumSafe() == nullptr || GetEnumSafe()->IsValidEnumValue(InValue) == false);
+#endif
+	return false;
+}
+
+int32 FUnrealEnumDefinitionInfo::GetIndexByName(const FName InName, EGetByNameFlags Flags) const
+{
+	ENameCase ComparisonMethod = !!(Flags & EGetByNameFlags::CaseSensitive) ? ENameCase::CaseSensitive : ENameCase::IgnoreCase;
+
+	// First try the fast path
+	const int32 Count = Names.Num();
+	for (int32 Counter = 0; Counter < Count; ++Counter)
+	{
+		if (Names[Counter].Key.IsEqual(InName, ComparisonMethod))
+		{
+#if UHT_ENABLE_ENGINE_TYPE_CHECKS
+			check(GetEnumSafe() == nullptr || GetEnumSafe()->GetIndexByName(InName, Flags) == Counter);
+#endif
+			return Counter;
+		}
+	}
+
+	// Otherwise see if it is in the redirect table
+	int32 Result = GetIndexByNameString(InName.ToString(), Flags);
+	return Result;
+}
+
+int32 FUnrealEnumDefinitionInfo::GetIndexByNameString(const FString& InSearchString, EGetByNameFlags Flags) const
+{
+	ENameCase         NameComparisonMethod = !!(Flags & EGetByNameFlags::CaseSensitive) ? ENameCase::CaseSensitive : ENameCase::IgnoreCase;
+	ESearchCase::Type StringComparisonMethod = !!(Flags & EGetByNameFlags::CaseSensitive) ? ESearchCase::CaseSensitive : ESearchCase::IgnoreCase;
+
+	FString SearchEnumEntryString = InSearchString;
+	FString ModifiedEnumEntryString;
+
+	// Strip or add the namespace
+	int32 DoubleColonIndex = SearchEnumEntryString.Find(TEXT("::"), ESearchCase::CaseSensitive);
+	if (DoubleColonIndex == INDEX_NONE)
+	{
+		ModifiedEnumEntryString = GenerateFullEnumName(*SearchEnumEntryString);
+	}
+	else
+	{
+		ModifiedEnumEntryString = SearchEnumEntryString.RightChop(DoubleColonIndex + 2);
+	}
+
+	if (DoubleColonIndex != INDEX_NONE)
+	{
+		// If we didn't find a value redirect and our original string was namespaced, we need to fix the namespace now as it may have changed due to enum type redirect
+		SearchEnumEntryString = GenerateFullEnumName(*ModifiedEnumEntryString);
+	}
+
+	// Search for names both with and without namespace
+	FName SearchName = FName(*SearchEnumEntryString, FNAME_Find);
+	FName ModifiedName = FName(*ModifiedEnumEntryString, FNAME_Find);
+
+	const int32 Count = Names.Num();
+	for (int32 Counter = 0; Counter < Count; ++Counter)
+	{
+		if (Names[Counter].Key.IsEqual(SearchName, NameComparisonMethod) || Names[Counter].Key.IsEqual(ModifiedName, NameComparisonMethod))
+		{
+			return Counter;
+		}
+	}
+
+	if (!InSearchString.Equals(SearchEnumEntryString, StringComparisonMethod))
+	{
+		// There was an actual redirect, and we didn't find it
+		//UE_LOG(LogEnum, Warning, TEXT("EnumRedirect for enum %s maps '%s' to invalid value '%s'!"), *GetName(), *InSearchString, *SearchEnumEntryString);
+	}
+	return INDEX_NONE;
+}
+
+FString FUnrealEnumDefinitionInfo::GetNameStringByIndex(int32 InIndex) const
+{
+	if (Names.IsValidIndex(InIndex))
+	{
+		FName EnumEntryName = GetNameByIndex(InIndex);
+		if (CppForm == UEnum::ECppForm::Regular)
+		{
+			return EnumEntryName.ToString();
+		}
+
+		// Strip the namespace from the name.
+		FString EnumNameString(EnumEntryName.ToString());
+		int32 ScopeIndex = EnumNameString.Find(TEXT("::"), ESearchCase::CaseSensitive);
+		if (ScopeIndex != INDEX_NONE)
+		{
+			return EnumNameString.Mid(ScopeIndex + 2);
+		}
+	}
+	return FString();
+}
+
+void FUnrealEnumDefinitionInfo::CreateUObjectEngineTypesInternal()
+{
+	UPackage* Package = GetPackageDef().GetPackage();
+	const FString& EnumName = GetNameCPP();
+
+	// Create enum definition.
+	UEnum* Enum = new(EC_InternalUseOnlyConstructor, Package, FName(EnumName), RF_Public) UEnum(FObjectInitializer());
+	Enum->SetEnums(Names, CppForm, EnumFlags, false);
+	Enum->CppType = CppType;
+	Enum->GetPackage()->GetMetaData()->SetObjectValues(Enum, MetaData);
+	SetObject(Enum);
+	GTypeDefinitionInfoMap.AddObjectLookup(Enum, SharedThis(this));
+}
+
 void FUnrealStructDefinitionInfo::AddProperty(TSharedRef<FUnrealPropertyDefinitionInfo> PropertyDef)
 {
 	Properties.Add(PropertyDef);
@@ -1297,8 +1550,7 @@ uint32 FUnrealScriptStructDefinitionInfo::GetHash(bool bIncludeNoExport) const
 {
 	if (!bIncludeNoExport)
 	{
-		UScriptStruct* Struct = GetScriptStruct();
-		if ((Struct->StructFlags & STRUCT_NoExport) != 0)
+		if (HasAnyStructFlags(STRUCT_NoExport))
 		{
 			return 0;
 		}
@@ -2123,6 +2375,7 @@ void FUnrealFunctionDefinitionInfo::CreateUObjectEngineTypesInternal()
 	Function->ReturnValueOffset = MAX_uint16;
 	Function->FirstPropertyToInit = nullptr;
 	Function->FunctionFlags |= FunctionData.FunctionFlags;
+	Function->GetPackage()->GetMetaData()->SetObjectValues(Function, MetaData);
 
 	SetObject(Function);
 	GTypeDefinitionInfoMap.Add(Function, SharedThis(this));
diff --git a/Engine/Source/Programs/UnrealHeaderTool/Private/UnrealTypeDefinitionInfo.h b/Engine/Source/Programs/UnrealHeaderTool/Private/UnrealTypeDefinitionInfo.h
--- a/Engine/Source/Programs/UnrealHeaderTool/Private/UnrealTypeDefinitionInfo.h
+++ b/Engine/Source/Programs/UnrealHeaderTool/Private/UnrealTypeDefinitionInfo.h
@@ -176,6 +176,25 @@ public:
 		SetMetaDataHelper(GetMetaDataKey(Key, NameIndex, FNAME_Add), InValue);
 	}
 
+	/**
+	 * Find the metadata value associated with the key
+	 * and return bool
+	 * @param Key The key to lookup in the metadata
+	 * @return return true if the value was true (case insensitive)
+	 */
+	bool GetBoolMetaData(const TCHAR* Key) const
+	{
+		const FString& BoolString = GetMetaData(Key);
+		// FString == operator does case insensitive comparison
+		return (BoolString == "true");
+	}
+	bool GetBoolMetaData(const FName& Key) const
+	{
+		const FString& BoolString = GetMetaData(Key);
+		// FString == operator does case insensitive comparison
+		return (BoolString == "true");
+	}
+
 	virtual TMap<FName, FString>* GetMetaDataMap() const = 0;
 
 protected:
@@ -1756,9 +1775,24 @@ public:
 		return static_cast<UEnum*>(GetObject());
 	}
 
+	/**
+	 * Return the Engine instance associated with the compiler instance but does not check to see if it is set
+	 */
+	UEnum* GetEnumSafe() const
+	{
+		return static_cast<UEnum*>(GetObjectSafe());
+	}
+
+	/**
+	 * Test if the enum contains and of the given flags
+	 */
 	bool HasAnyEnumFlags(EEnumFlags InFlags) const
 	{
-		return GetEnum()->HasAnyEnumFlags(InFlags);
+		bool bResults = EnumHasAnyFlags(EnumFlags, InFlags);
+#if UHT_ENABLE_ENGINE_TYPE_CHECKS
+		check(GetEnumSafe() == nullptr || GetEnumSafe()->HasAnyEnumFlags(InFlags) == bResults);
+#endif
+		return bResults;
 	}
 
 	/**
@@ -1766,48 +1800,40 @@ public:
 	 *
 	 * @return	true if the enum contains a MAX enum, false otherwise.
 	 */
-	bool ContainsExistingMax() const
-	{
-		return GetEnum()->ContainsExistingMax();
-	}
+	bool ContainsExistingMax() const;
 
 	/**
-	 * @return	 The number of enum names.
+	 * @return The number of enum names.
 	 */
 	int32 NumEnums() const
 	{
-		return GetEnum()->NumEnums();
+#if UHT_ENABLE_ENGINE_TYPE_CHECKS
+		check(GetEnumSafe() == nullptr || GetEnumSafe()->NumEnums() == Names.Num());
+#endif
+		return Names.Num();
 	}
 
 	/** Gets max value of Enum. Defaults to zero if there are no entries. */
-	int64 GetMaxEnumValue() const
-	{
-		return GetEnum()->GetMaxEnumValue();
-	}
+	int64 GetMaxEnumValue() const;
 
 	/** Gets enum value by index in Names array. Asserts on invalid index */
-	FORCEINLINE int64 GetValueByIndex(int32 Index) const
+	int64 GetValueByIndex(int32 Index) const
 	{
-		return GetEnum()->GetValueByIndex(Index);
+		check(Names.IsValidIndex(Index));
+#if UHT_ENABLE_ENGINE_TYPE_CHECKS
+		check(GetEnumSafe() == nullptr || GetEnumSafe()->GetValueByIndex(Index) == Names[Index].Value);
+#endif
+		return Names[Index].Value;
 	}
 
 	/** Gets enum name by index in Names array. Returns NAME_None if Index is not valid. */
-	FName GetNameByIndex(int32 Index) const
-	{
-		return GetEnum()->GetNameByIndex(Index);
-	}
+	FName GetNameByIndex(int32 Index) const;
 
 	/** Gets index of name in enum, returns INDEX_NONE and optionally errors when name is not found. This is faster than ByNameString if the FName is exact, but will fall back if needed */
-	int32 GetIndexByName(FName InName, EGetByNameFlags Flags = EGetByNameFlags::None) const
-	{
-		return GetEnum()->GetIndexByName(InName, Flags);
-	}
+	int32 GetIndexByName(FName InName, EGetByNameFlags Flags = EGetByNameFlags::None) const;
 
 	/** Checks if enum has entry with given value. Includes autogenerated _MAX entry. */
-	bool IsValidEnumValue(int64 InValue) const
-	{
-		return GetEnum()->IsValidEnumValue(InValue);
-	}
+	bool IsValidEnumValue(int64 InValue) const;
 
 	/**
 	 * Returns the type of enum: whether it's a regular enum, namespaced enum or C++11 enum class.
@@ -1816,7 +1842,10 @@ public:
 	 */
 	UEnum::ECppForm GetCppForm() const
 	{
-		return GetEnum()->GetCppForm();
+#if UHT_ENABLE_ENGINE_TYPE_CHECKS
+		check(GetEnumSafe() == nullptr || GetEnumSafe()->GetCppForm() == CppForm);
+#endif
+		return CppForm;
 	}
 
 	/**
@@ -1824,15 +1853,22 @@ public:
 	 */
 	const FString& GetCppType() const
 	{
-		return GetEnum()->CppType;
+#if UHT_ENABLE_ENGINE_TYPE_CHECKS
+		check(GetEnumSafe() == nullptr || GetEnumSafe()->CppType == CppType);
+#endif
+		return CppType;
 	}
 
 	/**
 	 * Sets the enum type
 	 */
-	void SetCppType(FString&& CppType)
+	void SetCppType(FString&& InCppType)
 	{
-		GetEnum()->CppType = MoveTemp(CppType);
+		CppType = MoveTemp(InCppType);
+		if (GetEnumSafe() != nullptr)
+		{
+			GetEnumSafe()->CppType = CppType;
+		}
 	}
 
 	/**
@@ -1841,10 +1877,7 @@ public:
 	 * @return	the longest common prefix between all items in the enum.  If a common prefix
 	 *			cannot be found, returns the full name of the enum.
 	 */
-	FString GenerateEnumPrefix() const
-	{
-		return GetEnum()->GenerateEnumPrefix();
-	}
+	FString GenerateEnumPrefix() const;
 
 	/**
 	 * Sets the array of enums.
@@ -1854,9 +1887,19 @@ public:
 	 * @param bAddMaxKeyIfMissing Should a default Max item be added.
 	 * @return	true unless the MAX enum already exists and isn't the last enum.
 	 */
-	bool SetEnums(TArray<TPair<FName, int64>>& InNames, UEnum::ECppForm InCppForm, EEnumFlags InFlags = EEnumFlags::None, bool bAddMaxKeyIfMissing = true)
+	bool SetEnums(TArray<TPair<FName, int64>>& InNames, UEnum::ECppForm InCppForm, EEnumFlags InFlags = EEnumFlags::None)
 	{
-		return GetEnum()->SetEnums(InNames, InCppForm, InFlags, bAddMaxKeyIfMissing);
+		Names = InNames;
+		CppForm = InCppForm;
+		EnumFlags = InFlags;
+#if UHT_ENABLE_ENGINE_TYPE_CHECKS
+		if (UEnum* Enum = GetEnumSafe())
+		{
+			Enum->SetEnums(InNames, InCppForm, InFlags, false /*bAddMaxKeyIfMissing*/);
+			check(Enum->GetCppForm() == InCppForm);
+		}
+#endif
+		return true;
 	}
 
 	/**
@@ -1892,9 +1935,16 @@ public:
 	}
 
 protected:
+	virtual void CreateUObjectEngineTypesInternal() override;
+
+	virtual TMap<FName, FString>* GetMetaDataMap() const override
+	{
+		return const_cast<TMap<FName, FString>*>(&MetaData);
+	}
+
 	virtual FString GetMetaDataIndexName(int32 NameIndex) const
 	{
-		FString EnumName = GetEnum()->GetNameStringByIndex(NameIndex);
+		FString EnumName = GetNameStringByIndex(NameIndex);
 		check(!EnumName.IsEmpty());
 		return EnumName;
 	}
@@ -1904,8 +1954,17 @@ protected:
 		return TEXT("EnumRemap");
 	}
 
+	FString GenerateFullEnumName(const TCHAR* InEnumName) const;
+	int32 GetIndexByNameString(const FString& InSearchString, EGetByNameFlags Flags) const;
+	FString GetNameStringByIndex(int32 InIndex) const;
+
 private:
+	TArray<TPair<FName, int64>> Names;
+	TMap<FName, FString> MetaData;
+	FString CppType;
 	EUnderlyingEnumType UnderlyingType = EUnderlyingEnumType::Unspecified;
+	EEnumFlags EnumFlags = EEnumFlags::None;
+	UEnum::ECppForm CppForm = UEnum::ECppForm::Regular;
 	bool bIsEditorOnly = false;
 };
 
