Classes and ScriptStructs are now created post-parse

#rnx

diff --git a/Engine/Source/Programs/UnrealHeaderTool/Private/ClassMaps.h b/Engine/Source/Programs/UnrealHeaderTool/Private/ClassMaps.h
--- a/Engine/Source/Programs/UnrealHeaderTool/Private/ClassMaps.h
+++ b/Engine/Source/Programs/UnrealHeaderTool/Private/ClassMaps.h
@@ -146,9 +146,9 @@ struct FTypeDefinitionInfoMap : public FFreezableContainer
 	}
 
 	template <typename Lambda>
-	void ForAllTypes(Lambda&& InLambda)
+	void ForAllTypesByName(Lambda&& InLambda)
 	{
-		for (const TPair<UObject*, TSharedRef<FUnrealTypeDefinitionInfo>>& KVP : DefinitionsByUObject)
+		for (const TPair<FName, TSharedRef<FUnrealTypeDefinitionInfo>>& KVP : DefinitionsByName)
 		{
 			InLambda(*KVP.Value);
 		}
diff --git a/Engine/Source/Programs/UnrealHeaderTool/Private/CodeGenerator.cpp b/Engine/Source/Programs/UnrealHeaderTool/Private/CodeGenerator.cpp
--- a/Engine/Source/Programs/UnrealHeaderTool/Private/CodeGenerator.cpp
+++ b/Engine/Source/Programs/UnrealHeaderTool/Private/CodeGenerator.cpp
@@ -76,9 +76,9 @@ static TArray<FString> ChangeMessages;
 static bool bWriteContents = false;
 static bool bVerifyContents = false;
 
-UClass* ProcessParsedClass(FUnrealClassDefinitionInfo& ClassDef);
+void ProcessParsedClass(FUnrealClassDefinitionInfo& ClassDef);
 void ProcessParsedEnum(FUnrealEnumDefinitionInfo& EnumDef);
-UScriptStruct* ProcessParsedStruct(FUnrealScriptStructDefinitionInfo& ScriptStructDef);
+void ProcessParsedStruct(FUnrealScriptStructDefinitionInfo& ScriptStructDef);
 
 // Array of all the temporary header async file tasks so we can ensure they have completed before issuing our timings
 static FGraphEventArray GAsyncFileTasks;
@@ -170,7 +170,7 @@ namespace
 	const TCHAR EndEditorOnlyGuard[] = TEXT("#endif //WITH_EDITOR") LINE_TERMINATOR;
 
 	/** Whether or not the given class has any replicated properties. */
-	static bool ClassHasReplicatedProperties(FUnrealClassDefinitionInfo& ClassDef)
+	static bool ClassHasReplicatedProperties(const FUnrealClassDefinitionInfo& ClassDef)
 	{
 		if (!ClassDef.HasAnyClassFlags(CLASS_ReplicationDataIsSetUp))
 		{
@@ -183,13 +183,12 @@ namespace
 			}
 		}
 
-		UClass* Class = ClassDef.GetClass();
-		return Class->FirstOwnedClassRep < Class->ClassReps.Num();
+		return ClassDef.HasOwnedClassReps();
 	}
 
-	static void ExportNetData(FOutputDevice& Out, UClass* Class, const TCHAR* API)
+	static void ExportNetData(FOutputDevice& Out, const FUnrealClassDefinitionInfo& ClassDef, const TCHAR* API)
 	{
-		const TArray<FRepRecord>& ClassReps = Class->ClassReps;
+		const TArray<FUnrealPropertyDefinitionInfo*>& ClassReps = ClassDef.GetClassReps();
 
 		FUHTStringBuilder NetFieldBuilder;
 		NetFieldBuilder.Logf(TEXT(""
@@ -201,12 +200,12 @@ namespace
 
 		bool bAnyStaticArrays = false;
 		bool bIsFirst = true;
-		for (int32 ClassRepIndex = Class->FirstOwnedClassRep; ClassRepIndex < ClassReps.Num(); ++ClassRepIndex)
+		for (int32 ClassRepIndex = ClassDef.GetFirstOwnedClassRep(); ClassRepIndex < ClassReps.Num(); ++ClassRepIndex)
 		{
-			const FRepRecord& ClassRep = ClassReps[ClassRepIndex];
-			const FString PropertyName = ClassRep.Property->GetName();
+			const FUnrealPropertyDefinitionInfo* PropertyDef = ClassReps[ClassRepIndex];
+			const FString PropertyName = PropertyDef->GetName();
 
-			if (ClassRep.Property->ArrayDim == 1)
+			if (!PropertyDef->IsStaticArray())
 			{
 				if (UNLIKELY(bIsFirst))
 				{
@@ -220,10 +219,8 @@ namespace
 			}
 			else
 			{
-				FUnrealPropertyDefinitionInfo& PropDef = GTypeDefinitionInfoMap.FindChecked<FUnrealPropertyDefinitionInfo>(ClassReps[ClassRepIndex].Property);
-
 				bAnyStaticArrays = true;
-				ArrayDimBuilder.Logf(TEXT("\t\t%s=%s,\r\n"), *PropertyName, PropDef.GetArrayDimensions());
+				ArrayDimBuilder.Logf(TEXT("\t\t%s=%s,\r\n"), *PropertyName, PropertyDef->GetArrayDimensions());
 
 				if (UNLIKELY(bIsFirst))
 				{
@@ -239,8 +236,8 @@ namespace
 			}
 		}
 
-		const FProperty* LastProperty = ClassReps.Last().Property;
-		NetFieldBuilder.Logf(TEXT("\t\tNETFIELD_REP_END=%s%s"), *LastProperty->GetName(), LastProperty->ArrayDim > 1 ? TEXT("_STATIC_ARRAY_END") : TEXT(""));
+		const FUnrealPropertyDefinitionInfo* LastPropertyDef = ClassReps.Last();
+		NetFieldBuilder.Logf(TEXT("\t\tNETFIELD_REP_END=%s%s"), *LastPropertyDef->GetName(), LastPropertyDef->IsStaticArray() ? TEXT("_STATIC_ARRAY_END") : TEXT(""));
 
 		NetFieldBuilder.Log(TEXT("\t};"));
 
@@ -271,7 +268,6 @@ namespace
 		const FUnrealSourceFile& SourceFile,
 		EExportClassOutFlags& OutFlags)
 	{
-		UClass* Class = ClassDef.GetClass();
 		const bool bHasGetLifetimeReplicatedProps = HasIdentifierExactMatch(ClassDef.GetDefinitionRange().Start, ClassDef.GetDefinitionRange().End, STRING_GetLifetimeReplicatedPropsStr);
 
 		if (!bHasGetLifetimeReplicatedProps)
@@ -288,11 +284,11 @@ namespace
 		}
 
 
-		ExportNetData(Writer, Class, API);
+		ExportNetData(Writer, ClassDef, API);
 		
 		// If this class has replicated properties and it owns the first one, that means
 		// it's the base most replicated class. In that case, go ahead and add our interface macro.
-		if (Class->ClassReps.Num() > 0 && Class->FirstOwnedClassRep == 0)
+		if (ClassDef.HasClassReps() && ClassDef.GetFirstOwnedClassRep() == 0)
 		{
 			OutFlags |= EExportClassOutFlags::NeedsPushModelHeaders;
 			Writer.Logf(TEXT(
@@ -2061,13 +2057,9 @@ void FNativeClassHeaderGenerator::ExportNativeGeneratedInitCode(FOutputDevice& O
 			FunctionsCount = TEXT("0");
 		}
 
-		TMap<FName, FString>* MetaDataMap = UMetaData::GetMapForObject(Class);
-		if (MetaDataMap)
+		if (ClassDef.GetStructMetaData().bObjectInitializerConstructorDeclared)
 		{
-			if (ClassDef.GetStructMetaData().bObjectInitializerConstructorDeclared)
-			{
-				MetaDataMap->Add(NAME_ObjectInitializerConstructorDeclared, FString());
-			}
+			ClassDef.SetMetaData(NAME_ObjectInitializerConstructorDeclared, TEXT(""));
 		}
 
 		FString MetaDataParams = OutputMetaDataCodeForObject(GeneratedClassRegisterFunctionText, StaticDefinitions, ClassDef, *FString::Printf(TEXT("%s::Class_MetaDataParams"), *StaticsStructName), TEXT("\t\t"), TEXT("\t"));
@@ -2179,7 +2171,7 @@ void FNativeClassHeaderGenerator::ExportNativeGeneratedInitCode(FOutputDevice& O
 
 		if (bIsDynamic)
 		{
-			FString* CustomDynamicClassInitializationMD = MetaDataMap ? MetaDataMap->Find(TEXT("CustomDynamicClassInitialization")) : nullptr;
+			const FString* CustomDynamicClassInitializationMD = ClassDef.FindMetaData(TEXT("CustomDynamicClassInitialization"));
 			if (CustomDynamicClassInitializationMD)
 			{
 				GeneratedClassRegisterFunctionText.Logf(TEXT("\t\t\t\t%s(CastChecked<UDynamicClass>(OuterClass));\n"), *(*CustomDynamicClassInitializationMD));
@@ -6121,74 +6113,81 @@ void GetScriptPlugins(TArray<IScriptGeneratorPluginInterface*>& ScriptPlugins)
 	}
 }
 
+TArray<TSharedRef<FUnrealTypeDefinitionInfo>> GEngineClasses;
+
 /**
- * Tries to resolve super classes for classes defined in the given
- * module.
- *
- * @param Package Modules package.
+ * Tries to resolve super classes for classes defined in the given class
  */
-void ResolveSuperClasses(UPackage* Package)
+void ResolveSuperClasses(const TCHAR* PackageName, FUnrealClassDefinitionInfo& ClassDef)
 {
-	TArray<UObject*> Objects;
-	GetObjectsWithPackage(Package, Objects);
+	UClass* Class = ClassDef.GetClassSafe();
 
-	for (UObject* Object : Objects)
+	// Propagate any already set ClassWithin (this only happens for the temporary engine objects)
+	if (Class != nullptr && Class->ClassWithin != nullptr)
 	{
-		if (!Object->IsA<UClass>() || Object->HasAnyFlags(RF_ClassDefaultObject))
-		{
-			continue;
-		}
-
-		UClass* DefinedClass = Cast<UClass>(Object);
-
-		FUnrealClassDefinitionInfo& ClassDef = GTypeDefinitionInfoMap.FindChecked<FUnrealClassDefinitionInfo>(DefinedClass);
+		ClassDef.SetClassWithin(&GTypeDefinitionInfoMap.FindByNameChecked<FUnrealClassDefinitionInfo>(*Class->ClassWithin->GetFName().ToString()));
+	}
 
-		// Propagate any already set ClassWithin
-		if (DefinedClass->ClassWithin != nullptr)
-		{
-			ClassDef.SetClassWithin(&GTypeDefinitionInfoMap.FindChecked<FUnrealClassDefinitionInfo>(DefinedClass->ClassWithin));
-		}
+	// Resolve the base class
+	{
+		FUnrealStructDefinitionInfo::FBaseStructInfo& SuperClassInfo = ClassDef.GetSuperStructInfo();
 
-		// Resolve the base class
+		if (Class == nullptr || Class->GetSuperClass() == nullptr)
 		{
-			FUnrealStructDefinitionInfo::FBaseStructInfo& SuperClassInfo = ClassDef.GetSuperStructInfo();
-
-			if (DefinedClass->GetSuperClass() == nullptr)
+			if (!SuperClassInfo.Name.IsEmpty())
 			{
-				if (!SuperClassInfo.Name.IsEmpty() && !DefinedClass->GetSuperClass())
-				{
-					const FString& BaseClassName = SuperClassInfo.Name;
-					const FString& BaseClassNameStripped = GetClassNameWithPrefixRemoved(BaseClassName);
+				const FString& BaseClassName = SuperClassInfo.Name;
+				const FString& BaseClassNameStripped = GetClassNameWithPrefixRemoved(BaseClassName);
 
-					FUnrealClassDefinitionInfo* FoundBaseClassDef = GTypeDefinitionInfoMap.FindByName<FUnrealClassDefinitionInfo>(*BaseClassNameStripped);
-					if (FoundBaseClassDef == nullptr)
-					{
-						FoundBaseClassDef = GTypeDefinitionInfoMap.FindByName<FUnrealClassDefinitionInfo>(*BaseClassName);
-					}
+				FUnrealClassDefinitionInfo* FoundBaseClassDef = GTypeDefinitionInfoMap.FindByName<FUnrealClassDefinitionInfo>(*BaseClassNameStripped);
+				if (FoundBaseClassDef == nullptr)
+				{
+					FoundBaseClassDef = GTypeDefinitionInfoMap.FindByName<FUnrealClassDefinitionInfo>(*BaseClassName);
+				}
 
-					if (FoundBaseClassDef == nullptr)
-					{
-						// Don't know its parent class. Raise error.
-						FUHTException::Throwf(ClassDef, TEXT("Couldn't find parent type for '%s' named '%s' in current module (Package: %s) or any other module parsed so far."), *DefinedClass->GetName(), *BaseClassName, *Package->GetName());
-					}
+				if (FoundBaseClassDef == nullptr)
+				{
+					// Don't know its parent class. Raise error.
+					FUHTException::Throwf(ClassDef, TEXT("Couldn't find parent type for '%s' named '%s' in current module (Package: %s) or any other module parsed so far."), *ClassDef.GetName(), *BaseClassName, PackageName);
+				}
 
-					DefinedClass->SetSuperStruct(FoundBaseClassDef->GetClass());
-					DefinedClass->ClassCastFlags |= FoundBaseClassDef->GetClass()->ClassCastFlags;
-					SuperClassInfo.Struct = FoundBaseClassDef;
+				SuperClassInfo.Struct = FoundBaseClassDef;
+				ClassDef.SetClassCastFlags(FoundBaseClassDef->GetClassCastFlags());
+				if (Class != nullptr)
+				{
+					Class->SetSuperStruct(FoundBaseClassDef->GetClass());
 				}
 			}
-			else
-			{
-				SuperClassInfo.Struct = &GTypeDefinitionInfoMap.FindChecked<FUnrealClassDefinitionInfo>(DefinedClass->GetSuperClass());
-			}
 		}
+		else
+		{
+			SuperClassInfo.Struct = &GTypeDefinitionInfoMap.FindByNameChecked<FUnrealClassDefinitionInfo>(*Class->GetSuperClass()->GetFName().ToString());
+		}
+	}
 
-		// Resolve the inherited classes
+	// Resolve the inherited classes
+	{
+		for (FUnrealStructDefinitionInfo::FBaseStructInfo& BaseClassInfo : ClassDef.GetBaseStructInfos())
 		{
-			for (FUnrealStructDefinitionInfo::FBaseStructInfo& BaseClassInfo : ClassDef.GetBaseStructInfo())
-			{
-				BaseClassInfo.Struct = FUnrealClassDefinitionInfo::FindScriptClass(BaseClassInfo.Name);
-			}
+			BaseClassInfo.Struct = FUnrealClassDefinitionInfo::FindScriptClass(BaseClassInfo.Name);
+		}
+	}
+}
+
+/**
+ * Tries to resolve super classes for classes defined in the given
+ * module.
+ *
+ * @param Package Modules package.
+ */
+void ResolveSuperClasses(FUnrealPackageDefinitionInfo& PackageDef)
+{
+	FString PackageName = PackageDef.GetName();
+	for (TSharedRef<FUnrealSourceFile> SourceFile : PackageDef.GetAllSourceFiles())
+	{
+		for (TSharedRef<FUnrealTypeDefinitionInfo> TypeDef : SourceFile->GetDefinedClasses())
+		{
+			ResolveSuperClasses(*PackageName, UHTCastChecked<FUnrealClassDefinitionInfo>(TypeDef));
 		}
 	}
 }
@@ -6435,16 +6434,15 @@ void DefineTypes(TArray<FUnrealPackageDefinitionInfo*>& PackageDefs)
 		{
 			FResults::Try(*SourceFile, [PackageDef, &SourceFile = *SourceFile]()
 			{
-				TArray<UClass*>& AllClasses = PackageDef->GetAllClasses();
+				TArray<TSharedRef<FUnrealTypeDefinitionInfo>>& AllClasses = PackageDef->GetAllClasses();
 				UPackage* Package = PackageDef->GetPackage();
 
 				for (TSharedRef<FUnrealTypeDefinitionInfo>& TypeDef : SourceFile.GetDefinedClasses())
 				{
 					FUnrealClassDefinitionInfo& ClassDef = TypeDef->AsClassChecked();
-					UClass* ResultClass = ProcessParsedClass(ClassDef);
-					GTypeDefinitionInfoMap.Add(ResultClass, TypeDef);
-					ClassDef.SetObject(ResultClass);
-					AllClasses.Add(ResultClass);
+					ProcessParsedClass(ClassDef);
+					GTypeDefinitionInfoMap.AddNameLookup(UHTCastChecked<FUnrealObjectDefinitionInfo>(TypeDef));
+					AllClasses.Add(TypeDef);
 				}
 
 				for (TSharedRef<FUnrealTypeDefinitionInfo>& TypeDef : SourceFile.GetDefinedEnums())
@@ -6457,9 +6455,8 @@ void DefineTypes(TArray<FUnrealPackageDefinitionInfo*>& PackageDefs)
 				for (TSharedRef<FUnrealTypeDefinitionInfo>& TypeDef : SourceFile.GetDefinedStructs())
 				{
 					FUnrealScriptStructDefinitionInfo& ScriptStructDef = TypeDef->AsScriptStructChecked();
-					UScriptStruct* ResultScriptStruct = ProcessParsedStruct(ScriptStructDef);
-					ScriptStructDef.SetObject(ResultScriptStruct);
-					GTypeDefinitionInfoMap.Add(ResultScriptStruct, TypeDef);
+					ProcessParsedStruct(ScriptStructDef);
+					GTypeDefinitionInfoMap.AddNameLookup(UHTCastChecked<FUnrealObjectDefinitionInfo>(TypeDef));
 				}
 
 				static const bool bVerbose = FParse::Param(FCommandLine::Get(), TEXT("VERBOSE"));
@@ -6480,13 +6477,26 @@ void DefineTypes(TArray<FUnrealPackageDefinitionInfo*>& PackageDefs)
 
 void ResolveParents(TArray<FUnrealPackageDefinitionInfo*>& PackageDefs)
 {
-	GUObjectDef = &GTypeDefinitionInfoMap.FindChecked<FUnrealClassDefinitionInfo>(UObject::StaticClass());
-	GUClassDef = &GTypeDefinitionInfoMap.FindChecked<FUnrealClassDefinitionInfo>(UClass::StaticClass());
-	GUInterfaceDef = &GTypeDefinitionInfoMap.FindChecked<FUnrealClassDefinitionInfo>(UInterface::StaticClass());
+	GUObjectDef = &GTypeDefinitionInfoMap.FindByNameChecked<FUnrealClassDefinitionInfo>(*UObject::StaticClass()->GetFName().ToString());
+	GUClassDef = &GTypeDefinitionInfoMap.FindByNameChecked<FUnrealClassDefinitionInfo>(*UClass::StaticClass()->GetFName().ToString());
+	GUInterfaceDef = &GTypeDefinitionInfoMap.FindByNameChecked<FUnrealClassDefinitionInfo>(*UInterface::StaticClass()->GetFName().ToString());
+
+	for (TSharedRef<FUnrealTypeDefinitionInfo> TypeDef : GEngineClasses)
+	{
+		FUnrealClassDefinitionInfo& ClassDef = UHTCastChecked<FUnrealClassDefinitionInfo>(TypeDef);
+
+		// Some of the temporary classes are contained within the no export file.  When they are, a new
+		// instance of the definition will be created.  So ignore this one if the found class def doesn't
+		// match this one.
+		if (&GTypeDefinitionInfoMap.FindByNameChecked(*ClassDef.GetFName().ToString()) == &ClassDef)
+		{
+			ResolveSuperClasses(TEXT("EngineObject"), ClassDef);
+		}
+	}
 
 	for (FUnrealPackageDefinitionInfo* PackageDef : PackageDefs)
 	{
-		FResults::Try([PackageDef]() { ResolveSuperClasses(PackageDef->GetPackage()); });
+		FResults::Try([PackageDef]() { ResolveSuperClasses(*PackageDef); });
 	}
 	FResults::WaitForErrorTasks();
 }
@@ -6677,10 +6687,18 @@ void PrepareTypesForExport(TArray<FUnrealPackageDefinitionInfo*>& PackageDefs)
 	{
 		FResults::Try([PackageDef]()
 		{
-			PackageDef->CreateUObjectEngineTypes();
+			PackageDef->CreateUObjectEngineTypes(FUnrealTypeDefinitionInfo::ECreateEngineTypesPhase::Phase1);
 		});
 	}
-		
+
+	for (FUnrealPackageDefinitionInfo* PackageDef : PackageDefs)
+	{
+		FResults::Try([PackageDef]()
+		{
+			PackageDef->CreateUObjectEngineTypes(FUnrealTypeDefinitionInfo::ECreateEngineTypesPhase::Phase2);
+		});
+	}
+
 	for (FUnrealPackageDefinitionInfo* PackageDef : PackageDefs)
 	{
 		FResults::Try([PackageDef]()
@@ -6691,12 +6709,16 @@ void PrepareTypesForExport(TArray<FUnrealPackageDefinitionInfo*>& PackageDefs)
 	FResults::WaitForErrorTasks();
 
 	// Until we fix the issue with types not having sources, we need to do this
-	GTypeDefinitionInfoMap.ForAllTypes([](FUnrealTypeDefinitionInfo& TypeDef) 
-	{ 
+	GTypeDefinitionInfoMap.ForAllTypesByName([](FUnrealTypeDefinitionInfo& TypeDef)
+	{
 		if (!TypeDef.HasSource())
 		{
 			TypeDef.PostParseFinalize();
 		}
+		if (FUnrealObjectDefinitionInfo* ObjectDef = UHTCast<FUnrealObjectDefinitionInfo>(TypeDef))
+		{
+			check(ObjectDef->GetObject());
+		}
 	});
 }
 
@@ -6769,9 +6791,9 @@ void ExportToScriptPlugins(TArray<FUnrealPackageDefinitionInfo*>& PackageDefs, F
 		const FManifestModule& Module = PackageDef->GetModule();
 
 		FClassTree ClassTree(UObject::StaticClass());
-		for (UClass* Class : PackageDef->GetAllClasses())
+		for (TSharedRef<FUnrealTypeDefinitionInfo> TypeDef : PackageDef->GetAllClasses())
 		{
-			ClassTree.AddClass(Class);
+			ClassTree.AddClass(UHTCastChecked<FUnrealClassDefinitionInfo>(TypeDef).GetClass());
 		}
 		ClassTree.Validate();
 
@@ -6857,9 +6879,9 @@ ECompilationResult::Type UnrealHeaderTool_Main(const FString& ModuleInfoFilename
 	// Use -WRITEREF and -VERIFYREF to detect these changes.
 	for (TObjectIterator<UClass> ClassIt; ClassIt; ++ClassIt)
 	{
-		FUnrealClassDefinitionInfo* ClassDef = new FUnrealClassDefinitionInfo(*ClassIt);
-		ClassDef->SetObject(*ClassIt);
-		GTypeDefinitionInfoMap.Add(*ClassIt, MakeShareable<FUnrealTypeDefinitionInfo>(ClassDef));
+		TSharedRef<FUnrealClassDefinitionInfo> ClassDefRef = MakeShared<FUnrealClassDefinitionInfo>(*ClassIt);
+		GTypeDefinitionInfoMap.AddNameLookup(*ClassDefRef);
+		GEngineClasses.Add(ClassDefRef);
 	}
 
 	// Load the manifest file, giving a list of all modules to be processed, pre-sorted by dependency ordering
@@ -6945,7 +6967,7 @@ ECompilationResult::Type UnrealHeaderTool_Main(const FString& ModuleInfoFilename
 	return FResults::GetOverallResults();
 }
 
-UClass* ProcessParsedClass(FUnrealClassDefinitionInfo& ClassDef)
+void ProcessParsedClass(FUnrealClassDefinitionInfo& ClassDef)
 {
 	UPackage* Package = ClassDef.GetPackageDef().GetPackage();
 	const FString& ClassName = ClassDef.GetNameCPP();
@@ -6986,46 +7008,29 @@ UClass* ProcessParsedClass(FUnrealClassDefinitionInfo& ClassDef)
 		FUHTException::Throwf(ClassDef, TEXT("Class '%s' cannot inherit from itself"), *ClassName);
 	}
 
-	UClass* ResultClass = FEngineAPI::FindObject<UClass>(Package, *ClassNameStripped);
+	UClass* ResultClass = FEngineAPI::FindObject<UClass>(ANY_PACKAGE, *ClassNameStripped);
 
 	// if we aren't generating headers, then we shouldn't set misaligned object, since it won't get cleared
 
 	const static bool bVerboseOutput = FParse::Param(FCommandLine::Get(), TEXT("VERBOSE"));
 
-	if (ResultClass == nullptr || !ResultClass->IsNative())
+	if (ResultClass == nullptr)
 	{
-		// detect if the same class name is used in multiple packages
-		if (ResultClass == nullptr)
-		{
-			UClass* ConflictingClass = FEngineAPI::FindObject<UClass>(ANY_PACKAGE, *ClassNameStripped, true);
-			if (ConflictingClass != nullptr)
-			{
-				UE_LOG_WARNING_UHT(ClassDef, TEXT("Duplicate class name: %s also exists in file %s"), *ClassName, *ConflictingClass->GetOutermost()->GetName());
-			}
-		}
-
-		// Create new class.
-		ResultClass = new(EC_InternalUseOnlyConstructor, Package, *ClassNameStripped, RF_Public | RF_Standalone) UClass(FObjectInitializer(), nullptr);
-
-		// add CLASS_Interface flag if the class is an interface
-		// NOTE: at this pre-parsing/importing stage, we cannot know if our super class is an interface or not,
-		// we leave the validation to the main header parser
-		if (ClassDef.IsInterface())
+		if (TSharedRef<FUnrealTypeDefinitionInfo>* Existing = GTypeDefinitionInfoMap.FindByName(*ClassNameStripped))
 		{
-			ResultClass->ClassFlags |= CLASS_Interface;
+			FUHTException::Throwf(ClassDef, TEXT("Duplicate class name: %s also exists in file %s"), *ClassName, *(*Existing)->GetFilename());
 		}
 
 		if (bVerboseOutput)
 		{
-			UE_LOG(LogCompile, Log, TEXT("Imported: %s"), *ResultClass->GetFullName());
+			UE_LOG(LogCompile, Log, TEXT("Imported: %s"), *ClassDef.GetFullName());
 		}
 	}
 	else
 	{
-		ClassDef.SetInternalFlags(ResultClass->GetInternalFlags());
+		ClassDef.InitializeFromExistingUObject(ResultClass);
+		ClassDef.SetObject(ResultClass);
 	}
-
-	return ResultClass;
 }
 
 void ProcessParsedEnum(FUnrealEnumDefinitionInfo& EnumDef)
@@ -7045,16 +7050,14 @@ void ProcessParsedEnum(FUnrealEnumDefinitionInfo& EnumDef)
 	}
 }
 
-UScriptStruct* ProcessParsedStruct(FUnrealScriptStructDefinitionInfo& ScriptStructDef)
+void ProcessParsedStruct(FUnrealScriptStructDefinitionInfo& ScriptStructDef)
 {
-	UPackage* Package = ScriptStructDef.GetPackageDef().GetPackage();
 	const FString& StructName = ScriptStructDef.GetNameCPP();
-
 	FString StructNameStripped = GetClassNameWithPrefixRemoved(*StructName);
 
-	if (UScriptStruct* Existing = FEngineAPI::FindObject<UScriptStruct>(ANY_PACKAGE, *StructNameStripped))
+	if (TSharedRef<FUnrealTypeDefinitionInfo>* Existing = GTypeDefinitionInfoMap.FindByName(*StructNameStripped))
 	{
-		FUHTException::Throwf(ScriptStructDef, TEXT("Duplicate struct name: %s also exists in file %s"), *StructNameStripped, *Existing->GetOutermost()->GetName());
+		FUHTException::Throwf(ScriptStructDef, TEXT("Duplicate struct name: %s also exists in file %s"), *StructNameStripped, *(*Existing)->GetFilename());
 	}
 
 	// Check if the enum name is using a reserved keyword
@@ -7062,8 +7065,4 @@ UScriptStruct* ProcessParsedStruct(FUnrealScriptStructDefinitionInfo& ScriptStru
 	{
 		FUHTException::Throwf(ScriptStructDef, TEXT("struct: '%s' uses a reserved type name."), *StructNameStripped);
 	}
-
-	// Create enum definition.
-	UScriptStruct* Struct = new(EC_InternalUseOnlyConstructor, Package, FName(StructNameStripped), RF_Public) UScriptStruct(FObjectInitializer());
-	return Struct;
 }
diff --git a/Engine/Source/Programs/UnrealHeaderTool/Private/HeaderParser.cpp b/Engine/Source/Programs/UnrealHeaderTool/Private/HeaderParser.cpp
--- a/Engine/Source/Programs/UnrealHeaderTool/Private/HeaderParser.cpp
+++ b/Engine/Source/Programs/UnrealHeaderTool/Private/HeaderParser.cpp
@@ -73,7 +73,6 @@ const FName FHeaderParserNames::NAME_HideCategories(TEXT("HideCategories"));
 const FName FHeaderParserNames::NAME_ShowCategories(TEXT("ShowCategories"));
 const FName FHeaderParserNames::NAME_SparseClassDataTypes(TEXT("SparseClassDataTypes"));
 const FName FHeaderParserNames::NAME_IsConversionRoot(TEXT("IsConversionRoot"));
-const FName FHeaderParserNames::NAME_AdvancedClassDisplay(TEXT("AdvancedClassDisplay"));
 const FName FHeaderParserNames::NAME_BlueprintType(TEXT("BlueprintType"));
 const FName FHeaderParserNames::NAME_AutoCollapseCategories(TEXT("AutoCollapseCategories"));
 const FName FHeaderParserNames::NAME_HideFunctions(TEXT("HideFunctions"));
@@ -1718,7 +1717,7 @@ FUnrealScriptStructDefinitionInfo& FHeaderParser::CompileStructDeclaration()
 	bool IsNative = false;
 	bool IsExport = false;
 	bool IsTransient = false;
-	uint32 StructFlags = STRUCT_Native;
+	int32 StructFlags = STRUCT_Native;
 	TMap<FName, FString> MetaData;
 
 	// Get the struct specifier list
@@ -1755,7 +1754,6 @@ FUnrealScriptStructDefinitionInfo& FHeaderParser::CompileStructDeclaration()
 
 	// Create.
 	FUnrealScriptStructDefinitionInfo& StructDef = GTypeDefinitionInfoMap.FindByNameChecked<FUnrealScriptStructDefinitionInfo>(*StructNameStripped);
-	UScriptStruct* Struct = StructDef.GetScriptStruct();
 
 	// Process the list of specifiers
 	for (const FPropertySpecifier& Specifier : SpecifiersFound)
@@ -1799,9 +1797,10 @@ FUnrealScriptStructDefinitionInfo& FHeaderParser::CompileStructDeclaration()
 	ParseInheritance(TEXT("struct"), [](const TCHAR* StructName, bool bIsSuperClass) {}); // Eat the results, already been parsed
 
 	// if we have a base struct, propagate inherited struct flags now
-	FUnrealFieldDefinitionInfo* BaseStructDef = nullptr;
 	if (!StructDef.GetSuperStructInfo().Name.IsEmpty())
 	{
+		FUnrealScriptStructDefinitionInfo* BaseStructDef = nullptr;
+
 		const FString& ParentStructNameInScript = StructDef.GetSuperStructInfo().Name;
 		FString ParentStructNameStripped;
 		bool bOverrideParentStructName = false;
@@ -1815,13 +1814,13 @@ FUnrealScriptStructDefinitionInfo& FHeaderParser::CompileStructDeclaration()
 		// If we're expecting a prefix, first try finding the correct field with the stripped struct name
 		if (bOverrideParentStructName)
 		{
-			BaseStructDef = Scope->FindTypeByName(*ParentStructNameStripped);
+			BaseStructDef = UHTCast<FUnrealScriptStructDefinitionInfo>(Scope->FindTypeByName(*ParentStructNameStripped));
 		}
 
 		// If it wasn't found, try to find the literal name given
 		if (BaseStructDef == nullptr)
 		{
-			BaseStructDef = Scope->FindTypeByName(*ParentStructNameInScript);
+			BaseStructDef = UHTCast<FUnrealScriptStructDefinitionInfo>(Scope->FindTypeByName(*ParentStructNameInScript));
 		}
 
 		// Try to locate globally  //@TODO: UCREMOVAL: This seems extreme
@@ -1839,31 +1838,27 @@ FUnrealScriptStructDefinitionInfo& FHeaderParser::CompileStructDeclaration()
 		}
 
 		// If the struct still wasn't found, throw an error
-		if (BaseStructDef == NULL)
+		if (BaseStructDef == nullptr)
 		{
 			FUHTException::Throwf(*this, TEXT("'struct': Can't find struct '%s'"), *ParentStructNameInScript);
 		}
-		else
+
+		// If the struct was found, confirm it adheres to the correct syntax. This should always fail if we were expecting an override that was not found.
+		if (bOverrideParentStructName)
 		{
-			// If the struct was found, confirm it adheres to the correct syntax. This should always fail if we were expecting an override that was not found.
-			if (bOverrideParentStructName)
+			const TCHAR* PrefixCPP = UHTConfig.StructsWithTPrefix.Contains(ParentStructNameStripped) ? TEXT("T") : BaseStructDef->GetPrefixCPP();
+			if (ParentStructNameInScript != FString::Printf(TEXT("%s%s"), PrefixCPP, *ParentStructNameStripped))
 			{
-				const TCHAR* PrefixCPP = UHTConfig.StructsWithTPrefix.Contains(ParentStructNameStripped) ? TEXT("T") : BaseStructDef->AsScriptStructChecked().GetPrefixCPP();
-				if (ParentStructNameInScript != FString::Printf(TEXT("%s%s"), PrefixCPP, *ParentStructNameStripped))
-				{
-					FUHTException::Throwf(*this, TEXT("Parent Struct '%s' is missing a valid Unreal prefix, expecting '%s'"), *ParentStructNameInScript, *FString::Printf(TEXT("%s%s"), PrefixCPP, *BaseStructDef->GetName()));
-					BaseStructDef = nullptr;
-				}
+				FUHTException::Throwf(*this, TEXT("Parent Struct '%s' is missing a valid Unreal prefix, expecting '%s'"), *ParentStructNameInScript, *FString::Printf(TEXT("%s%s"), PrefixCPP, *BaseStructDef->GetName()));
 			}
 		}
-	}
 
-	if (BaseStructDef != NULL)
-	{
-		UScriptStruct* BaseStruct = BaseStructDef->AsScriptStructChecked().GetScriptStruct();
-		Struct->SetSuperStruct(BaseStruct);
-		StructFlags |= (BaseStruct->StructFlags & STRUCT_Inherit);
 		StructDef.GetSuperStructInfo().Struct = BaseStructDef->AsStruct();
+		StructFlags |= (BaseStructDef->GetStructFlags() & STRUCT_Inherit);
+		if (StructDef.GetScriptStructSafe() != nullptr && BaseStructDef->GetScriptStructSafe())
+		{
+			StructDef.GetScriptStructSafe()->SetSuperStruct(BaseStructDef->GetScriptStructSafe());
+		}
 	}
 
 	Scope->AddType(StructDef);
@@ -1890,7 +1885,7 @@ FUnrealScriptStructDefinitionInfo& FHeaderParser::CompileStructDeclaration()
 		FUHTException::Throwf(*this, TEXT("Struct '%s' is missing a valid Unreal prefix, expecting '%s'"), *StructNameInScript, *ExpectedStructName);
 	}
 
-	Struct->StructFlags = EStructFlags(Struct->StructFlags | StructFlags);
+	StructDef.SetStructFlags(EStructFlags(StructFlags));
 
 	AddFormattedPrevCommentAsTooltipMetaData(MetaData);
 
@@ -2078,10 +2073,6 @@ FUnrealScriptStructDefinitionInfo& FHeaderParser::CompileStructDeclaration()
 
 	// Validate sparse class data
 	CheckSparseClassData(StructDef);
-
-	// Link the properties within the struct
-	Struct->PrepareCppStructOps();
-
 	return StructDef;
 }
 
@@ -2245,13 +2236,6 @@ void FHeaderParser::PopNest(ENestType NestType, const TCHAR* Descr)
 		FUHTException::Throwf(*this, TEXT("Bad first pass NestType %i"), (uint8)NestType);
 	}
 
-	bool bLinkProps = true;
-	if (NestType == ENestType::Class)
-	{
-		UClass* TopClass = GetCurrentClass();
-		bLinkProps = !TopClass->HasAnyClassFlags(CLASS_Intrinsic);
-	}
-
 	// Pop the nesting level.
 	NestType = TopNest->NestType;
 	NestLevel--;
@@ -2303,23 +2287,18 @@ void FHeaderParser::FixupDelegateProperties(FUnrealStructDefinitionInfo& StructD
 				{
 					SourceDelegateFunctionDef = FoundDef->AsFunction();
 				}
+				if (SourceDelegateFunctionDef == nullptr && DelegatePropertyBase.DelegateSignatureOwnerClassDef != nullptr)
+				{
+					SourceDelegateFunctionDef = FindFunction(*DelegatePropertyBase.DelegateSignatureOwnerClassDef, *NameOfDelegateFunction, true, nullptr);
+				}
 				if (SourceDelegateFunctionDef == nullptr)
 				{
-					// Try to find in other packages.
-					UObject* DelegateSignatureOuter = DelegatePropertyBase.DelegateSignatureOwnerClass
-						? ((UObject*)DelegatePropertyBase.DelegateSignatureOwnerClass)
-						: ((UObject*)ANY_PACKAGE);
-					UFunction* SourceDelegateFunction = FEngineAPI::FindObject<UFunction>(DelegateSignatureOuter, *NameOfDelegateFunction);
-
-					if (SourceDelegateFunction == nullptr)
-					{
-						// convert this into a fully qualified path name for the error message.
-						NameOfDelegateFunction = Scope.GetName().ToString() + TEXT(".") + NameOfDelegateFunction;
-					}
-					else
-					{
-						SourceDelegateFunctionDef = &GTypeDefinitionInfoMap.FindChecked<FUnrealFunctionDefinitionInfo>(SourceDelegateFunction);
-					}
+					SourceDelegateFunctionDef = UHTCast<FUnrealFunctionDefinitionInfo>(GTypeDefinitionInfoMap.FindByName(*NameOfDelegateFunction));
+				}
+				if (SourceDelegateFunctionDef == nullptr)
+				{
+					// convert this into a fully qualified path name for the error message.
+					NameOfDelegateFunction = Scope.GetName().ToString() + TEXT(".") + NameOfDelegateFunction;
 				}
 			}
 			else
@@ -2501,8 +2480,7 @@ void FHeaderParser::CheckSparseClassData(const FUnrealStructDefinitionInfo& Stru
 		{
 			if (FUnrealScriptStructDefinitionInfo* ParentSparseClassDataStructDef = GTypeDefinitionInfoMap.FindByName<FUnrealScriptStructDefinitionInfo>(*ParentSparseClassDataTypeName))
 			{
-				UScriptStruct* ParentSparseClassDataStruct = ParentSparseClassDataStructDef->GetScriptStruct();
-				if (!SparseClassDataStructDef->GetScriptStruct()->IsChildOf(ParentSparseClassDataStruct))
+				if (!SparseClassDataStructDef->IsChildOf(*ParentSparseClassDataStructDef))
 				{
 					FUHTException::Throwf(StructDef, TEXT("Class %s is a child of %s but its sparse class data struct, %s, does not inherit from %s."), *ClassDef->GetName(), *ParentClassDef->GetName(), *SparseClassDataStructDef->GetName(), *ParentSparseClassDataStructDef->GetName());
 				}
@@ -2513,7 +2491,7 @@ void FHeaderParser::CheckSparseClassData(const FUnrealStructDefinitionInfo& Stru
 
 void FHeaderParser::ValidateClassFlags(const FUnrealClassDefinitionInfo& ToValidate)
 {
-	if (ToValidate.HasAnyClassFlags(CLASS_NeedsDeferredDependencyLoading) && !ToValidate.GetClass()->IsChildOf(UClass::StaticClass()))
+	if (ToValidate.HasAnyClassFlags(CLASS_NeedsDeferredDependencyLoading) && !ToValidate.IsChildOf(*GUClassDef))
 	{
 		// CLASS_NeedsDeferredDependencyLoading can only be set on classes derived from UClass
 		FUHTException::Throwf(ToValidate, TEXT("NeedsDeferredDependencyLoading is set on %s but the flag can only be used with classes derived from UClass."), *ToValidate.GetClass()->GetName());
@@ -2856,7 +2834,9 @@ void FHeaderParser::GetVarType(
 	ELayoutMacroType*               OutLayoutMacroType
 )
 {
-	UStruct* OwnerStruct = Scope->IsFileScope() ? nullptr : ((FStructScope*)Scope)->GetStruct();
+	FUnrealStructDefinitionInfo* OwnerStructDef = Scope->IsFileScope() ? nullptr : &((FStructScope*)Scope)->GetStructDef();
+	FUnrealScriptStructDefinitionInfo* OwnerScriptStructDef = UHTCast<FUnrealScriptStructDefinitionInfo>(OwnerStructDef);
+	FUnrealClassDefinitionInfo* OwnerClassDef = UHTCast<FUnrealClassDefinitionInfo>(OwnerStructDef);
 	FName RepCallbackName = FName(NAME_None);
 
 	// Get flags.
@@ -2866,9 +2846,9 @@ void FHeaderParser::GetVarType(
 	// force members to be 'blueprint read only' if in a const class
 	if (VariableCategory == EVariableCategory::Member)
 	{
-		if (UClass* OwnerClass = Cast<UClass>(OwnerStruct))
+		if (OwnerClassDef)
 		{
-			if (OwnerClass->ClassFlags & CLASS_Const)
+			if (OwnerClassDef->HasAnyClassFlags(CLASS_Const))
 			{
 				ImpliedFlags |= CPF_BlueprintReadOnly;
 			}
@@ -3008,7 +2988,7 @@ void FHeaderParser::GetVarType(
 						UE_LOG_ERROR_UHT(*this, TEXT("BlueprintReadWrite should not be used on private members"));
 					}
 
-					if ((Flags & CPF_EditorOnly) != 0 && OwnerStruct->IsA<UScriptStruct>())
+					if ((Flags & CPF_EditorOnly) != 0 && OwnerScriptStructDef != nullptr)
 					{
 						UE_LOG_ERROR_UHT(*this, TEXT("Blueprint exposed struct members cannot be editor only"));
 					}
@@ -3025,7 +3005,7 @@ void FHeaderParser::GetVarType(
 						UE_LOG_ERROR_UHT(*this, TEXT("Cannot specify a property as being both BlueprintReadOnly and having a BlueprintSetter."));
 					}
 
-					if (OwnerStruct->IsA<UScriptStruct>())
+					if (OwnerScriptStructDef != nullptr)
 					{
 						UE_LOG_ERROR_UHT(*this, TEXT("Cannot specify BlueprintSetter for a struct member."))
 					}
@@ -3052,7 +3032,7 @@ void FHeaderParser::GetVarType(
 						UE_LOG_ERROR_UHT(*this, TEXT("BlueprintReadOnly should not be used on private members"));
 					}
 
-					if ((Flags & CPF_EditorOnly) != 0 && OwnerStruct->IsA<UScriptStruct>())
+					if ((Flags & CPF_EditorOnly) != 0 && OwnerScriptStructDef != nullptr)
 					{
 						UE_LOG_ERROR_UHT(*this, TEXT("Blueprint exposed struct members cannot be editor only"));
 					}
@@ -3065,7 +3045,7 @@ void FHeaderParser::GetVarType(
 
 				case EVariableSpecifier::BlueprintGetter:
 				{
-					if (OwnerStruct->IsA<UScriptStruct>())
+					if (OwnerScriptStructDef != nullptr)
 					{
 						UE_LOG_ERROR_UHT(*this, TEXT("Cannot specify BlueprintGetter for a struct member."))
 					}
@@ -3154,7 +3134,7 @@ void FHeaderParser::GetVarType(
 				case EVariableSpecifier::Replicated:
 				case EVariableSpecifier::ReplicatedUsing:
 				{
-					if (OwnerStruct->IsA<UScriptStruct>())
+					if (OwnerScriptStructDef != nullptr)
 					{
 						UE_LOG_ERROR_UHT(*this, TEXT("Struct members cannot be replicated"));
 					}
@@ -3172,7 +3152,7 @@ void FHeaderParser::GetVarType(
 
 				case EVariableSpecifier::NotReplicated:
 				{
-					if (!OwnerStruct->IsA<UScriptStruct>())
+					if (OwnerScriptStructDef == nullptr)
 					{
 						UE_LOG_ERROR_UHT(*this, TEXT("Only Struct members can be marked NotReplicated"));
 					}
@@ -3557,8 +3537,7 @@ void FHeaderParser::GetVarType(
 	}
 	else if (VarType.Matches(TEXT("FLargeWorldCoordinatesReal"), ESearchCase::CaseSensitive))		// LWC_TODO: Remove with UE_LARGE_WORLD_COORDINATES_DISABLED toggle.
 	{
-		const FString& ForSourceFile = Scope->GetFileScope()->GetSourceFile()->GetFilename();
-		if(!ForSourceFile.EndsWith(TEXT("NoExportTypes.h")))
+		if (!SourceFile.GetFilename().EndsWith(TEXT("NoExportTypes.h")))
 		{
 			FUHTException::Throwf(*this, TEXT("FLargeWorldCoordinatesReal is intended for LWC support only and should not be used outside of NoExportTypes.h"));
 		}
@@ -3961,7 +3940,7 @@ void FHeaderParser::GetVarType(
 						if (FUnrealFunctionDefinitionInfo* DelegateFuncDef = FoundDef->AsFunction())
 						{
 							SetDelegateType(*DelegateFuncDef, DelegateIdentifierStripped);
-							VarProperty.DelegateSignatureOwnerClass = LocalOwnerClassDef->GetClass();
+							VarProperty.DelegateSignatureOwnerClassDef = LocalOwnerClassDef;
 						}
 					}
 				}
@@ -3977,7 +3956,6 @@ void FHeaderParser::GetVarType(
 		}
 		else if (StructDef)
 		{
-			UScriptStruct* Struct = StructDef->GetScriptStruct();
 			if (bStripped)
 			{
 				const TCHAR* PrefixCPP = UHTConfig.StructsWithTPrefix.Contains(IdentifierStripped) ? TEXT("T") : StructDef->GetPrefixCPP();
@@ -4128,7 +4106,7 @@ void FHeaderParser::GetVarType(
 					}
 
 					VarProperty = FPropertyBase(*TempClassDef, PropertyType, bWeakIsAuto);
-					if (TempClassDef->GetClass()->IsChildOf(UClass::StaticClass()))
+					if (TempClassDef->IsChildOf(*GUClassDef))
 					{
 						if (MatchSymbol(TEXT('<')))
 						{
@@ -4270,7 +4248,7 @@ void FHeaderParser::GetVarType(
 	}
 
 	//@TODO: UCREMOVAL: 'const' member variables that will get written post-construction by defaultproperties
-	if (VariableCategory == EVariableCategory::Member && OwnerStruct->IsA<UClass>() && ((UClass*)OwnerStruct)->HasAnyClassFlags(CLASS_Const))
+	if (VariableCategory == EVariableCategory::Member && OwnerClassDef != nullptr && OwnerClassDef->HasAnyClassFlags(CLASS_Const))
 	{
 		// Eat a 'not quite truthful' const after the type; autogenerated for member variables of const classes.
 		bNativeConst |= MatchIdentifier(TEXT("const"), ESearchCase::CaseSensitive);
@@ -4345,7 +4323,7 @@ void FHeaderParser::GetVarType(
 	{
 		if ((VarProperty.Type == CPT_ObjectReference) || (VarProperty.Type == CPT_ObjectPtrReference))
 		{
-			if (VarProperty.ClassDef->GetClass()->IsChildOf<UClass>())
+			if (VarProperty.ClassDef->IsChildOf(*GUClassDef))
 			{
 				FUHTException::Throwf(*this, TEXT("'Instanced' cannot be applied to class properties (UClass* or TSubclassOf<>)"));
 			}
@@ -4384,7 +4362,7 @@ void FHeaderParser::GetVarType(
 
 	// Check for invalid transients
 	EPropertyFlags Transients = VarProperty.PropertyFlags & (CPF_DuplicateTransient | CPF_TextExportTransient | CPF_NonPIEDuplicateTransient);
-	if (Transients && !Cast<UClass>(OwnerStruct))
+	if (Transients && OwnerClassDef == nullptr)
 	{
 		TArray<const TCHAR*> FlagStrs = ParsePropertyFlags(Transients);
 		FUHTException::Throwf(*this, TEXT("'%s' specifier(s) are only allowed on class member variables"), *FString::Join(FlagStrs, TEXT(", ")));
@@ -4686,13 +4664,13 @@ bool FHeaderParser::CompileDeclaration(TArray<FUnrealFunctionDefinitionInfo*>& D
 		FStructMetaData& StructData = GetCurrentClassData();
 		if (StructData.ParsedInterface == EParsedInterface::NotAnInterface)
 		{
-			FString CurrentClassName = GetCurrentClass()->GetName();
+			FString CurrentClassName = GetCurrentClassDef().GetName();
 			FUHTException::Throwf(*this, TEXT("Could not find the associated 'U%s' class while parsing 'I%s' - it could be missing or malformed"), *CurrentClassName, *CurrentClassName);
 		}
 
 		if (StructData.ParsedInterface == EParsedInterface::ParsedIInterface)
 		{
-			FString CurrentClassName = GetCurrentClass()->GetName();
+			FString CurrentClassName = GetCurrentClassDef().GetName();
 			FUHTException::Throwf(*this, TEXT("Duplicate IInterface definition found while parsing 'I%s'"), *CurrentClassName);
 		}
 
@@ -4846,11 +4824,10 @@ bool FHeaderParser::CompileDeclaration(TArray<FUnrealFunctionDefinitionInfo*>& D
 		// Instead, the whole statement-consumer loop should be in a nest
 		bEncounteredNewStyleClass_UnmatchedBrackets = false;
 
-		UClass* CurrentClass = GetCurrentClass();
 		FUnrealClassDefinitionInfo& CurrentClassDef = UHTCastChecked<FUnrealClassDefinitionInfo>(GetCurrentClassDef());
 
 		// Pop nesting here to allow other non UClass declarations in the header file.
-		if (CurrentClass->ClassFlags & CLASS_Interface)
+		if (CurrentClassDef.HasAllClassFlags(CLASS_Interface))
 		{
 			checkf(TopNest->NestType == ENestType::Interface || TopNest->NestType == ENestType::NativeInterface, TEXT("Unexpected end of interface block."));
 			PopNest(TopNest->NestType, TEXT("'Interface'"));
@@ -5195,8 +5172,7 @@ FUnrealClassDefinitionInfo& FHeaderParser::ParseClassNameDeclaration(FString& De
 
 	FUnrealClassDefinitionInfo* ClassDef = FUnrealClassDefinitionInfo::FindClass(*GetClassNameWithPrefixRemoved(*DeclaredClassName));
 	check(ClassDef);
-	UClass* FoundClass = ClassDef->GetClass();
-	FoundClass->ClassCastFlags |= ClassCastFlagMap::Get().GetCastFlag(DeclaredClassName);
+	ClassDef->SetClassCastFlags(ClassCastFlagMap::Get().GetCastFlag(DeclaredClassName));
 
 	// Skip optional final keyword
 	MatchIdentifier(TEXT("final"), ESearchCase::CaseSensitive);
@@ -5207,9 +5183,11 @@ FUnrealClassDefinitionInfo& FHeaderParser::ParseClassNameDeclaration(FString& De
 	// Collect data about the super and base classes
 	if (FUnrealClassDefinitionInfo* SuperClassDef = UHTCast<FUnrealClassDefinitionInfo>(ClassDef->GetSuperStructInfo().Struct))
 	{
-		FoundClass->ClassCastFlags |= SuperClassDef->GetClass()->ClassCastFlags;
+		ClassDef->SetClassCastFlags(SuperClassDef->GetClassCastFlags());
 	}
-	for (FUnrealStructDefinitionInfo::FBaseStructInfo& BaseClassInfo : ClassDef->GetBaseStructInfo())
+
+	// Add the class flags from the interfaces
+	for (FUnrealStructDefinitionInfo::FBaseStructInfo& BaseClassInfo : ClassDef->GetBaseStructInfos())
 	{
 		if (FUnrealClassDefinitionInfo* BaseClassDef = UHTCast<FUnrealClassDefinitionInfo>(BaseClassInfo.Struct))
 		{
@@ -5220,7 +5198,6 @@ FUnrealClassDefinitionInfo& FHeaderParser::ParseClassNameDeclaration(FString& De
 
 			// Propagate the inheritable ClassFlags
 			ClassDef->SetClassFlags(BaseClassDef->GetClassFlags() & CLASS_ScriptInherit);
-			FoundClass->Interfaces.Emplace(BaseClassDef->GetClass(), 0, false);
 		}
 	}
 	return *ClassDef;
@@ -5305,21 +5282,12 @@ FUnrealClassDefinitionInfo& FHeaderParser::CompileClassDeclaration()
 	// If we have existing class flags (preexisting engine objects), make sure we have some shared flags.
 	check(ClassDef.GetClassFlags() == 0 || (ClassDef.GetClassFlags() & ClassDef.GetParsedClassFlags()) != 0);
 
-	const EClassFlags PrevClassFlags = ClassDef.GetInitialEngineClassFlags();
-
 	ClassDef.SetClassFlags(CLASS_Parsed);
 
 	PushNest(ENestType::Class, &ClassDef);
 	
 	ResetClassData();
 
-	// Verify class variables haven't been filled in
-	UClass* Class = ClassDef.GetClass();
-	check(Class->Children == NULL);
-	check(Class->Next == NULL);
-	check(Class->NetFields.Num() == 0);
-	check(Class->FirstOwnedClassRep == 0);
-
 	// Make sure our parent classes is parsed.
 	for (FUnrealClassDefinitionInfo* TempDef = ClassDef.GetSuperClass(); TempDef; TempDef = TempDef->GetSuperClass())
 	{
@@ -5338,11 +5306,11 @@ FUnrealClassDefinitionInfo& FHeaderParser::CompileClassDeclaration()
 	FStructMetaData& StructData = ClassDef.GetStructMetaData();
 	StructData.SetPrologLine(PrologFinishLine);
 
-	ClassDef.MergeAndValidateClassFlags(DeclaredClassName, PrevClassFlags);
+	ClassDef.MergeAndValidateClassFlags(DeclaredClassName);
 	ClassDef.SetInternalFlags(EInternalObjectFlags::Native);
 
 	// Class metadata
-	MetaData.Append(ClassDef.MetaData);
+	MetaData.Append(ClassDef.GetParsedMetaData());
 	ClassDef.MergeCategoryMetaData(MetaData);
 	AddIncludePathToMetadata(ClassDef, MetaData);
 	AddModuleRelativePathToMetadata(ClassDef, MetaData);
@@ -5353,49 +5321,15 @@ FUnrealClassDefinitionInfo& FHeaderParser::CompileClassDeclaration()
 	// Handle the start of the rest of the class
 	RequireSymbol( TEXT('{'), TEXT("'Class'") );
 
-	// Make visible outside the package.
-	Class->ClearFlags(RF_Transient);
-	check(Class->HasAnyFlags(RF_Public));
-	check(Class->HasAnyFlags(RF_Standalone));
-
 	// Copy properties from parent class.
-	if (Class->GetSuperClass())
+	if (FUnrealClassDefinitionInfo* SuperClassDef = ClassDef.GetSuperClass())
 	{
-		Class->SetPropertiesSize(Class->GetSuperClass()->GetPropertiesSize());
-	}
-
-	// auto-create properties for all of the VFTables needed for the multiple inheritances
-	// get the inheritance parents
-	const TArray<FUnrealStructDefinitionInfo::FBaseStructInfo>& BaseClassInfos = ClassDef.GetBaseStructInfo();
-
-	// for all base class types, make a VfTable property
-	for (int32 ParentIndex = BaseClassInfos.Num() - 1; ParentIndex >= 0; ParentIndex--)
-	{
-		// if this base class corresponds to an interface class, assign the vtable FProperty in the class's Interfaces map now...
-		if (FUnrealStructDefinitionInfo* BaseClassStruct = BaseClassInfos[ParentIndex].Struct)
-		{
-			UClass* BaseClass = UHTCastChecked<FUnrealClassDefinitionInfo>(BaseClassStruct).GetClass();
-			if (BaseClass->HasAnyClassFlags(CLASS_Native))
-			{
-				FImplementedInterface* Found = Class->Interfaces.FindByPredicate([&BaseClass](const FImplementedInterface& Impl) { return Impl.Class == BaseClass; });
-				if (Found)
-				{
-					Found->PointerOffset = 1;
-				}
-				else
-				{
-					Class->Interfaces.Add(FImplementedInterface(BaseClass, 1, false));
-				}
-			}
-		}
+		ClassDef.SetPropertiesSize(SuperClassDef->GetPropertiesSize());
 	}
 
 	// Validate sparse class data
 	CheckSparseClassData(ClassDef);
 
-	// Check that the class has appropriate class flags set
-	ValidateClassFlags(ClassDef);
-
 	return ClassDef;
 }
 
@@ -5505,7 +5439,6 @@ void FHeaderParser::CompileInterfaceDeclaration()
 	RequireIdentifier(TEXT("class"), ESearchCase::CaseSensitive, TEXT("Interface declaration"));
 	FUnrealClassDefinitionInfo* InterfaceClassDef = ParseInterfaceNameDeclaration(/*out*/ DeclaredInterfaceName, /*out*/ RequiredAPIMacroIfPresent);
 	check(InterfaceClassDef);
-	UClass* InterfaceClass = InterfaceClassDef->GetClass();
 	InterfaceClassDef->GetDefinitionRange().Start = &Input[InputPos];
 
 	// Record that this interface is RequiredAPI if the CORE_API style macro was present
@@ -5516,17 +5449,20 @@ void FHeaderParser::CompileInterfaceDeclaration()
 
 	// Set the appropriate interface class flags
 	InterfaceClassDef->SetClassFlags(CLASS_Interface | CLASS_Abstract);
-	if (InterfaceClass->GetSuperClass() != NULL)
+	if (FUnrealClassDefinitionInfo* InterfaceSuperClassDef = InterfaceClassDef->GetSuperClass())
 	{
-		InterfaceClass->ClassCastFlags |= InterfaceClass->GetSuperClass()->ClassCastFlags;
+		InterfaceClassDef->SetClassCastFlags(InterfaceSuperClassDef->GetClassCastFlags());
+		// All classes that are parsed are expected to be native
+		if (!InterfaceSuperClassDef->HasAnyClassFlags(CLASS_Native))
+		{
+			FUHTException::Throwf(*this, TEXT("Native classes cannot extend non-native classes"));
+		}
+		InterfaceClassDef->SetClassWithin(InterfaceSuperClassDef->GetClassWithin());
 	}
-
-	// All classes that are parsed are expected to be native
-	if (InterfaceClassDef->GetSuperClass() && !InterfaceClassDef->GetSuperClass()->HasAnyClassFlags(CLASS_Native))
+	else
 	{
-		FUHTException::Throwf(*this, TEXT("Native classes cannot extend non-native classes") );
+		InterfaceClassDef->SetClassWithin(GUObjectDef);
 	}
-
 	InterfaceClassDef->SetInternalFlags(EInternalObjectFlags::Native);
 	InterfaceClassDef->SetClassFlags(CLASS_Native);
 
@@ -5579,11 +5515,6 @@ void FHeaderParser::CompileInterfaceDeclaration()
 	// Handle the start of the rest of the interface
 	RequireSymbol( TEXT('{'), TEXT("'Class'") );
 
-	// Make visible outside the package.
-	InterfaceClass->ClearFlags(RF_Transient);
-	check(InterfaceClass->HasAnyFlags(RF_Public));
-	check(InterfaceClass->HasAnyFlags(RF_Standalone));
-
 	// Push the interface class nesting.
 	// we need a more specific set of allow flags for ENestType::Interface, only function declaration is allowed, no other stuff are allowed
 	PushNest(ENestType::Interface, InterfaceClassDef);
@@ -6225,7 +6156,8 @@ FUnrealFunctionDefinitionInfo& FHeaderParser::CompileFunctionDeclaration()
 	TArray<FPropertySpecifier> SpecifiersFound;
 	ReadSpecifierSetInsideMacro(SpecifiersFound, TEXT("Function"), MetaData);
 
-	if (!GetCurrentClass()->HasAnyClassFlags(CLASS_Native))
+	FUnrealClassDefinitionInfo& OuterClassDef = UHTCastChecked<FUnrealClassDefinitionInfo>(GetCurrentClassDef());
+	if (!OuterClassDef.HasAnyClassFlags(CLASS_Native))
 	{
 		FUHTException::Throwf(*this, TEXT("Should only be here for native classes!"));
 	}
@@ -6262,7 +6194,7 @@ FUnrealFunctionDefinitionInfo& FHeaderParser::CompileFunctionDeclaration()
 	}
 
 	// non-static functions in a const class must be const themselves
-	if (GetCurrentClass()->HasAnyClassFlags(CLASS_Const))
+	if (OuterClassDef.HasAnyClassFlags(CLASS_Const))
 	{
 		FuncInfo.FunctionFlags |= FUNC_Const;
 	}
@@ -6292,7 +6224,7 @@ FUnrealFunctionDefinitionInfo& FHeaderParser::CompileFunctionDeclaration()
 		MetaData.Add(NAME_CustomThunk, TEXT("true"));
 	}
 
-	if ((FuncInfo.FunctionFlags & FUNC_BlueprintPure) && GetCurrentClass()->HasAnyClassFlags(CLASS_Interface))
+	if ((FuncInfo.FunctionFlags & FUNC_BlueprintPure) && OuterClassDef.HasAnyClassFlags(CLASS_Interface))
 	{
 		// Until pure interface casts are supported, we don't allow pures in interfaces
 		UE_LOG_ERROR_UHT(*this, TEXT("BlueprintPure specifier is not allowed for interface functions"));
@@ -6339,11 +6271,11 @@ FUnrealFunctionDefinitionInfo& FHeaderParser::CompileFunctionDeclaration()
 	}
 
 	// Verify interfaces with respect to their blueprint accessible functions
-	if (GetCurrentClass()->HasAnyClassFlags(CLASS_Interface))
+	if (OuterClassDef.HasAnyClassFlags(CLASS_Interface))
 	{
 		if((FuncInfo.FunctionFlags & FUNC_BlueprintEvent) != 0 && !bInternalOnly)
 		{
-			const bool bCanImplementInBlueprints = !GetCurrentClass()->HasMetaData(NAME_CannotImplementInterfaceInBlueprint);  //FBlueprintMetadata::MD_CannotImplementInterfaceInBlueprint
+			const bool bCanImplementInBlueprints = !OuterClassDef.HasMetaData(NAME_CannotImplementInterfaceInBlueprint);  //FBlueprintMetadata::MD_CannotImplementInterfaceInBlueprint
 
 			// Ensure that blueprint events are only allowed in implementable interfaces. Internal only functions allowed
 			if (!bCanImplementInBlueprints)
@@ -6354,7 +6286,7 @@ FUnrealFunctionDefinitionInfo& FHeaderParser::CompileFunctionDeclaration()
 		
 		if (((FuncInfo.FunctionFlags & FUNC_BlueprintCallable) != 0) && (((~FuncInfo.FunctionFlags) & FUNC_BlueprintEvent) != 0))
 		{
-			const bool bCanImplementInBlueprints = !GetCurrentClass()->HasMetaData(NAME_CannotImplementInterfaceInBlueprint);  //FBlueprintMetadata::MD_CannotImplementInterfaceInBlueprint
+			const bool bCanImplementInBlueprints = !OuterClassDef.HasMetaData(NAME_CannotImplementInterfaceInBlueprint);  //FBlueprintMetadata::MD_CannotImplementInterfaceInBlueprint
 
 			// Ensure that if this interface contains blueprint callable functions that are not blueprint defined, that it must be implemented natively
 			if (bCanImplementInBlueprints)
@@ -6379,7 +6311,7 @@ FUnrealFunctionDefinitionInfo& FHeaderParser::CompileFunctionDeclaration()
 					//@TODO: Validate the module name for RequiredAPIMacroIfPresent
 					bThrowTokenBack = false;
 
-					if (GetCurrentClass()->HasAnyClassFlags(CLASS_RequiredAPI))
+					if (OuterClassDef.HasAnyClassFlags(CLASS_RequiredAPI))
 					{
 						FUHTException::Throwf(*this, TEXT("'%s' must not be used on methods of a class that is marked '%s' itself."), *RequiredAPIMacroIfPresent, *RequiredAPIMacroIfPresent);
 					}
@@ -6418,7 +6350,7 @@ FUnrealFunctionDefinitionInfo& FHeaderParser::CompileFunctionDeclaration()
 		// if this is a BlueprintNativeEvent or BlueprintImplementableEvent in an interface, make sure it's not "virtual"
 		if (FuncInfo.FunctionFlags & FUNC_BlueprintEvent)
 		{
-			if (GetCurrentClass()->HasAnyClassFlags(CLASS_Interface))
+			if (OuterClassDef.HasAnyClassFlags(CLASS_Interface))
 			{
 				FUHTException::Throwf(*this, TEXT("BlueprintImplementableEvents in Interfaces must not be declared 'virtual'"));
 			}
@@ -6438,7 +6370,7 @@ FUnrealFunctionDefinitionInfo& FHeaderParser::CompileFunctionDeclaration()
 	else
 	{
 		// if this is a function in an Interface, it must be marked 'virtual' unless it's an event
-		if (GetCurrentClass()->HasAnyClassFlags(CLASS_Interface) && !(FuncInfo.FunctionFlags & FUNC_BlueprintEvent))
+		if (OuterClassDef.HasAnyClassFlags(CLASS_Interface) && !(FuncInfo.FunctionFlags & FUNC_BlueprintEvent))
 		{
 			FUHTException::Throwf(*this, TEXT("Interface functions that are not BlueprintImplementableEvents must be declared 'virtual'"));
 		}
@@ -6451,7 +6383,7 @@ FUnrealFunctionDefinitionInfo& FHeaderParser::CompileFunctionDeclaration()
 		FuncInfo.FunctionFlags |= FUNC_Final;
 		FuncInfo.FunctionExportFlags |= FUNCEXPORT_Final;
 
-		if (GetCurrentClass()->HasAnyClassFlags(CLASS_Interface))
+		if (OuterClassDef.HasAnyClassFlags(CLASS_Interface))
 		{
 			UE_LOG_ERROR_UHT(*this, TEXT("Interface functions cannot be declared 'final'"));
 		}
@@ -6548,17 +6480,17 @@ FUnrealFunctionDefinitionInfo& FHeaderParser::CompileFunctionDeclaration()
 	}
 
 	// Check to see if there is a function in the super class with the same name
-	UStruct* SuperStruct = GetCurrentClass();
-	if (SuperStruct)
+	FUnrealStructDefinitionInfo* SuperStructDef = &GetCurrentClassDef();
+	if (SuperStructDef)
 	{
-		SuperStruct = SuperStruct->GetSuperStruct();
+		SuperStructDef = SuperStructDef->GetSuperStruct();
 	}
-	if (SuperStruct)
+	if (SuperStructDef)
 	{
-		if (UFunction* OverriddenFunction = ::FindUField<UFunction>(SuperStruct, FuncDefFuncInfo.Function.Identifier))
+		if (FUnrealFunctionDefinitionInfo* OverriddenFunctionDef = FindFunction(*SuperStructDef, FuncDefFuncInfo.Function.Identifier, true, nullptr))
 		{
 			// Native function overrides should be done in CPP text, not in a UFUNCTION() declaration (you can't change flags, and it'd otherwise be a burden to keep them identical)
-			UE_LOG_ERROR_UHT(*this, TEXT("%s: Override of UFUNCTION in parent class (%s) cannot have a UFUNCTION() declaration above it; it will use the same parameters as the original declaration."), FuncDefFuncInfo.Function.Identifier, *OverriddenFunction->GetOuter()->GetName());
+			UE_LOG_ERROR_UHT(*this, TEXT("%s: Override of UFUNCTION in parent class (%s) cannot have a UFUNCTION() declaration above it; it will use the same parameters as the original declaration."), FuncDefFuncInfo.Function.Identifier, *OverriddenFunctionDef->GetOuter()->GetName());
 		}
 	}
 
@@ -6630,7 +6562,7 @@ FUnrealFunctionDefinitionInfo& FHeaderParser::CompileFunctionDeclaration()
 		// This is a final (prebinding, non-overridable) function
 		FuncDef.SetFunctionFlags(FUNC_Final);
 		FuncDef.SetFunctionExportFlags(FUNCEXPORT_Final);
-		if (GetCurrentClass()->HasAnyClassFlags(CLASS_Interface))
+		if (UHTCastChecked<FUnrealClassDefinitionInfo>(GetCurrentClassDef()).HasAnyClassFlags(CLASS_Interface))
 		{
 			FUHTException::Throwf(*this, TEXT("Interface functions cannot be declared 'final'"));
 		}
@@ -6936,9 +6868,9 @@ struct FExposeOnSpawnValidator
 
 void FHeaderParser::CompileVariableDeclaration(FUnrealStructDefinitionInfo& StructDef)
 {
-	UStruct* Struct = StructDef.GetStruct();
 	EPropertyFlags DisallowFlags = CPF_ParmFlags;
 	EPropertyFlags EdFlags       = CPF_None;
+	FUnrealScriptStructDefinitionInfo* ScriptStructDef = UHTCast< FUnrealScriptStructDefinitionInfo>(StructDef);
 
 	// Get variable type.
 	FPropertyBase OriginalPropertyBase(CPT_None);
@@ -7051,17 +6983,17 @@ void FHeaderParser::CompileVariableDeclaration(FUnrealStructDefinitionInfo& Stru
 		// we'll need any metadata tags we parsed later on when we call ConvertEOLCommentToTooltip() so the tags aren't clobbered
 		OriginalPropertyBase.MetaData = PropertyBase.MetaData;
 
-		if (UScriptStruct* StructBeingBuilt = Cast<UScriptStruct>(Struct))
+		if (ScriptStructDef != nullptr)
 		{
 			if (NewPropDef.ContainsInstancedObjectProperty())
 			{
-				StructBeingBuilt->StructFlags = EStructFlags(StructBeingBuilt->StructFlags | STRUCT_HasInstancedReference);
+				ScriptStructDef->SetStructFlags(STRUCT_HasInstancedReference);
 			}
 		}
 
 		if (NewPropDef.HasAnyPropertyFlags(CPF_BlueprintVisible))
 		{
-			if (Struct->IsA<UScriptStruct>() && !StructDef.GetBoolMetaDataHierarchical(FHeaderParserNames::NAME_BlueprintType))
+			if (ScriptStructDef != nullptr && !StructDef.GetBoolMetaDataHierarchical(FHeaderParserNames::NAME_BlueprintType))
 			{
 				UE_LOG_ERROR_UHT(*this, TEXT("Cannot expose property to blueprints in a struct that is not a BlueprintType. %s.%s"), *StructDef.GetName(), *NewPropDef.GetName());
 			}
@@ -7386,10 +7318,9 @@ void FHeaderParser::ParseHeader()
 	for (TSharedRef<FUnrealTypeDefinitionInfo>& TypeDef : SourceFile.GetDefinedClasses())
 	{
 		FUnrealClassDefinitionInfo& ClassDef = UHTCastChecked<FUnrealClassDefinitionInfo>(TypeDef);
-		UClass* Class = ClassDef.GetClass();
 		PostParsingClassSetup(ClassDef);
 
-		bNoExportClassesOnly = bNoExportClassesOnly && Class->HasAnyClassFlags(CLASS_NoExport);
+		bNoExportClassesOnly = bNoExportClassesOnly && ClassDef.HasAnyClassFlags(CLASS_NoExport);
 	}
 
 	if (!bSpottedAutogeneratedHeaderInclude && !bEmptyFile && !bNoExportClassesOnly)
@@ -8031,7 +7962,7 @@ void FHeaderParser::SimplifiedClassParse(FUnrealSourceFile& SourceFile, const TC
 					}
 
 					TSharedRef<FUnrealTypeDefinitionInfo> ClassDecl = Parser.ParseClassDeclaration(StartOfLine + (UInterfaceMacroDecl - Str), CurrentLine, true, TEXT("UINTERFACE"));
-					bFoundExportedClasses |= !EnumHasAnyFlags(ClassDecl->AsClassChecked().GetParsedClassFlags(), CLASS_NoExport);
+					bFoundExportedClasses |= !EnumHasAnyFlags(ClassDecl->AsClassChecked().GetParsedClassFlags(), CLASS_NoExport); //@TODO - Use proper API once we have internal class flags
 					SourceFile.AddDefinedClass(MoveTemp(ClassDecl));
 				}
 			}
@@ -8046,7 +7977,7 @@ void FHeaderParser::SimplifiedClassParse(FUnrealSourceFile& SourceFile, const TC
 					}
 
 					TSharedRef<FUnrealTypeDefinitionInfo> ClassDecl = Parser.ParseClassDeclaration(StartOfLine + (UClassMacroDecl - Str), CurrentLine, false, TEXT("UCLASS"));
-					bFoundExportedClasses |= !EnumHasAnyFlags(ClassDecl->AsClassChecked().GetParsedClassFlags(), CLASS_NoExport);
+					bFoundExportedClasses |= !EnumHasAnyFlags(ClassDecl->AsClassChecked().GetParsedClassFlags(), CLASS_NoExport); //@TODO - Use proper API once we have internal class flags
 					SourceFile.AddDefinedClass(MoveTemp(ClassDecl));
 				}
 			}
@@ -8134,12 +8065,12 @@ TSharedRef<FUnrealTypeDefinitionInfo> FHeaderPreParser::ParseClassDeclaration(co
 		}
 		else
 		{
-			ClassDef.GetBaseStructInfo().Emplace(FUnrealStructDefinitionInfo::FBaseStructInfo{ MoveTemp(ClassNameStr) });
+			ClassDef.GetBaseStructInfos().Emplace(FUnrealStructDefinitionInfo::FBaseStructInfo{ MoveTemp(ClassNameStr) });
 		}
 	}
 	);
 
-	ClassDef->MetaData = MoveTemp(MetaData);
+	ClassDef->GetParsedMetaData() = MoveTemp(MetaData);
 	ClassDef->ParseClassProperties(MoveTemp(SpecifiersFound), RequiredAPIMacroIfPresent);
 	return ClassDef;
 }
@@ -8260,7 +8191,7 @@ TSharedRef<FUnrealTypeDefinitionInfo> FHeaderPreParser::ParseStructDeclaration(c
 		}
 		else
 		{
-			StructDef.GetBaseStructInfo().Emplace(FUnrealStructDefinitionInfo::FBaseStructInfo{ MoveTemp(StructNameStr) });
+			StructDef.GetBaseStructInfos().Emplace(FUnrealStructDefinitionInfo::FBaseStructInfo{ MoveTemp(StructNameStr) });
 		}
 	}
 	);
@@ -8415,22 +8346,18 @@ void FHeaderParser::CompileVersionDeclaration(FUnrealStructDefinitionInfo& Struc
 void FHeaderParser::ResetClassData()
 {
 	FUnrealClassDefinitionInfo& CurrentClassDef = UHTCastChecked<FUnrealClassDefinitionInfo>(GetCurrentClassDef());
-	UClass* CurrentClass = CurrentClassDef.GetClass();
-	CurrentClass->PropertiesSize = 0;
 
 	// Set class flags and within.
 	CurrentClassDef.ClearClassFlags(CLASS_RecompilerClear);
 
 	if (FUnrealClassDefinitionInfo* SuperClassDef = CurrentClassDef.GetSuperClass())
 	{
-		UClass* SuperClass = SuperClassDef->GetClass();
-		CurrentClassDef.SetClassFlags(SuperClass->ClassFlags & CLASS_ScriptInherit);
-		CurrentClass->ClassConfigName = SuperClass->ClassConfigName;
-		check(SuperClass->ClassWithin);
+		CurrentClassDef.SetClassFlags(SuperClassDef->GetClassFlags() & CLASS_ScriptInherit);
+		CurrentClassDef.SetClassConfigName(SuperClassDef->GetClassConfigName());
+
 		check(SuperClassDef->GetClassWithin());
-		if (CurrentClass->ClassWithin == nullptr)
+		if (CurrentClassDef.GetClassWithin() == nullptr)
 		{
-			CurrentClass->ClassWithin = SuperClass->ClassWithin;
 			CurrentClassDef.SetClassWithin(SuperClassDef->GetClassWithin());
 		}
 
@@ -8461,7 +8388,6 @@ void FHeaderParser::ResetClassData()
 		}
 	}
 
-	check(CurrentClass->ClassWithin);
 	check(CurrentClassDef.GetClassWithin());
 }
 
@@ -8471,7 +8397,7 @@ void FHeaderParser::PostPopNestClass(FUnrealClassDefinitionInfo& CurrentClassDef
 	VerifyPropertyMarkups(CurrentClassDef);
 
 	// Iterate over all the interfaces we claim to implement
-	for (const FUnrealStructDefinitionInfo::FBaseStructInfo& BaseClassInfo : CurrentClassDef.GetBaseStructInfo())
+	for (const FUnrealStructDefinitionInfo::FBaseStructInfo& BaseClassInfo : CurrentClassDef.GetBaseStructInfos())
 	{
 
 		// Skip unknown base classes or things that aren't interfaces
diff --git a/Engine/Source/Programs/UnrealHeaderTool/Private/HeaderParser.h b/Engine/Source/Programs/UnrealHeaderTool/Private/HeaderParser.h
--- a/Engine/Source/Programs/UnrealHeaderTool/Private/HeaderParser.h
+++ b/Engine/Source/Programs/UnrealHeaderTool/Private/HeaderParser.h
@@ -228,14 +228,6 @@ protected:
 		return false;
 	}
 
-	/**
-	 * Gets current class.
-	 */
-	UClass* GetCurrentClass() const
-	{
-		return (UClass*)GetCurrentClassScope()->GetStruct();
-	}
-
 	/**
 	 * Gets current class definition.
 	 */
diff --git a/Engine/Source/Programs/UnrealHeaderTool/Private/ParserHelper.cpp b/Engine/Source/Programs/UnrealHeaderTool/Private/ParserHelper.cpp
--- a/Engine/Source/Programs/UnrealHeaderTool/Private/ParserHelper.cpp
+++ b/Engine/Source/Programs/UnrealHeaderTool/Private/ParserHelper.cpp
@@ -8,6 +8,11 @@
 #include "UnrealTypeDefinitionInfo.h"
 #include "ClassMaps.h"
 
+// Globals for common class definitions
+extern FUnrealClassDefinitionInfo* GUObjectDef;
+extern FUnrealClassDefinitionInfo* GUClassDef;
+extern FUnrealClassDefinitionInfo* GUInterfaceDef;
+
 /////////////////////////////////////////////////////
 // FPropertyBase
 
@@ -89,7 +94,7 @@ EUHTPropertyType FPropertyBase::GetUHTPropertyType() const
 
 bool FPropertyBase::IsClassRefOrClassRefStaticArray() const
 {
-	return IsObjectRefOrObjectRefStaticArray() && ClassDef->GetClass()->IsChildOf(UClass::StaticClass());
+	return IsObjectRefOrObjectRefStaticArray() && ClassDef->IsChildOf(*GUClassDef);
 }
 
 bool FPropertyBase::IsByteEnumOrByteEnumStaticArray() const
@@ -193,34 +198,30 @@ bool FPropertyBase::MatchesType(const FPropertyBase& Other, bool bDisallowGenera
 		}
 		else
 		{
-			UClass* PropertyClass = ClassDef->GetClass();
-			UClass* OtherPropertyClass = Other.ClassDef->GetClass();
-			UClass* MetaClass = MetaClassDef ? MetaClassDef->GetClass() : nullptr;
-			UClass* OtherMetaClass = Other.MetaClassDef ? Other.MetaClassDef->GetClass() : nullptr;
 			// Generalization is ok (typical example of this check would look like: VarA = VarB;, where this is VarB and Other is VarA)
-			if (OtherPropertyClass->IsChildOf(PropertyClass))
+			if (Other.ClassDef->IsChildOf(*ClassDef))
 			{
 				if (!bIgnoreImplementedInterfaces || ((Type == CPT_Interface) == (Other.Type == CPT_Interface)))
 				{
-					if (!PropertyClass->IsChildOf(UClass::StaticClass()) || MetaClass == NULL || OtherMetaClass->IsChildOf(MetaClass) ||
-						(bReverseClassChainCheck && (OtherMetaClass == NULL || MetaClass->IsChildOf(OtherMetaClass))))
+					if (!ClassDef->IsChildOf(*GUClassDef) || MetaClassDef == nullptr || Other.MetaClassDef->IsChildOf(*MetaClassDef) ||
+						(bReverseClassChainCheck && (Other.MetaClassDef == nullptr || MetaClassDef->IsChildOf(*Other.MetaClassDef))))
 					{
 						return true;
 					}
 				}
 			}
 			// check the opposite class chain for object types
-			else if (bReverseClassChainCheck && Type != CPT_Interface && bIsObjectComparison && PropertyClass != NULL && PropertyClass->IsChildOf(OtherPropertyClass))
+			else if (bReverseClassChainCheck && Type != CPT_Interface && bIsObjectComparison && ClassDef != nullptr && ClassDef->IsChildOf(*Other.ClassDef))
 			{
-				if (!OtherPropertyClass->IsChildOf(UClass::StaticClass()) || MetaClass == NULL || OtherMetaClass == NULL || MetaClass->IsChildOf(OtherMetaClass) || OtherMetaClass->IsChildOf(MetaClass))
+				if (!Other.ClassDef->IsChildOf(*GUClassDef) || MetaClassDef == nullptr || Other.MetaClassDef == nullptr || MetaClassDef->IsChildOf(*Other.MetaClassDef) || Other.MetaClassDef->IsChildOf(*MetaClassDef))
 				{
 					return true;
 				}
 			}
 
-			if (PropertyClass->HasAnyClassFlags(CLASS_Interface) && !bIgnoreImplementedInterfaces)
+			if (ClassDef->HasAnyClassFlags(CLASS_Interface) && !bIgnoreImplementedInterfaces)
 			{
-				if (OtherPropertyClass->ImplementsInterface(PropertyClass))
+				if (Other.ClassDef->ImplementsInterface(*ClassDef))
 				{
 					return true;
 				}
@@ -262,13 +263,10 @@ bool FPropertyBase::MatchesType(const FPropertyBase& Other, bool bDisallowGenera
 			return false;
 		}
 
-		UScriptStruct* Struct = ScriptStructDef->GetScriptStruct();
-		UScriptStruct* OtherStruct = Other.ScriptStructDef->GetScriptStruct();
-
 		// Generalization is ok if this is not a dynamic array
 		if (ArrayType != EArrayType::Dynamic && Other.ArrayType != EArrayType::Dynamic)
 		{
-			if (!OtherStruct->IsChildOf(Struct) && Struct->IsChildOf(OtherStruct))
+			if (!Other.ScriptStructDef->IsChildOf(*ScriptStructDef) && ScriptStructDef->IsChildOf(*Other.ScriptStructDef))
 			{
 				return true;
 			}
diff --git a/Engine/Source/Programs/UnrealHeaderTool/Private/ParserHelper.h b/Engine/Source/Programs/UnrealHeaderTool/Private/ParserHelper.h
--- a/Engine/Source/Programs/UnrealHeaderTool/Private/ParserHelper.h
+++ b/Engine/Source/Programs/UnrealHeaderTool/Private/ParserHelper.h
@@ -194,7 +194,7 @@ public:
 	FUnrealClassDefinitionInfo* MetaClassDef = nullptr;
 
 	FName   DelegateName = NAME_None;
-	UClass*	DelegateSignatureOwnerClass = nullptr;
+	FUnrealClassDefinitionInfo* DelegateSignatureOwnerClassDef = nullptr;
 	FName   RepNotifyName = NAME_None;
 
 	/** Raw string (not type-checked) used for specifying special text when exporting a property to the *Classes.h file */
diff --git a/Engine/Source/Programs/UnrealHeaderTool/Private/PropertyTypes.cpp b/Engine/Source/Programs/UnrealHeaderTool/Private/PropertyTypes.cpp
--- a/Engine/Source/Programs/UnrealHeaderTool/Private/PropertyTypes.cpp
+++ b/Engine/Source/Programs/UnrealHeaderTool/Private/PropertyTypes.cpp
@@ -10,6 +10,11 @@
 #include "Misc/DefaultValueHelper.h"
 #include "UObject/ObjectMacros.h"
 
+// Globals for common class definitions
+extern FUnrealClassDefinitionInfo* GUObjectDef = nullptr;
+extern FUnrealClassDefinitionInfo* GUClassDef = nullptr;
+extern FUnrealClassDefinitionInfo* GUInterfaceDef = nullptr;
+
 void AddEditInlineMetaData(TMap<FName, FString>& MetaData);
 
 // Following is the relationship between the property types
@@ -1126,13 +1131,12 @@ struct FPropertyTypeTraitsObjectReference : public FPropertyTypeTraitsObjectBase
 	{
 		FPropertyBase& VarProperty = PropDef.GetPropertyBase();
 		check(VarProperty.ClassDef);
-		UClass* PropertyClass = VarProperty.ClassDef->GetClass();
 
 #if UHT_ENABLE_PTR_PROPERTY_TAG
 		PropDef.GetUnrealSourceFile().AddTypeDefIncludeIfNeeded(VarProperty.ClassDef);
 #endif
 
-		if (PropertyClass->IsChildOf(UClass::StaticClass()))
+		if (VarProperty.ClassDef->IsChildOf(*GUClassDef))
 		{
 		}
 		else
@@ -1149,19 +1153,17 @@ struct FPropertyTypeTraitsObjectReference : public FPropertyTypeTraitsObjectBase
 	{
 		FPropertyBase& VarProperty = PropDef.GetPropertyBase();
 		check(VarProperty.ClassDef);
-		UClass* PropertyClass = VarProperty.ClassDef->GetClass();
-
-		if (PropertyClass->IsChildOf(UClass::StaticClass()))
+		if (VarProperty.ClassDef->IsChildOf(*GUClassDef))
 		{
 			FClassProperty* Result = new FClassProperty(Scope, Name, ObjectFlags);
 			Result->MetaClass = VarProperty.MetaClassDef ? VarProperty.MetaClassDef->GetClass() : nullptr;
-			Result->PropertyClass = PropertyClass;
+			Result->PropertyClass = VarProperty.ClassDef->GetClass();
 			return Result;
 		}
 		else
 		{
 			FObjectProperty* Result = new FObjectProperty(Scope, Name, ObjectFlags);
-			Result->PropertyClass = PropertyClass;
+			Result->PropertyClass = VarProperty.ClassDef->GetClass();
 			return Result;
 		}
 	}
@@ -1175,8 +1177,7 @@ struct FPropertyTypeTraitsObjectReference : public FPropertyTypeTraitsObjectBase
 	{
 		const FPropertyBase& VarProperty = PropDef.GetPropertyBase();
 		check(VarProperty.ClassDef);
-		UClass* PropertyClass = VarProperty.ClassDef->GetClass();
-		if (PropertyClass->IsChildOf(UClass::StaticClass()))
+		if (VarProperty.ClassDef->IsChildOf(*GUClassDef))
 		{
 			return FClassProperty::StaticClass()->GetName();
 		}
@@ -1190,13 +1191,11 @@ struct FPropertyTypeTraitsObjectReference : public FPropertyTypeTraitsObjectBase
 	{
 		const FPropertyBase& VarProperty = PropDef.GetPropertyBase();
 		check(VarProperty.ClassDef);
-		UClass* PropertyClass = VarProperty.ClassDef->GetClass();
-		if (PropertyClass->IsChildOf(UClass::StaticClass()))
+		if (VarProperty.ClassDef->IsChildOf(*GUClassDef))
 		{
 			if (PropDef.HasAnyPropertyFlags(CPF_UObjectWrapper))
 			{
-				UClass* MetaClass = VarProperty.MetaClassDef->GetClass();
-				return FString::Printf(TEXT("TSubclassOf<%s%s> "), MetaClass->GetPrefixCPP(), *MetaClass->GetName());
+				return FString::Printf(TEXT("TSubclassOf<%s%s> "), VarProperty.MetaClassDef->GetPrefixCPP(), *VarProperty.MetaClassDef->GetName());
 			}
 			else
 			{
@@ -1205,7 +1204,7 @@ struct FPropertyTypeTraitsObjectReference : public FPropertyTypeTraitsObjectBase
 		}
 		else
 		{
-			return FString::Printf(TEXT("%s%s*"), PropertyClass->GetPrefixCPP(), *PropertyClass->GetName());
+			return FString::Printf(TEXT("%s%s*"), VarProperty.ClassDef->GetPrefixCPP(), *VarProperty.ClassDef->GetName());
 		}
 	}
 
@@ -1213,15 +1212,13 @@ struct FPropertyTypeTraitsObjectReference : public FPropertyTypeTraitsObjectBase
 	{
 		const FPropertyBase& VarProperty = PropDef.GetPropertyBase();
 		check(VarProperty.ClassDef);
-		UClass* PropertyClass = VarProperty.ClassDef->GetClass();
-		if (PropertyClass->IsChildOf(UClass::StaticClass()))
+		if (VarProperty.ClassDef->IsChildOf(*GUClassDef))
 		{
-			UClass* MetaClass = VarProperty.MetaClassDef->GetClass();
-			return FString::Printf(TEXT("class %s%s;"), MetaClass->GetPrefixCPP(), *MetaClass->GetName());
+			return FString::Printf(TEXT("class %s%s;"), VarProperty.MetaClassDef->GetPrefixCPP(), *VarProperty.MetaClassDef->GetName());
 		}
 		else
 		{
-			return FString::Printf(TEXT("class %s%s;"), PropertyClass->GetPrefixCPP(), *PropertyClass->GetName());
+			return FString::Printf(TEXT("class %s%s;"), VarProperty.ClassDef->GetPrefixCPP(), *VarProperty.ClassDef->GetName());
 		}
 	}
 };
@@ -1260,8 +1257,7 @@ struct FPropertyTypeTraitsWeakObjectReference : public FPropertyTypeTraitsObject
 	{
 		const FPropertyBase& VarProperty = PropDef.GetPropertyBase();
 		check(VarProperty.ClassDef);
-		UClass* PropertyClass = VarProperty.ClassDef->GetClass();
-		FString InnerNativeTypeName = FString::Printf(TEXT("%s%s"), PropertyClass->GetPrefixCPP(), *PropertyClass->GetName());
+		FString InnerNativeTypeName = FString::Printf(TEXT("%s%s"), VarProperty.ClassDef->GetPrefixCPP(), *VarProperty.ClassDef->GetName());
 		if (PropDef.HasAnyPropertyFlags(CPF_AutoWeak))
 		{
 			return FString::Printf(TEXT("TAutoWeakObjectPtr<%s>"), *InnerNativeTypeName);
@@ -1273,8 +1269,7 @@ struct FPropertyTypeTraitsWeakObjectReference : public FPropertyTypeTraitsObject
 	{
 		const FPropertyBase& VarProperty = PropDef.GetPropertyBase();
 		check(VarProperty.ClassDef);
-		UClass* PropertyClass = VarProperty.ClassDef->GetClass();
-		return FString::Printf(TEXT("class %s%s;"), PropertyClass->GetPrefixCPP(), *PropertyClass->GetName());
+		return FString::Printf(TEXT("class %s%s;"), VarProperty.ClassDef->GetPrefixCPP(), *VarProperty.ClassDef->GetName());
 	}
 };
 
@@ -1307,16 +1302,14 @@ struct FPropertyTypeTraitsLazyObjectReference : public FPropertyTypeTraitsObject
 	{
 		const FPropertyBase& VarProperty = PropDef.GetPropertyBase();
 		check(VarProperty.ClassDef);
-		UClass* PropertyClass = VarProperty.ClassDef->GetClass();
-		return FString::Printf(TEXT("TLazyObjectPtr<%s%s>"), PropertyClass->GetPrefixCPP(), *PropertyClass->GetName());
+		return FString::Printf(TEXT("TLazyObjectPtr<%s%s>"), VarProperty.ClassDef->GetPrefixCPP(), *VarProperty.ClassDef->GetName());
 	}
 
 	static FString GetCPPTypeForwardDeclaration(const FUnrealPropertyDefinitionInfo& PropDef)
 	{
 		const FPropertyBase& VarProperty = PropDef.GetPropertyBase();
 		check(VarProperty.ClassDef);
-		UClass* PropertyClass = VarProperty.ClassDef->GetClass();
-		return FString::Printf(TEXT("class %s%s;"), PropertyClass->GetPrefixCPP(), *PropertyClass->GetName());
+		return FString::Printf(TEXT("class %s%s;"), VarProperty.ClassDef->GetPrefixCPP(), *VarProperty.ClassDef->GetName());
 	}
 };
 
@@ -1326,13 +1319,12 @@ struct FPropertyTypeTraitsObjectPtrReference : public FPropertyTypeTraitsObjectB
 	{
 		FPropertyBase& VarProperty = PropDef.GetPropertyBase();
 		check(VarProperty.ClassDef);
-		UClass* PropertyClass = VarProperty.ClassDef->GetClass();
 
 #if UHT_ENABLE_PTR_PROPERTY_TAG
 		PropDef.GetUnrealSourceFile().AddTypeDefIncludeIfNeeded(VarProperty.ClassDef);
 #endif
 
-		if (PropertyClass->IsChildOf(UClass::StaticClass()))
+		if (VarProperty.ClassDef->IsChildOf(*GUClassDef))
 		{
 		}
 		else
@@ -1349,19 +1341,17 @@ struct FPropertyTypeTraitsObjectPtrReference : public FPropertyTypeTraitsObjectB
 	{
 		FPropertyBase& VarProperty = PropDef.GetPropertyBase();
 		check(VarProperty.ClassDef);
-		UClass* PropertyClass = VarProperty.ClassDef->GetClass();
-
-		if (PropertyClass->IsChildOf(UClass::StaticClass()))
+		if (VarProperty.ClassDef->IsChildOf(*GUClassDef))
 		{
 			FClassPtrProperty* Result = new FClassPtrProperty(Scope, Name, ObjectFlags);
 			Result->MetaClass = VarProperty.MetaClassDef ? VarProperty.MetaClassDef->GetClass() : nullptr;
-			Result->PropertyClass = PropertyClass;
+			Result->PropertyClass = VarProperty.ClassDef->GetClass();
 			return Result;
 		}
 		else
 		{
 			FObjectPtrProperty* Result = new FObjectPtrProperty(Scope, Name, ObjectFlags);
-			Result->PropertyClass = PropertyClass;
+			Result->PropertyClass = VarProperty.ClassDef->GetClass();
 			return Result;
 		}
 	}
@@ -1375,8 +1365,7 @@ struct FPropertyTypeTraitsObjectPtrReference : public FPropertyTypeTraitsObjectB
 	{
 		const FPropertyBase& VarProperty = PropDef.GetPropertyBase();
 		check(VarProperty.ClassDef);
-		UClass* PropertyClass = VarProperty.ClassDef->GetClass();
-		if (PropertyClass->IsChildOf(UClass::StaticClass()))
+		if (VarProperty.ClassDef->IsChildOf(*GUClassDef))
 		{
 			return FClassPtrProperty::StaticClass()->GetName();
 		}
@@ -1390,14 +1379,13 @@ struct FPropertyTypeTraitsObjectPtrReference : public FPropertyTypeTraitsObjectB
 	{
 		const FPropertyBase& VarProperty = PropDef.GetPropertyBase();
 		check(VarProperty.ClassDef);
-		UClass* PropertyClass = VarProperty.ClassDef->GetClass();
-		if (PropertyClass->IsChildOf(UClass::StaticClass()))
+		if (VarProperty.ClassDef->IsChildOf(*GUClassDef))
 		{
-			return FString::Printf(TEXT("TObjectPtr<%s%s>"), PropertyClass->GetPrefixCPP(), *PropertyClass->GetName());
+			return FString::Printf(TEXT("TObjectPtr<%s%s>"), VarProperty.ClassDef->GetPrefixCPP(), *VarProperty.ClassDef->GetName());
 		}
 		else
 		{
-			return FString::Printf(TEXT("TObjectPtr<%s%s>"), PropertyClass->GetPrefixCPP(), *PropertyClass->GetName());
+			return FString::Printf(TEXT("TObjectPtr<%s%s>"), VarProperty.ClassDef->GetPrefixCPP(), *VarProperty.ClassDef->GetName());
 		}
 	}
 
@@ -1405,15 +1393,13 @@ struct FPropertyTypeTraitsObjectPtrReference : public FPropertyTypeTraitsObjectB
 	{
 		const FPropertyBase& VarProperty = PropDef.GetPropertyBase();
 		check(VarProperty.ClassDef);
-		UClass* PropertyClass = VarProperty.ClassDef->GetClass();
-		if (PropertyClass->IsChildOf(UClass::StaticClass()))
+		if (VarProperty.ClassDef->IsChildOf(*GUClassDef))
 		{
-			UClass* MetaClass = VarProperty.MetaClassDef->GetClass();
-			return FString::Printf(TEXT("class %s%s;"), MetaClass->GetPrefixCPP(), *MetaClass->GetName());
+			return FString::Printf(TEXT("class %s%s;"), VarProperty.MetaClassDef->GetPrefixCPP(), *VarProperty.MetaClassDef->GetName());
 		}
 		else
 		{
-			return FString::Printf(TEXT("class %s%s;"), PropertyClass->GetPrefixCPP(), *PropertyClass->GetName());
+			return FString::Printf(TEXT("class %s%s;"), VarProperty.ClassDef->GetPrefixCPP(), *VarProperty.ClassDef->GetName());
 		}
 	}
 };
@@ -1433,19 +1419,17 @@ struct FPropertyTypeTraitsSoftObjectReference : public FPropertyTypeTraitsObject
 	{
 		FPropertyBase& VarProperty = PropDef.GetPropertyBase();
 		check(VarProperty.ClassDef);
-		UClass* PropertyClass = VarProperty.ClassDef->GetClass();
-
-		if (PropertyClass->IsChildOf(UClass::StaticClass()))
+		if (VarProperty.ClassDef->IsChildOf(*GUClassDef))
 		{
 			FSoftClassProperty* Result = new FSoftClassProperty(Scope, Name, ObjectFlags);
 			Result->MetaClass = VarProperty.MetaClassDef ? VarProperty.MetaClassDef->GetClass() : nullptr;
-			Result->PropertyClass = PropertyClass;
+			Result->PropertyClass = VarProperty.ClassDef->GetClass();
 			return Result;
 		}
 		else
 		{
 			FSoftObjectProperty* Result = new FSoftObjectProperty(Scope, Name, ObjectFlags);
-			Result->PropertyClass = PropertyClass;
+			Result->PropertyClass = VarProperty.ClassDef->GetClass();
 			return Result;
 		}
 	}
@@ -1459,8 +1443,7 @@ struct FPropertyTypeTraitsSoftObjectReference : public FPropertyTypeTraitsObject
 	{
 		const FPropertyBase& VarProperty = PropDef.GetPropertyBase();
 		check(VarProperty.ClassDef);
-		UClass* PropertyClass = VarProperty.ClassDef->GetClass();
-		if (PropertyClass->IsChildOf(UClass::StaticClass()))
+		if (VarProperty.ClassDef->IsChildOf(*GUClassDef))
 		{
 			return FSoftClassProperty::StaticClass()->GetName();
 		}
@@ -1474,15 +1457,14 @@ struct FPropertyTypeTraitsSoftObjectReference : public FPropertyTypeTraitsObject
 	{
 		const FPropertyBase& VarProperty = PropDef.GetPropertyBase();
 		check(VarProperty.ClassDef);
-		UClass* PropertyClass = VarProperty.ClassDef->GetClass();
-		if (PropertyClass->IsChildOf(UClass::StaticClass()))
+		if (VarProperty.ClassDef->IsChildOf(*GUClassDef))
 		{
 			UClass* MetaClass = VarProperty.MetaClassDef->GetClass();
-			return FString::Printf(TEXT("TSoftClassPtr<%s%s> "), MetaClass->GetPrefixCPP(), *MetaClass->GetName()); //@TODO - This has an extra space in it
+			return FString::Printf(TEXT("TSoftClassPtr<%s%s> "), VarProperty.MetaClassDef->GetPrefixCPP(), *VarProperty.MetaClassDef->GetName()); //@TODO - This has an extra space in it
 		}
 		else
 		{
-			return FString::Printf(TEXT("TSoftObjectPtr<%s%s>"), PropertyClass->GetPrefixCPP(), *PropertyClass->GetName());
+			return FString::Printf(TEXT("TSoftObjectPtr<%s%s>"), VarProperty.ClassDef->GetPrefixCPP(), *VarProperty.ClassDef->GetName());
 		}
 	}
 
@@ -1490,15 +1472,14 @@ struct FPropertyTypeTraitsSoftObjectReference : public FPropertyTypeTraitsObject
 	{
 		const FPropertyBase& VarProperty = PropDef.GetPropertyBase();
 		check(VarProperty.ClassDef);
-		UClass* PropertyClass = VarProperty.ClassDef->GetClass();
-		if (PropertyClass->IsChildOf(UClass::StaticClass()))
+		if (VarProperty.ClassDef->IsChildOf(*GUClassDef))
 		{
 			UClass* MetaClass = VarProperty.MetaClassDef->GetClass();
-			return FString::Printf(TEXT("class %s%s;"), MetaClass->GetPrefixCPP(), *MetaClass->GetName());
+			return FString::Printf(TEXT("class %s%s;"), VarProperty.MetaClassDef->GetPrefixCPP(), *VarProperty.MetaClassDef->GetName());
 		}
 		else
 		{
-			return FString::Printf(TEXT("class %s%s;"), PropertyClass->GetPrefixCPP(), *PropertyClass->GetName());
+			return FString::Printf(TEXT("class %s%s;"), VarProperty.ClassDef->GetPrefixCPP(), *VarProperty.ClassDef->GetName());
 		}
 	}
 };
@@ -1509,8 +1490,7 @@ struct FPropertyTypeTraitsInterface : public FPropertyTypeTraitsBase
 	{
 		FPropertyBase& VarProperty = PropDef.GetPropertyBase();
 		check(VarProperty.ClassDef);
-		UClass* PropertyClass = VarProperty.ClassDef->GetClass();
-		check(PropertyClass->HasAnyClassFlags(CLASS_Interface));
+		check(VarProperty.ClassDef->HasAnyClassFlags(CLASS_Interface));
 
 #if UHT_ENABLE_PTR_PROPERTY_TAG
 		PropDef.GetUnrealSourceFile().AddTypeDefIncludeIfNeeded(VarProperty.ClassDef);
@@ -1521,10 +1501,9 @@ struct FPropertyTypeTraitsInterface : public FPropertyTypeTraitsBase
 	{
 		FPropertyBase& VarProperty = PropDef.GetPropertyBase();
 		check(VarProperty.ClassDef);
-		UClass* PropertyClass = VarProperty.ClassDef->GetClass();
-		check(PropertyClass->HasAnyClassFlags(CLASS_Interface));
+		check(VarProperty.ClassDef->HasAnyClassFlags(CLASS_Interface));
 		FInterfaceProperty* Result = new  FInterfaceProperty(Scope, Name, ObjectFlags);
-		Result->InterfaceClass = PropertyClass;
+		Result->InterfaceClass = VarProperty.ClassDef->GetClass();
 		return Result;
 	}
 
diff --git a/Engine/Source/Programs/UnrealHeaderTool/Private/Scope.cpp b/Engine/Source/Programs/UnrealHeaderTool/Private/Scope.cpp
--- a/Engine/Source/Programs/UnrealHeaderTool/Private/Scope.cpp
+++ b/Engine/Source/Programs/UnrealHeaderTool/Private/Scope.cpp
@@ -147,5 +147,5 @@ FName FFileScope::GetName() const
 
 FName FStructScope::GetName() const
 {
-	return Struct->GetFName();
+	return StructDef.GetFName();
 }
diff --git a/Engine/Source/Programs/UnrealHeaderTool/Private/Scope.h b/Engine/Source/Programs/UnrealHeaderTool/Private/Scope.h
--- a/Engine/Source/Programs/UnrealHeaderTool/Private/Scope.h
+++ b/Engine/Source/Programs/UnrealHeaderTool/Private/Scope.h
@@ -325,23 +325,14 @@ class FStructScope : public FScope
 {
 public:
 	// Constructor.
-	FStructScope(FUnrealStructDefinitionInfo& InStructDef, UStruct* InStruct, FScope* InParent)
+	FStructScope(FUnrealStructDefinitionInfo& InStructDef, FScope* InParent)
 		: FScope(InParent)
-		, Struct(InStruct)
 		, StructDef(InStructDef)
 	{
 	}
 
 	virtual FStructScope* AsStructScope() override { return this; }
 
-	/**
-	 * Gets struct associated with this scope.
-	 */
-	UStruct* GetStruct() const
-	{
-		return Struct;
-	}
-
 	/**
 	 * Get the structure definition associated with this scope
 	 */
@@ -357,7 +348,6 @@ public:
 
 private:
 	// Struct associated with this scope.
-	UStruct* Struct;
 	FUnrealStructDefinitionInfo& StructDef;
 };
 
@@ -390,7 +380,7 @@ public:
 			if (FUnrealClassDefinitionInfo* ClassDef = UHTCast<FUnrealClassDefinitionInfo>(StructDef))
 			{
 				// Skip myself when starting this loop, we only care about the parents
-				for (ClassDef = ClassDef->GetSuperClass(); ClassDef && (ClassDef->GetClass()->ClassFlags & EClassFlags::CLASS_Intrinsic) == 0; ClassDef = ClassDef->GetSuperClass())
+				for (ClassDef = ClassDef->GetSuperClass(); ClassDef && !ClassDef->HasAnyClassFlags(EClassFlags::CLASS_Intrinsic); ClassDef = ClassDef->GetSuperClass())
 				{
 					ScopesToTraverse.Add(&ClassDef->GetScope().Get());
 				}
diff --git a/Engine/Source/Programs/UnrealHeaderTool/Private/UnrealHeaderTool.h b/Engine/Source/Programs/UnrealHeaderTool/Private/UnrealHeaderTool.h
--- a/Engine/Source/Programs/UnrealHeaderTool/Private/UnrealHeaderTool.h
+++ b/Engine/Source/Programs/UnrealHeaderTool/Private/UnrealHeaderTool.h
@@ -16,7 +16,6 @@ struct FHeaderParserNames
 	static const FName NAME_HideCategories;
 	static const FName NAME_ShowCategories;
 	static const FName NAME_SparseClassDataTypes;
-	static const FName NAME_AdvancedClassDisplay;
 	static const FName NAME_BlueprintType;
 	static const FName NAME_AutoCollapseCategories;
 	static const FName NAME_HideFunctions;
diff --git a/Engine/Source/Programs/UnrealHeaderTool/Private/UnrealTypeDefinitionInfo.cpp b/Engine/Source/Programs/UnrealHeaderTool/Private/UnrealTypeDefinitionInfo.cpp
--- a/Engine/Source/Programs/UnrealHeaderTool/Private/UnrealTypeDefinitionInfo.cpp
+++ b/Engine/Source/Programs/UnrealHeaderTool/Private/UnrealTypeDefinitionInfo.cpp
@@ -31,6 +31,7 @@ namespace
 	const FName NAME_ReplaceConverted(TEXT("ReplaceConverted"));
 	const FName NAME_ClassGroupNames(TEXT("ClassGroupNames"));
 	const FName NAME_EditorConfig(TEXT("EditorConfig"));
+	const FName NAME_AdvancedClassDisplay(TEXT("AdvancedClassDisplay"));
 
 	/**
 	 * As part of the singleton name, collect the parent chain names
@@ -84,15 +85,14 @@ namespace
 		Out.Append(TEXT("()"));
 	}
 
-	bool IsActorClass(UClass* Class)
+	bool IsActorClass(const FUnrealClassDefinitionInfo& ClassDef)
 	{
-		while (Class)
+		for (const FUnrealClassDefinitionInfo* TestDef = &ClassDef; TestDef; TestDef = TestDef->GetSuperClass())
 		{
-			if (Class->GetFName() == NAME_Actor)
+			if (TestDef->GetFName() == NAME_Actor)
 			{
 				return true;
 			}
-			Class = Class->GetSuperClass();
 		}
 		return false;
 	}
@@ -245,6 +245,21 @@ FName FUHTMetaData::GetMetaDataKey(FName Key, int32 NameIndex, EFindName FindNam
 	}
 }
 
+void FUHTMetaData::GetStringArrayMetaData(const FName& Key, TArray<FString>& Out) const
+{
+	if (const FString* String = FindMetaData(Key))
+	{
+		String->ParseIntoArray(Out, TEXT(" "), true);
+	}
+}
+
+TArray<FString> FUHTMetaData::GetStringArrayMetaData(const FName& Key) const
+{
+	TArray<FString> Out;
+	GetStringArrayMetaData(Key, Out);
+	return Out;
+}
+
 FUnrealPropertyDefinitionInfo* FUnrealTypeDefinitionInfo::AsProperty()
 {
 	return nullptr;
@@ -334,14 +349,15 @@ void FUnrealTypeDefinitionInfo::GetHashTag(FUHTStringBuilder& Out) const
 	}
 }
 
-void FUnrealTypeDefinitionInfo::CreateUObjectEngineTypes()
+void FUnrealTypeDefinitionInfo::CreateUObjectEngineTypes(ECreateEngineTypesPhase Phase)
 {
-	switch (CreateUObectEngineTypesState)
+	EFinalizeState& FinalizeState = CreateUObectEngineTypesState[uint8(Phase)];
+	switch (FinalizeState)
 	{
 	case EFinalizeState::None:
-		CreateUObectEngineTypesState = EFinalizeState::InProgress;
-		CreateUObjectEngineTypesInternal();
-		CreateUObectEngineTypesState = EFinalizeState::Finished;
+		FinalizeState = EFinalizeState::InProgress;
+		CreateUObjectEngineTypesInternal(Phase);
+		FinalizeState = EFinalizeState::Finished;
 		break;
 	case EFinalizeState::InProgress:
 		checkf(false, TEXT("Recursive call to CreateUObectEngineTypes detected"))
@@ -875,11 +891,12 @@ FUnrealPackageDefinitionInfo& FUnrealObjectDefinitionInfo::GetPackageDef() const
 	return GTypeDefinitionInfoMap.FindChecked<FUnrealPackageDefinitionInfo>(GetObject()->GetPackage());
 }
 
-FUnrealObjectDefinitionInfo::FUnrealObjectDefinitionInfo(UObject* Object)
-	: FUnrealTypeDefinitionInfo(Object->GetName())
-	, Name(Object->GetFName())
-	, InternalObjectFlags(Object->GetInternalFlags())
+FUnrealObjectDefinitionInfo::FUnrealObjectDefinitionInfo(UObject* InObject)
+	: FUnrealTypeDefinitionInfo(InObject->GetName())
+	, Name(InObject->GetFName())
+	, InternalObjectFlags(InObject->GetInternalFlags())
 {
+	SetObject(InObject);
 }
 
 FString FUnrealObjectDefinitionInfo::GetFullName() const
@@ -945,7 +962,7 @@ FString FUnrealObjectDefinitionInfo::GetFullGroupName(bool bStartWithOuter) cons
 	return Obj ? Obj->GetPathName(&GetPackageDef()) : TEXT("");
 }
 
-void FUnrealObjectDefinitionInfo::SetMetaDataHelper(const FName& Key, const TCHAR* InValue) 
+void FUnrealObjectDefinitionInfo::SetMetaDataHelper(const FName& Key, const TCHAR* InValue)
 {
 	// If we don't have an existing meta data map, then we must be using the meta data off
 	// the object and no values have been set yet.  In that case just set the whole block
@@ -984,16 +1001,15 @@ FUnrealPackageDefinitionInfo::FUnrealPackageDefinitionInfo(const FManifestModule
 	, ShortUpperName(FPackageName::GetShortName(InPackage).ToUpper())
 	, API(FString::Printf(TEXT("%s_API "), *ShortUpperName))
 {
-	SetObject(InPackage);
 }
 
-void FUnrealPackageDefinitionInfo::CreateUObjectEngineTypesInternal()
+void FUnrealPackageDefinitionInfo::CreateUObjectEngineTypesInternal(ECreateEngineTypesPhase Phase)
 {
 	for (TSharedRef<FUnrealSourceFile>& LocalSourceFile : GetAllSourceFiles())
 	{
 		for (TSharedRef<FUnrealTypeDefinitionInfo>& TypeDef : LocalSourceFile->GetDefinedTypes())
 		{
-			TypeDef->CreateUObjectEngineTypes();
+			TypeDef->CreateUObjectEngineTypes(Phase);
 		}
 	}
 }
@@ -1419,18 +1435,54 @@ FString FUnrealEnumDefinitionInfo::GetNameStringByIndex(int32 InIndex) const
 	return FString();
 }
 
-void FUnrealEnumDefinitionInfo::CreateUObjectEngineTypesInternal()
+void FUnrealEnumDefinitionInfo::CreateUObjectEngineTypesInternal(ECreateEngineTypesPhase Phase)
 {
-	UPackage* Package = GetPackageDef().GetPackage();
-	const FString& EnumName = GetNameCPP();
+	FUnrealFieldDefinitionInfo::CreateUObjectEngineTypesInternal(Phase);
 
-	// Create enum definition.
-	UEnum* Enum = new(EC_InternalUseOnlyConstructor, Package, FName(EnumName), RF_Public) UEnum(FObjectInitializer());
-	Enum->SetEnums(Names, CppForm, EnumFlags, false);
-	Enum->CppType = CppType;
-	Enum->GetPackage()->GetMetaData()->SetObjectValues(Enum, MetaData);
-	SetObject(Enum);
-	GTypeDefinitionInfoMap.AddObjectLookup(Enum, SharedThis(this));
+	switch (Phase)
+	{
+	case ECreateEngineTypesPhase::Phase1:
+	{
+		UPackage* Package = GetPackageDef().GetPackage();
+		const FString& EnumName = GetNameCPP();
+
+		// Create enum definition.
+		UEnum* Enum = new(EC_InternalUseOnlyConstructor, Package, FName(EnumName), RF_Public) UEnum(FObjectInitializer());
+		Enum->SetEnums(Names, CppForm, EnumFlags, false);
+		Enum->CppType = CppType;
+		Enum->GetPackage()->GetMetaData()->SetObjectValues(Enum, MetaData);
+		SetObject(Enum);
+		GTypeDefinitionInfoMap.AddObjectLookup(Enum, SharedThis(this));
+		break;
+	}
+
+	case ECreateEngineTypesPhase::Phase2:
+		break;
+	}
+}
+
+FUnrealStructDefinitionInfo::FUnrealStructDefinitionInfo(FUnrealSourceFile& InSourceFile, int32 InLineNumber, FString&& InNameCPP, FName InName, FUnrealObjectDefinitionInfo& InOuter)
+	: FUnrealFieldDefinitionInfo(InSourceFile, InLineNumber, MoveTemp(InNameCPP), InName, InOuter)
+	, StructScope(MakeShared<FStructScope>(*this, &InSourceFile.GetScope().Get()))
+{
+}
+
+bool FUnrealStructDefinitionInfo::IsChildOf(const FUnrealStructDefinitionInfo& SomeBase) const
+{
+	for (const FUnrealStructDefinitionInfo* Current = this; Current; Current = Current->GetSuperStructInfo().Struct)
+	{
+		if (Current == &SomeBase)
+		{
+#if UHT_ENABLE_ENGINE_TYPE_CHECKS
+			check(GetStructSafe() == nullptr || SomeBase.GetStructSafe() == nullptr || GetStructSafe()->IsChildOf(SomeBase.GetStructSafe()) == true);
+#endif
+			return true;
+		}
+	}
+#if UHT_ENABLE_ENGINE_TYPE_CHECKS
+	check(GetStructSafe() == nullptr || SomeBase.GetStructSafe() == nullptr || GetStructSafe()->IsChildOf(SomeBase.GetStructSafe()) == false);
+#endif
+	return false;
 }
 
 void FUnrealStructDefinitionInfo::AddProperty(TSharedRef<FUnrealPropertyDefinitionInfo> PropertyDef)
@@ -1456,26 +1508,33 @@ void FUnrealStructDefinitionInfo::AddProperty(TSharedRef<FUnrealPropertyDefiniti
 	}
 }
 
-void FUnrealStructDefinitionInfo::CreateUObjectEngineTypesInternal()
+void FUnrealStructDefinitionInfo::CreateUObjectEngineTypesInternal(ECreateEngineTypesPhase Phase)
 {
-	FUnrealFieldDefinitionInfo::CreateUObjectEngineTypesInternal();
+	FUnrealFieldDefinitionInfo::CreateUObjectEngineTypesInternal(Phase);
 
-	if (SuperStructInfo.Struct)
+	switch (Phase)
 	{
-		SuperStructInfo.Struct->CreateUObjectEngineTypes();
-	}
+	case ECreateEngineTypesPhase::Phase1:
+		if (SuperStructInfo.Struct)
+		{
+			SuperStructInfo.Struct->CreateUObjectEngineTypes(Phase);
+		}
 
-	for (const FBaseStructInfo& Info : BaseStructInfo)
-	{
-		if (Info.Struct)
+		for (const FBaseStructInfo& Info : BaseStructInfos)
 		{
-			Info.Struct->CreateUObjectEngineTypes();
+			if (Info.Struct)
+			{
+				Info.Struct->CreateUObjectEngineTypes(Phase);
+			}
 		}
-	}
+		break;
 
-	for (TSharedRef<FUnrealFunctionDefinitionInfo> FunctionDef : Functions)
-	{
-		FunctionDef->CreateUObjectEngineTypes();
+	case ECreateEngineTypesPhase::Phase2:
+		for (TSharedRef<FUnrealFunctionDefinitionInfo> FunctionDef : Functions)
+		{
+			FunctionDef->CreateUObjectEngineTypes(Phase);
+		}
+		break;
 	}
 }
 
@@ -1488,7 +1547,7 @@ void FUnrealStructDefinitionInfo::PostParseFinalizeInternal()
 		SuperStructInfo.Struct->PostParseFinalize();
 	}
 
-	for (const FBaseStructInfo& Info : BaseStructInfo)
+	for (const FBaseStructInfo& Info : BaseStructInfos)
 	{
 		if (Info.Struct)
 		{
@@ -1542,6 +1601,25 @@ void FUnrealStructDefinitionInfo::AddFunction(TSharedRef<FUnrealFunctionDefiniti
 	}
 }
 
+TSharedRef<FScope> FUnrealStructDefinitionInfo::GetScope()
+{
+	if (StructScope.IsValid())
+	{
+		return StructScope.ToSharedRef();
+	}
+	else
+	{
+		return FUnrealFieldDefinitionInfo::GetScope();
+	}
+}
+
+void FUnrealStructDefinitionInfo::SetObject(UObject* InObject)
+{
+	check(InObject != nullptr);
+
+	FUnrealFieldDefinitionInfo::SetObject(InObject);
+}
+
 FUnrealScriptStructDefinitionInfo::FUnrealScriptStructDefinitionInfo(FUnrealSourceFile& InSourceFile, int32 InLineNumber, FString&& InNameCPP, FName InName)
 	: FUnrealStructDefinitionInfo(InSourceFile, InLineNumber, MoveTemp(InNameCPP), InName, InSourceFile.GetPackageDef())
 { }
@@ -1558,41 +1636,86 @@ uint32 FUnrealScriptStructDefinitionInfo::GetHash(bool bIncludeNoExport) const
 	return FUnrealStructDefinitionInfo::GetHash(bIncludeNoExport);
 }
 
-TSharedRef<FScope> FUnrealStructDefinitionInfo::GetScope()
+void FUnrealScriptStructDefinitionInfo::CreateUObjectEngineTypesInternal(ECreateEngineTypesPhase Phase)
 {
-	if (StructScope.IsValid())
+	FUnrealStructDefinitionInfo::CreateUObjectEngineTypesInternal(Phase);
+
+	switch (Phase)
 	{
-		return StructScope.ToSharedRef();
-	}
-	else
+	case ECreateEngineTypesPhase::Phase1:
 	{
-		return FUnrealFieldDefinitionInfo::GetScope();
-	}
-}
+		UPackage* Package = GetPackageDef().GetPackage();
+		const FString& StructName = GetNameCPP();
+		FString StructNameStripped = GetClassNameWithPrefixRemoved(*StructName);
 
-void FUnrealStructDefinitionInfo::SetObject(UObject* InObject)
-{
-	check(InObject != nullptr);
+		UScriptStruct* ScriptStruct = new(EC_InternalUseOnlyConstructor, Package, FName(StructNameStripped), RF_Public) UScriptStruct(FObjectInitializer());
+		ScriptStruct->StructFlags = StructFlags;
+		ScriptStruct->GetPackage()->GetMetaData()->SetObjectValues(ScriptStruct, MetaData);
+		if (FUnrealStructDefinitionInfo* SuperStructDef = GetSuperStruct())
+		{
+			ScriptStruct->SetSuperStruct(SuperStructDef->GetStruct());
+		}
+		SetObject(ScriptStruct);
+		GTypeDefinitionInfoMap.AddObjectLookup(ScriptStruct, SharedThis(this));
 
-	FUnrealFieldDefinitionInfo::SetObject(InObject);
+		// The structure needs to be prepared at this point.  
+		ScriptStruct->PrepareCppStructOps();
+		break;
+	}
 
-	// Don't create a scope for things without a source.  Those are builtin types
-	if (HasSource())
-	{
-		StructScope = MakeShared<FStructScope>(*this, static_cast<UStruct*>(InObject), &GetUnrealSourceFile().GetScope().Get());
+	case ECreateEngineTypesPhase::Phase2:
+		break;
 	}
 }
 
+void FUnrealScriptStructDefinitionInfo::PostParseFinalizeInternal()
+{
+	FUnrealStructDefinitionInfo::PostParseFinalizeInternal();
+}
+
 FUnrealClassDefinitionInfo::FUnrealClassDefinitionInfo(FUnrealSourceFile& InSourceFile, int32 InLineNumber, FString&& InNameCPP, FName InName, bool bInIsInterface)
 	: FUnrealStructDefinitionInfo(InSourceFile, InLineNumber, MoveTemp(InNameCPP), InName, InSourceFile.GetPackageDef())
 	, bIsInterface(bInIsInterface)
 {
 	if (bInIsInterface)
 	{
+		ClassFlags |= CLASS_Interface;
 		GetStructMetaData().ParsedInterface = EParsedInterface::ParsedUInterface;
 	}
 }
 
+/**
+ * Returns the struct/ class prefix used for the C++ declaration of this struct/ class.
+ * Classes deriving from AActor have an 'A' prefix and other UObject classes an 'U' prefix.
+ *
+ * @return Prefix character used for C++ declaration of this struct/ class.
+ */
+const TCHAR* FUnrealClassDefinitionInfo::GetPrefixCPP() const
+{
+	if (IsActorClass(*this))
+	{
+		if (HasAnyClassFlags(CLASS_Deprecated))
+		{
+			return TEXT("ADEPRECATED_");
+		}
+		else
+		{
+			return TEXT("A");
+		}
+	}
+	else
+	{
+		if (HasAnyClassFlags(CLASS_Deprecated))
+		{
+			return TEXT("UDEPRECATED_");
+		}
+		else
+		{
+			return TEXT("U");
+		}
+	}
+}
+
 FUnrealClassDefinitionInfo* FUnrealClassDefinitionInfo::FindClass(const TCHAR* ClassName)
 {
 	FUnrealClassDefinitionInfo* ClassDef = nullptr;
@@ -1674,8 +1797,7 @@ uint32 FUnrealClassDefinitionInfo::GetHash(bool bIncludeNoExport) const
 {
 	if (!bIncludeNoExport)
 	{
-		UClass* Class = GetClass();
-		if (Class->HasAnyClassFlags(CLASS_NoExport))
+		if (HasAnyClassFlags(CLASS_NoExport))
 		{
 			return 0;
 		}
@@ -1683,11 +1805,64 @@ uint32 FUnrealClassDefinitionInfo::GetHash(bool bIncludeNoExport) const
 	return FUnrealStructDefinitionInfo::GetHash(bIncludeNoExport);
 }
 
+void FUnrealClassDefinitionInfo::CreateUObjectEngineTypesInternal(ECreateEngineTypesPhase Phase)
+{
+	FUnrealStructDefinitionInfo::CreateUObjectEngineTypesInternal(Phase);
+
+	switch (Phase)
+	{
+	case ECreateEngineTypesPhase::Phase1:
+	{
+		check(ClassWithin);
+		if (GetObjectSafe() == nullptr)
+		{
+			UPackage* Package = GetPackageDef().GetPackage();
+			const FString& ClassName = GetNameCPP();
+			FString ClassNameStripped = GetClassNameWithPrefixRemoved(*ClassName);
+
+			// Create new class.
+			UClass* Class = new(EC_InternalUseOnlyConstructor, Package, *ClassNameStripped, RF_Public | RF_Standalone) UClass(FObjectInitializer(), nullptr);
+
+			Class->ClassFlags = ClassFlags;
+			Class->ClassCastFlags = ClassCastFlags;
+			Class->PropertiesSize = PropertiesSize;
+			Class->ClassConfigName = ClassConfigName;
+			Class->ClassWithin = ClassWithin->GetClassSafe();
+			Class->SetInternalFlags(GetInternalFlags());
+			Class->GetPackage()->GetMetaData()->SetObjectValues(Class, MetaData);
+
+			// Setup the base class
+			if (FUnrealClassDefinitionInfo* SuperClassDef = GetSuperClass())
+			{
+				Class->SetSuperStruct(SuperClassDef->GetClass());
+			}
+
+			// Add the class flags from the interfaces
+			for (FUnrealStructDefinitionInfo::FBaseStructInfo& BaseClassInfo : GetBaseStructInfos())
+			{
+				if (FUnrealClassDefinitionInfo* BaseClassDef = UHTCast<FUnrealClassDefinitionInfo>(BaseClassInfo.Struct))
+				{
+					if (BaseClassDef->HasAnyClassFlags(CLASS_Interface))
+					{
+						Class->Interfaces.Emplace(BaseClassDef->GetClass(), 1, false);
+					}
+				}
+			}
+			SetObject(Class);
+		}
+
+		// Even if we already had an object, add it back to the types regardless
+		GTypeDefinitionInfoMap.AddObjectLookup(GetObjectSafe(), SharedThis(this));
+		break;
+	}
+
+	case ECreateEngineTypesPhase::Phase2:
+		break;
+	}
+}
 
 void FUnrealClassDefinitionInfo::PostParseFinalizeInternal()
 {
-	UClass* Class = GetClass();
-
 	FUnrealStructDefinitionInfo::PostParseFinalizeInternal();
 
 	if (IsInterface() && GetStructMetaData().ParsedInterface == EParsedInterface::ParsedUInterface)
@@ -1699,6 +1874,31 @@ void FUnrealClassDefinitionInfo::PostParseFinalizeInternal()
 
 	FHeaderParser::CheckSparseClassData(*this);
 
+	// Collect the class replication properties
+	if (FUnrealClassDefinitionInfo* SuperClassDef = GetSuperClass())
+	{
+		ClassReps = SuperClassDef->GetClassReps();
+		FirstOwnedClassRep = ClassReps.Num();
+	}
+	for (TSharedRef<FUnrealPropertyDefinitionInfo> PropertyDef : GetProperties())
+	{
+		if (PropertyDef->HasAnyPropertyFlags(CPF_Net))
+		{
+			ClassReps.Add(&*PropertyDef);
+		}
+	}
+
+	// Initialize the class object
+	UClass* Class = GetClass();
+
+	// Clear the property size
+	Class->PropertiesSize = 0;
+
+	// Make the visible to the package
+	Class->ClearFlags(RF_Transient);
+	check(Class->HasAnyFlags(RF_Public));
+	check(Class->HasAnyFlags(RF_Standalone));
+
 	// Finalize all of the children introduced in this class
 	for (TSharedRef<FUnrealFunctionDefinitionInfo> FunctionDef : GetFunctions())
 	{
@@ -1723,6 +1923,38 @@ void FUnrealClassDefinitionInfo::PostParseFinalizeInternal()
 	Class->SetUpUhtReplicationData();
 }
 
+bool FUnrealClassDefinitionInfo::ImplementsInterface(const FUnrealClassDefinitionInfo& SomeInterface) const
+{
+	if (SomeInterface.HasAnyClassFlags(CLASS_Interface) && &SomeInterface != GUInterfaceDef)
+	{
+		for (const FUnrealClassDefinitionInfo* CurrentClassDef = this; CurrentClassDef; CurrentClassDef = CurrentClassDef->GetSuperClass())
+		{
+			// SomeInterface might be a base interface of our implemented interface
+			for (const FBaseStructInfo& BaseStructInfo : GetBaseStructInfos())
+			{
+				if (const FUnrealClassDefinitionInfo* InterfaceDef = UHTCast<FUnrealClassDefinitionInfo>(BaseStructInfo.Struct))
+				{
+					if (InterfaceDef->IsChildOf(SomeInterface))
+					{
+						return true;
+					}
+				}
+			}
+		}
+	}
+	return false;
+}
+
+void FUnrealClassDefinitionInfo::InitializeFromExistingUObject(UClass* Class)
+{
+	ClassFlags = Class->ClassFlags;
+	ClassCastFlags = Class->ClassCastFlags;
+	InitialEngineClassFlags = Class->ClassFlags;
+	PropertiesSize = Class->PropertiesSize;
+	ClassConfigName = Class->ClassConfigName;
+	SetInternalFlags(Class->GetInternalFlags());
+}
+
 void FUnrealClassDefinitionInfo::ParseClassProperties(TArray<FPropertySpecifier>&& InClassSpecifiers, const FString& InRequiredAPIMacroIfPresent)
 {
 	ParsedClassFlags = CLASS_None;
@@ -1751,7 +1983,7 @@ void FUnrealClassDefinitionInfo::ParseClassProperties(TArray<FPropertySpecifier>
 
 		case EClassMetadataSpecifier::ComponentWrapperClass:
 
-			MetaData.Add(NAME_IgnoreCategoryKeywordsInSubclasses, TEXT("true"));
+			ParsedMetaData.Add(NAME_IgnoreCategoryKeywordsInSubclasses, TEXT("true"));
 			break;
 
 		case EClassMetadataSpecifier::Within:
@@ -1879,7 +2111,7 @@ void FUnrealClassDefinitionInfo::ParseClassProperties(TArray<FPropertySpecifier>
 
 		case EClassMetadataSpecifier::EditorConfig:
 			// Save EditorConfig properties to the given JSON file.
-			MetaData.Add(NAME_EditorConfig, FHeaderParser::RequireExactlyOneSpecifierValue(*this, PropSpecifier));
+			ParsedMetaData.Add(NAME_EditorConfig, FHeaderParser::RequireExactlyOneSpecifierValue(*this, PropSpecifier));
 			break;
 
 		case EClassMetadataSpecifier::ShowCategories:
@@ -1985,12 +2217,12 @@ void FUnrealClassDefinitionInfo::ParseClassProperties(TArray<FPropertySpecifier>
 		case EClassMetadataSpecifier::AdvancedClassDisplay:
 
 			// By default the class properties are shown in advanced sections in UI
-			MetaData.Add(FHeaderParserNames::NAME_AdvancedClassDisplay, TEXT("true"));
+			ParsedMetaData.Add(NAME_AdvancedClassDisplay, TEXT("true"));
 			break;
 
 		case EClassMetadataSpecifier::ConversionRoot:
 
-			MetaData.Add(FHeaderParserNames::NAME_IsConversionRoot, TEXT("true"));
+			ParsedMetaData.Add(FHeaderParserNames::NAME_IsConversionRoot, TEXT("true"));
 			break;
 
 		case EClassMetadataSpecifier::NeedsDeferredDependencyLoading:
@@ -2002,6 +2234,7 @@ void FUnrealClassDefinitionInfo::ParseClassProperties(TArray<FPropertySpecifier>
 			FUHTException::Throwf(*this, TEXT("Unknown class specifier '%s'"), *PropSpecifier.Key);
 		}
 	}
+	SetClassFlags(ParsedClassFlags);
 }
 
 void FUnrealClassDefinitionInfo::MergeShowCategories()
@@ -2035,22 +2268,10 @@ void FUnrealClassDefinitionInfo::MergeShowCategories()
 
 void FUnrealClassDefinitionInfo::MergeClassCategories()
 {
-	UClass* Class = GetClass();
-	TArray<FString> ParentHideCategories;
-	TArray<FString> ParentShowSubCatgories;
-	TArray<FString> ParentHideFunctions;
-	TArray<FString> ParentAutoExpandCategories;
-	TArray<FString> ParentAutoCollapseCategories;
-	GetHideCategories(ParentHideCategories);
-	GetShowCategories(ParentShowSubCatgories);
-	Class->GetHideFunctions(ParentHideFunctions);
-	Class->GetAutoExpandCategories(ParentAutoExpandCategories);
-	Class->GetAutoCollapseCategories(ParentAutoCollapseCategories);
-
 	// Add parent categories. We store the opposite of HideCategories and HideFunctions in a separate array anyway.
-	HideCategories.Append(MoveTemp(ParentHideCategories));
-	ShowSubCatgories.Append(MoveTemp(ParentShowSubCatgories));
-	HideFunctions.Append(MoveTemp(ParentHideFunctions));
+	HideCategories.Append(GetStringArrayMetaData(FHeaderParserNames::NAME_HideCategories));
+	ShowSubCatgories.Append(GetStringArrayMetaData(FHeaderParserNames::NAME_ShowCategories));
+	HideFunctions.Append(GetStringArrayMetaData(FHeaderParserNames::NAME_HideFunctions));
 
 	MergeShowCategories();
 
@@ -2076,6 +2297,9 @@ void FUnrealClassDefinitionInfo::MergeClassCategories()
 	ShowFunctions.Empty();
 
 	// Merge AutoExpandCategories and AutoCollapseCategories (we still want to keep AutoExpandCategories though!)
+	TArray<FString> ParentAutoExpandCategories = GetStringArrayMetaData(FHeaderParserNames::NAME_AutoExpandCategories);
+	TArray<FString> ParentAutoCollapseCategories = GetStringArrayMetaData(FHeaderParserNames::NAME_AutoCollapseCategories);
+
 	for (const FString& Value : AutoExpandCategories)
 	{
 		AutoCollapseCategories.RemoveSwap(Value);
@@ -2094,37 +2318,36 @@ void FUnrealClassDefinitionInfo::MergeClassCategories()
 	AutoExpandCategories.Append(MoveTemp(ParentAutoExpandCategories));
 }
 
-void FUnrealClassDefinitionInfo::MergeAndValidateClassFlags(const FString& DeclaredClassName, EClassFlags PreviousClassFlags)
+void FUnrealClassDefinitionInfo::MergeAndValidateClassFlags(const FString& DeclaredClassName)
 {
-	UClass* Class = GetClass();
 	if (bWantsToBePlaceable)
 	{
-		if (!(Class->ClassFlags & CLASS_NotPlaceable))
+		if (!HasAnyClassFlags(CLASS_NotPlaceable))
 		{
 			FUHTException::Throwf(*this, TEXT("The 'placeable' specifier is only allowed on classes which have a base class that's marked as not placeable. Classes are assumed to be placeable by default."));
 		}
-		Class->ClassFlags &= ~CLASS_NotPlaceable;
+		ClearClassFlags(CLASS_NotPlaceable);
 		bWantsToBePlaceable = false; // Reset this flag after it's been merged
 	}
 
 	// Now merge all remaining flags/properties
-	Class->ClassFlags |= ParsedClassFlags;
-	Class->ClassConfigName = FName(*ConfigName);
+	SetClassFlags(ParsedClassFlags);
+	SetClassConfigName(FName(*ConfigName));
 
 	SetAndValidateWithinClass();
 	SetAndValidateConfigName();
 
-	if (!!(Class->ClassFlags & CLASS_EditInlineNew))
+	if (HasAnyClassFlags(CLASS_EditInlineNew))
 	{
 		// don't allow actor classes to be declared editinlinenew
-		if (IsActorClass(Class))
+		if (IsActorClass(*this))
 		{
 			FUHTException::Throwf(*this, TEXT("Invalid class attribute: Creating actor instances via the property window is not allowed"));
 		}
 	}
 
 	// Make sure both RequiredAPI and MinimalAPI aren't specified
-	if (Class->HasAllClassFlags(CLASS_MinimalAPI | CLASS_RequiredAPI))
+	if (HasAllClassFlags(CLASS_MinimalAPI | CLASS_RequiredAPI))
 	{
 		FUHTException::Throwf(*this, TEXT("MinimalAPI cannot be specified when the class is fully exported using a MODULENAME_API macro"));
 	}
@@ -2136,69 +2359,70 @@ void FUnrealClassDefinitionInfo::MergeAndValidateClassFlags(const FString& Decla
 		FUHTException::Throwf(*this, TEXT("Class name '%s' is invalid, should be identified as '%s'"), *DeclaredClassName, *ExpectedClassName);
 	}
 
-	if ((Class->ClassFlags & CLASS_NoExport))
+	// This check only works if we already have an object.  This has to be moved to the engine type create code
+	if (InitialEngineClassFlags != CLASS_None)
 	{
+
 		// if the class's class flags didn't contain CLASS_NoExport before it was parsed, it means either:
 		// a) the DECLARE_CLASS macro for this native class doesn't contain the CLASS_NoExport flag (this is an error)
 		// b) this is a new native class, which isn't yet hooked up to static registration (this is OK)
-		if (!(Class->ClassFlags & CLASS_Intrinsic) && (PreviousClassFlags & CLASS_NoExport) == 0 &&
-			(PreviousClassFlags & CLASS_Native) != 0)	// a new native class (one that hasn't been compiled into C++ yet) won't have this set
+		if (HasAnyClassFlags(CLASS_NoExport) && !EnumHasAnyFlags(InitialEngineClassFlags, CLASS_NoExport))
 		{
-			FUHTException::Throwf(*this, TEXT("'noexport': Must include CLASS_NoExport in native class declaration"));
+			if (!HasAnyClassFlags(CLASS_Intrinsic) && EnumHasAnyFlags(InitialEngineClassFlags, CLASS_Native))
+			{
+				FUHTException::Throwf(*this, TEXT("'noexport': Must include CLASS_NoExport in native class declaration"));
+			}
 		}
-	}
 
-	if (!Class->HasAnyClassFlags(CLASS_Abstract) && ((PreviousClassFlags & CLASS_Abstract) != 0))
-	{
-		if (Class->HasAnyClassFlags(CLASS_NoExport))
+		if (!HasAnyClassFlags(CLASS_Abstract) && EnumHasAnyFlags(InitialEngineClassFlags, CLASS_Abstract))
 		{
-			FUHTException::Throwf(*this, TEXT("'abstract': NoExport class missing abstract keyword from class declaration (must change C++ version first)"));
-			Class->ClassFlags |= CLASS_Abstract;
-		}
-		else if (IsNative())
-		{
-			FUHTException::Throwf(*this, TEXT("'abstract': missing abstract keyword from class declaration - class will no longer be exported as abstract"));
+			if (HasAnyClassFlags(CLASS_NoExport))
+			{
+				FUHTException::Throwf(*this, TEXT("'abstract': NoExport class missing abstract keyword from class declaration (must change C++ version first)"));
+				SetClassFlags(CLASS_Abstract);
+			}
+			else if (IsNative())
+			{
+				FUHTException::Throwf(*this, TEXT("'abstract': missing abstract keyword from class declaration - class will no longer be exported as abstract"));
+			}
 		}
 	}
 }
 
 void FUnrealClassDefinitionInfo::SetAndValidateConfigName()
 {
-	UClass* Class = GetClass();
 	if (ConfigName.IsEmpty() == false)
 	{
 		// if the user specified "inherit", we're just going to use the parent class's config filename
 		// this is not actually necessary but it can be useful for explicitly communicating config-ness
 		if (ConfigName == TEXT("inherit"))
 		{
-			UClass* SuperClass = Class->GetSuperClass();
-			if (!SuperClass)
+			FUnrealClassDefinitionInfo* SuperClassDef = GetSuperClass();
+			if (!SuperClassDef)
 			{
-				FUHTException::Throwf(*this, TEXT("Cannot inherit config filename: %s has no super class"), *Class->GetName());
+				FUHTException::Throwf(*this, TEXT("Cannot inherit config filename: %s has no super class"), *GetName());
 			}
 
-			if (SuperClass->ClassConfigName == NAME_None)
+			if (SuperClassDef->GetClassConfigName() == NAME_None)
 			{
-				FUHTException::Throwf(*this, TEXT("Cannot inherit config filename: parent class %s is not marked config."), *SuperClass->GetPathName());
+				FUHTException::Throwf(*this, TEXT("Cannot inherit config filename: parent class %s is not marked config."), *SuperClassDef->GetPathName());
 			}
 		}
 		else
 		{
 			// otherwise, set the config name to the parsed identifier
-			Class->ClassConfigName = FName(*ConfigName);
+			SetClassConfigName(FName(*ConfigName));
 		}
 	}
 	else
 	{
 		// Invalidate config name if not specifically declared.
-		Class->ClassConfigName = NAME_None;
+		SetClassConfigName(NAME_None);
 	}
 }
 
 void FUnrealClassDefinitionInfo::SetAndValidateWithinClass()
 {
-	UClass* Class = GetClass();
-
 	// Process all of the class specifiers
 	if (ClassWithinStr.IsEmpty() == false)
 	{
@@ -2207,57 +2431,30 @@ void FUnrealClassDefinitionInfo::SetAndValidateWithinClass()
 		{
 			FUHTException::Throwf(*this, TEXT("Within class '%s' not found."), *ClassWithinStr);
 		}
-		UClass* RequiredWithinClass = RequiredWithinClassDef->GetClass();
-		if (RequiredWithinClass->IsChildOf(UInterface::StaticClass()))
+		if (RequiredWithinClassDef->IsChildOf(*GUInterfaceDef))
 		{
 			FUHTException::Throwf(*this, TEXT("Classes cannot be 'within' interfaces"));
 		}
-		else if (Class->ClassWithin == NULL || Class->ClassWithin == UObject::StaticClass() || RequiredWithinClass->IsChildOf(Class->ClassWithin))
+		else if (ClassWithin == NULL || ClassWithin == GUObjectDef || RequiredWithinClassDef->IsChildOf(*ClassWithin))
 		{
-			Class->ClassWithin = RequiredWithinClass;
+			SetClassWithin(RequiredWithinClassDef);
 		}
-		else if (Class->ClassWithin != RequiredWithinClass)
+		else if (ClassWithin != RequiredWithinClassDef)
 		{
-			FUHTException::Throwf(*this, TEXT("%s must be within %s, not %s"), *Class->GetPathName(), *Class->ClassWithin->GetPathName(), *RequiredWithinClass->GetPathName());
+			FUHTException::Throwf(*this, TEXT("%s must be within %s, not %s"), *GetPathName(), *ClassWithin->GetPathName(), *RequiredWithinClassDef->GetPathName());
 		}
 	}
 	else
 	{
 		// Make sure there is a valid within
-		Class->ClassWithin = Class->GetSuperClass()
-			? Class->GetSuperClass()->ClassWithin
-			: UObject::StaticClass();
+		SetClassWithin(GetSuperClass() ? GetSuperClass()->GetClassWithin() : GUObjectDef);
 	}
 
-	UClass* ExpectedWithin = Class->GetSuperClass()
-		? Class->GetSuperClass()->ClassWithin
-		: UObject::StaticClass();
+	FUnrealClassDefinitionInfo* ExpectedWithinDef = GetSuperClass() ? GetSuperClass()->GetClassWithin() : GUObjectDef;
 
-	if (!Class->ClassWithin->IsChildOf(ExpectedWithin))
-	{
-		FUHTException::Throwf(*this, TEXT("Parent class declared within '%s'.  Cannot override within class with '%s' since it isn't a child"), *ExpectedWithin->GetName(), *Class->ClassWithin->GetName());
-	}
-
-	ClassWithin = &GTypeDefinitionInfoMap.FindChecked<FUnrealClassDefinitionInfo>(Class->ClassWithin);
-}
-
-void FUnrealClassDefinitionInfo::GetHideCategories(TArray<FString>& OutHideCategories) const
-{
-	UClass* Class = GetClass();
-	if (Class->HasMetaData(FHeaderParserNames::NAME_HideCategories))
+	if (!ClassWithin->IsChildOf(*ExpectedWithinDef))
 	{
-		const FString& LocalHideCategories = Class->GetMetaData(FHeaderParserNames::NAME_HideCategories);
-		LocalHideCategories.ParseIntoArray(OutHideCategories, TEXT(" "), true);
-	}
-}
-
-void FUnrealClassDefinitionInfo::GetShowCategories(TArray<FString>& OutShowCategories) const
-{
-	UClass* Class = GetClass();
-	if (Class->HasMetaData(FHeaderParserNames::NAME_ShowCategories))
-	{
-		const FString& LocalShowCategories = Class->GetMetaData(FHeaderParserNames::NAME_ShowCategories);
-		LocalShowCategories.ParseIntoArray(OutShowCategories, TEXT(" "), true);
+		FUHTException::Throwf(*this, TEXT("Parent class declared within '%s'.  Cannot override within class with '%s' since it isn't a child"), *ExpectedWithinDef->GetName(), *ClassWithin->GetName());
 	}
 }
 
@@ -2274,17 +2471,12 @@ void FUnrealClassDefinitionInfo::MergeCategoryMetaData(TMap<FName, FString>& InM
 
 void FUnrealClassDefinitionInfo::GetSparseClassDataTypes(TArray<FString>& OutSparseClassDataTypes) const
 {
-	UClass* Class = GetClass();
-	if (Class->HasMetaData(FHeaderParserNames::NAME_SparseClassDataTypes))
-	{
-		const FString& LocalSparseClassDataTypes = Class->GetMetaData(FHeaderParserNames::NAME_SparseClassDataTypes);
-		LocalSparseClassDataTypes.ParseIntoArray(OutSparseClassDataTypes, TEXT(" "), true);
-	}
+	GetStringArrayMetaData(FHeaderParserNames::NAME_SparseClassDataTypes, OutSparseClassDataTypes);
 }
 
 FString FUnrealClassDefinitionInfo::GetNameWithPrefix(EEnforceInterfacePrefix EnforceInterfacePrefix) const
 {
-	const TCHAR* Prefix = 0;
+	const TCHAR* Prefix = nullptr;
 
 	if (HasAnyClassFlags(CLASS_Interface))
 	{
@@ -2346,48 +2538,58 @@ FUnrealFunctionDefinitionInfo* FUnrealFunctionDefinitionInfo::GetSuperFunction()
 	return UHTCast<FUnrealFunctionDefinitionInfo>(GetSuperStruct());
 }
 
-void FUnrealFunctionDefinitionInfo::CreateUObjectEngineTypesInternal()
+void FUnrealFunctionDefinitionInfo::CreateUObjectEngineTypesInternal(ECreateEngineTypesPhase Phase)
 {
 	// Invoke the base class creatation
-	FUnrealStructDefinitionInfo::CreateUObjectEngineTypesInternal();
+	FUnrealStructDefinitionInfo::CreateUObjectEngineTypesInternal(Phase);
 
-	// We have to precreate the function prior to invoking the parent finalize
-	UFunction* Function = nullptr;
-	switch (FunctionType)
+	switch (Phase)
 	{
-	case EFunctionType::Function:
-		Function = new(EC_InternalUseOnlyConstructor, GetOuter()->GetObject(), *GetNameCPP(), RF_Public) UFunction(FObjectInitializer(), nullptr);
-		break;
-	case EFunctionType::Delegate:
-		Function = new(EC_InternalUseOnlyConstructor, GetOuter()->GetObject(), *GetNameCPP(), RF_Public) UDelegateFunction(FObjectInitializer(), nullptr);
+	case ECreateEngineTypesPhase::Phase1:
 		break;
-	case EFunctionType::SparseDelegate:
+
+	case ECreateEngineTypesPhase::Phase2:
 	{
-		USparseDelegateFunction* USPF = new(EC_InternalUseOnlyConstructor, GetOuter()->GetObject(), *GetNameCPP(), RF_Public) USparseDelegateFunction(FObjectInitializer(), nullptr);
-		USPF->OwningClassName = SparseOwningClassName;
-		USPF->DelegateName = SparseDelegateName;
-		Function = USPF;
-		break;
-	}
-	}
-	check(Function);
+		// We have to precreate the function prior to invoking the parent finalize
+		UFunction* Function = nullptr;
+		switch (FunctionType)
+		{
+		case EFunctionType::Function:
+			Function = new(EC_InternalUseOnlyConstructor, GetOuter()->GetObject(), *GetNameCPP(), RF_Public) UFunction(FObjectInitializer(), nullptr);
+			break;
+		case EFunctionType::Delegate:
+			Function = new(EC_InternalUseOnlyConstructor, GetOuter()->GetObject(), *GetNameCPP(), RF_Public) UDelegateFunction(FObjectInitializer(), nullptr);
+			break;
+		case EFunctionType::SparseDelegate:
+		{
+			USparseDelegateFunction* USPF = new(EC_InternalUseOnlyConstructor, GetOuter()->GetObject(), *GetNameCPP(), RF_Public) USparseDelegateFunction(FObjectInitializer(), nullptr);
+			USPF->OwningClassName = SparseOwningClassName;
+			USPF->DelegateName = SparseDelegateName;
+			Function = USPF;
+			break;
+		}
+		}
+		check(Function);
 
-	Function->ReturnValueOffset = MAX_uint16;
-	Function->FirstPropertyToInit = nullptr;
-	Function->FunctionFlags |= FunctionData.FunctionFlags;
-	Function->GetPackage()->GetMetaData()->SetObjectValues(Function, MetaData);
+		Function->ReturnValueOffset = MAX_uint16;
+		Function->FirstPropertyToInit = nullptr;
+		Function->FunctionFlags |= FunctionData.FunctionFlags;
+		Function->GetPackage()->GetMetaData()->SetObjectValues(Function, MetaData);
 
-	SetObject(Function);
-	GTypeDefinitionInfoMap.Add(Function, SharedThis(this));
+		SetObject(Function);
+		GTypeDefinitionInfoMap.Add(Function, SharedThis(this));
 
-	if (FUnrealStructDefinitionInfo* StructDef = UHTCast<FUnrealStructDefinitionInfo>(GetOuter()))
-	{
-		UStruct* Struct = StructDef->GetStruct();
-		Function->Next = Struct->Children;
-		Struct->Children = Function;
+		if (FUnrealStructDefinitionInfo* StructDef = UHTCast<FUnrealStructDefinitionInfo>(GetOuter()))
+		{
+			UStruct* Struct = StructDef->GetStruct();
+			Function->Next = Struct->Children;
+			Struct->Children = Function;
+		}
+		Function->NumParms = uint8(GetProperties().Num());
+		Function->Bind();
+		break;
+	}
 	}
-	Function->NumParms = uint8(GetProperties().Num());
-	Function->Bind();
 }
 
 void FUnrealFunctionDefinitionInfo::PostParseFinalizeInternal()
diff --git a/Engine/Source/Programs/UnrealHeaderTool/Private/UnrealTypeDefinitionInfo.h b/Engine/Source/Programs/UnrealHeaderTool/Private/UnrealTypeDefinitionInfo.h
--- a/Engine/Source/Programs/UnrealHeaderTool/Private/UnrealTypeDefinitionInfo.h
+++ b/Engine/Source/Programs/UnrealHeaderTool/Private/UnrealTypeDefinitionInfo.h
@@ -195,6 +195,12 @@ public:
 		return (BoolString == "true");
 	}
 
+	/**
+	 * Split the ginen meta data into a string array
+	 */
+	void GetStringArrayMetaData(const FName& Key, TArray<FString>& Out) const;
+	TArray<FString> GetStringArrayMetaData(const FName& Key) const;
+
 	virtual TMap<FName, FString>* GetMetaDataMap() const = 0;
 
 protected:
@@ -228,6 +234,14 @@ class FUnrealTypeDefinitionInfo
 	: public TSharedFromThis<FUnrealTypeDefinitionInfo>
 	, public FUHTExceptionContext
 {
+public:
+	enum class ECreateEngineTypesPhase : uint8
+	{
+		Phase1,				// Create the Class, ScriptStruct, and Enum types
+		Phase2,				// Create the functions and properties
+		MAX_Phases,
+	};
+
 private:
 	enum class EFinalizeState : uint8
 	{
@@ -393,9 +407,9 @@ public:
 	}
 
 	/**
-	 * Create UObject engine types
+	 * Create UObject engine types.
 	 */
-	void CreateUObjectEngineTypes();
+	void CreateUObjectEngineTypes(ECreateEngineTypesPhase Phase);
 
 	/**
 	 * Perform any post parsing finalization and validation
@@ -568,11 +582,24 @@ protected:
 
 	virtual void ValidateMetaDataFormat(const FName InKey, ECheckedMetadataSpecifier InCheckedMetadataSpecifier, const FString& InValue) const;
 
-	virtual void PostParseFinalizeInternal()
-	{}
+	/**
+	 * Create UObject engine types
+	 */
+	virtual void CreateUObjectEngineTypesInternal(ECreateEngineTypesPhase Phase) 
+	{
+		if (Phase == ECreateEngineTypesPhase::Phase1)
+		{
+			if (Outer != nullptr && Outer->AsPackage() == nullptr)
+			{
+				CreateUObjectEngineTypes(Phase);
+			}
+		}
+	}
 
-	virtual void CreateUObjectEngineTypesInternal()
-	{}
+	/**
+	 * Perform any post parsing finalization and validation
+	 */
+	virtual void PostParseFinalizeInternal() {}
 
 private:
 	FString NameCPP;
@@ -580,7 +607,7 @@ private:
 	FUnrealSourceFile* SourceFile = nullptr;
 	int32 LineNumber = 0;
 	std::atomic<uint32> Hash = 0;
-	EFinalizeState CreateUObectEngineTypesState = EFinalizeState::None;
+	EFinalizeState CreateUObectEngineTypesState[uint8(ECreateEngineTypesPhase::MAX_Phases)] = { EFinalizeState::None, EFinalizeState::None };
 	EFinalizeState PostParseFinalizeState = EFinalizeState::None;
 };
 
@@ -1536,7 +1563,7 @@ public:
 	/**
 	 * Return a collection of all classes associated with this package.  This is not valid until parsing begins.
 	 */
-	TArray<UClass*>& GetAllClasses()
+	TArray<TSharedRef<FUnrealTypeDefinitionInfo>>& GetAllClasses()
 	{
 		return AllClasses;
 	}
@@ -1606,7 +1633,7 @@ protected:
 	/**
 	 * Create UObject engine types
 	 */
-	virtual void CreateUObjectEngineTypesInternal() override;
+	virtual void CreateUObjectEngineTypesInternal(ECreateEngineTypesPhase Phase) override;
 
 	/**
 	 * Perform any post parsing finalization and validation
@@ -1616,7 +1643,7 @@ protected:
 private:
 	const FManifestModule& Module;
 	TArray<TSharedRef<FUnrealSourceFile>> AllSourceFiles;
-	TArray<UClass*> AllClasses;
+	TArray<TSharedRef<FUnrealTypeDefinitionInfo>> AllClasses;
 	FString SingletonName;
 	FString SingletonNameChopped;
 	FString ExternDecl;
@@ -1935,7 +1962,10 @@ public:
 	}
 
 protected:
-	virtual void CreateUObjectEngineTypesInternal() override;
+	/**
+	 * Create the shell of the UObject.  Only invoked if the object has not been set
+	 */
+	virtual void CreateUObjectEngineTypesInternal(ECreateEngineTypesPhase Phase) override;
 
 	virtual TMap<FName, FString>* GetMetaDataMap() const override
 	{
@@ -1988,8 +2018,6 @@ public:
 
 
 public:
-	using FUnrealFieldDefinitionInfo::FUnrealFieldDefinitionInfo;
-
 	virtual FUnrealStructDefinitionInfo* AsStruct() override
 	{
 		return this;
@@ -2007,6 +2035,19 @@ public:
 		return static_cast<UStruct*>(GetObject());
 	}
 
+	/**
+	 * Return the Engine instance associated with the compiler instance without the null check
+	 */
+	UStruct* GetStructSafe() const
+	{
+		return static_cast<UStruct*>(GetObjectSafe());
+	}
+
+	/**
+	 * @return true if this object is of the specified type.
+	 */
+	bool IsChildOf(const FUnrealStructDefinitionInfo& SomeBase) const;
+
 	/**
 	 * Return the engine class name
 	 */
@@ -2019,9 +2060,17 @@ public:
 	/**
 	 * Return the CPP prefix 
 	 */
-	const TCHAR* GetPrefixCPP() const
+	virtual const TCHAR* GetPrefixCPP() const
 	{
-		return GetStruct()->GetPrefixCPP();
+		const TCHAR* ReturnValue = TEXT("F");
+#if UHT_ENABLE_ENGINE_TYPE_CHECKS
+		if (UStruct* Struct = GetStructSafe())
+		{
+			const TCHAR* EngineValue = Struct->GetPrefixCPP();
+			check(FCString::Strcmp(EngineValue, ReturnValue) == 0);
+		}
+#endif
+		return ReturnValue;
 	}
 
 	/**
@@ -2108,14 +2157,14 @@ public:
 	/**
 	 * Return the base struct information
 	 */
-	const TArray<FBaseStructInfo>& GetBaseStructInfo() const
+	const TArray<FBaseStructInfo>& GetBaseStructInfos() const
 	{
-		return BaseStructInfo;
+		return BaseStructInfos;
 	}
 
-	TArray<FBaseStructInfo>& GetBaseStructInfo()
+	TArray<FBaseStructInfo>& GetBaseStructInfos()
 	{
-		return BaseStructInfo;
+		return BaseStructInfos;
 	}
 
 	/**
@@ -2134,21 +2183,6 @@ public:
 		bContainsDelegates = true;
 	}
 
-	/**
-	 * Test to see if this struct is a child of another struct
-	 */
-	bool IsChildOf(FUnrealStructDefinitionInfo& ParentStruct) const
-	{
-		for (const FUnrealStructDefinitionInfo* Current = this; Current; Current = Current->GetSuperStructInfo().Struct)
-		{
-			if (Current == &ParentStruct)
-			{
-				return true;
-			}
-		}
-		return false;
-	}
-
 	const FString* FindMetaDataHierarchical(const FName& Key) const
 	{
 		if (Key == NAME_None)
@@ -2245,10 +2279,16 @@ public:
 	}
 
 protected:
+	explicit FUnrealStructDefinitionInfo(UObject* Object)
+		: FUnrealFieldDefinitionInfo(Object)
+	{}
+
+	FUnrealStructDefinitionInfo(FUnrealSourceFile& InSourceFile, int32 InLineNumber, FString&& InNameCPP, FName InName, FUnrealObjectDefinitionInfo& InOuter);
+
 	/**
 	 * Create UObject engine types
 	 */
-	virtual void CreateUObjectEngineTypesInternal() override;
+	virtual void CreateUObjectEngineTypesInternal(ECreateEngineTypesPhase Phase) override;
 
 	/**
 	 * Perform any post parsing finalization and validation
@@ -2266,7 +2306,7 @@ private:
 
 	FStructMetaData StructMetaData;
 	FBaseStructInfo SuperStructInfo;
-	TArray<FBaseStructInfo> BaseStructInfo;
+	TArray<FBaseStructInfo> BaseStructInfos;
 
 	FDefinitionRange DefinitionRange;
 
@@ -2313,12 +2353,38 @@ public:
 		return static_cast<UScriptStruct*>(GetObject());
 	}
 
+	/**
+	 * Return the Engine instance associated with the compiler instance without the null check
+	 */
+	UScriptStruct* GetScriptStructSafe() const
+	{
+		return static_cast<UScriptStruct*>(GetObjectSafe());
+	}
+
 	/**
 	 * Return the struct flags
 	 */
-	FORCEINLINE EStructFlags GetStructFlags() const
+	EStructFlags GetStructFlags() const
 	{
-		return GetScriptStruct()->StructFlags;
+#if UHT_ENABLE_ENGINE_TYPE_CHECKS
+		check(GetScriptStructSafe() == nullptr || (GetScriptStructSafe()->StructFlags & ~STRUCT_ComputedFlags) == StructFlags);
+#endif
+		return StructFlags;
+	}
+
+	/**
+	 * Sets the given struct flags
+	 */
+	void SetStructFlags(EStructFlags FlagsToSet)
+	{
+#if UHT_ENABLE_ENGINE_TYPE_CHECKS
+		if (UScriptStruct* Struct = GetScriptStructSafe())
+		{
+			check((Struct->StructFlags & ~STRUCT_ComputedFlags) == StructFlags);
+			Struct->StructFlags = EStructFlags(int32(Struct->StructFlags) | int32(FlagsToSet));
+		}
+#endif
+		StructFlags = EStructFlags(int32(StructFlags) | int32(FlagsToSet));
 	}
 
 	/**
@@ -2329,9 +2395,17 @@ public:
 	 * @return	true if the passed in flag is set, false otherwise
 	 *			(including no flag passed in, unless the FlagsToCheck is CLASS_AllFlags)
 	 */
-	FORCEINLINE bool HasAnyStructFlags(EStructFlags FlagsToCheck) const
+	bool HasAnyStructFlags(EStructFlags FlagsToCheck) const
 	{
-		return EnumHasAnyFlags(GetScriptStruct()->StructFlags, FlagsToCheck);
+#if UHT_ENABLE_ENGINE_TYPE_CHECKS
+		if (GetScriptStructSafe())
+		{
+			int32 a = GetScriptStructSafe()->StructFlags & ~STRUCT_ComputedFlags;
+			check(StructFlags == a);
+		}
+		check(GetScriptStructSafe() == nullptr || (GetScriptStructSafe()->StructFlags & ~STRUCT_ComputedFlags) == StructFlags);
+#endif
+		return EnumHasAnyFlags(StructFlags, FlagsToCheck);
 	}
 
 	/**
@@ -2340,9 +2414,12 @@ public:
 	 * @param FlagsToCheck	Function flags to check for
 	 * @return true if all of the passed in flags are set (including no flags passed in), false otherwise
 	 */
-	FORCEINLINE bool HasAllStructFlags(EStructFlags FlagsToCheck) const
+	bool HasAllStructFlags(EStructFlags FlagsToCheck) const
 	{
-		return EnumHasAllFlags(GetScriptStruct()->StructFlags, FlagsToCheck);
+#if UHT_ENABLE_ENGINE_TYPE_CHECKS
+		check(GetScriptStructSafe() == nullptr || (GetScriptStructSafe()->StructFlags & ~STRUCT_ComputedFlags) == StructFlags);
+#endif
+		return EnumHasAllFlags(StructFlags, FlagsToCheck);
 	}
 
 	/**
@@ -2352,17 +2429,19 @@ public:
 	 * @param ExpectedFlags The flags from the flags to check that should be set
 	 * @return true if specific of the passed in flags are set (including no flags passed in), false otherwise
 	 */
-	FORCEINLINE bool HasSpecificStructFlags(EStructFlags FlagsToCheck, EStructFlags ExpectedFlags) const
+	bool HasSpecificStructFlags(EStructFlags FlagsToCheck, EStructFlags ExpectedFlags) const
 	{
-		EStructFlags Flags = GetScriptStruct()->StructFlags;
-		return (Flags & FlagsToCheck) == ExpectedFlags;
+#if UHT_ENABLE_ENGINE_TYPE_CHECKS
+		check(GetScriptStructSafe() == nullptr || (GetScriptStructSafe()->StructFlags & ~STRUCT_ComputedFlags) == StructFlags);
+#endif
+		return (StructFlags & FlagsToCheck) == ExpectedFlags;
 	}
 
 	/**
 	 * If it is native, it is assumed to have defaults because it has a constructor
 	 * @return true if this struct has defaults
 	 */
-	FORCEINLINE bool HasDefaults() const
+	bool HasDefaults() const
 	{
 		return !!GetScriptStruct()->GetCppStructOps();
 	}
@@ -2378,8 +2457,26 @@ public:
 		MacroDeclaredLineNumber = InMacroDeclaredLineNumber;
 	}
 
+	virtual TMap<FName, FString>* GetMetaDataMap() const override
+	{
+		return const_cast<TMap<FName, FString>*>(&MetaData);
+	}
+
+protected:
+	/**
+	 * Create UObject engine types
+	 */
+	virtual void CreateUObjectEngineTypesInternal(ECreateEngineTypesPhase Phase) override;
+
+	/**
+	 * Perform any post parsing finalization and validation
+	 */
+	virtual void PostParseFinalizeInternal() override;
+
 private:
+	TMap<FName, FString> MetaData;
 	int32 MacroDeclaredLineNumber = INDEX_NONE;
+	EStructFlags StructFlags = STRUCT_NoFlags;
 };
 
 /**
@@ -2447,7 +2544,7 @@ public:
 	/**
 	 * Return the function flags
 	 */
-	FORCEINLINE EFunctionFlags GetFunctionFlags() const
+	EFunctionFlags GetFunctionFlags() const
 	{
 #if UHT_ENABLE_ENGINE_TYPE_CHECKS
 		check(GetFunctionSafe() == nullptr || GetFunctionSafe()->FunctionFlags == FunctionData.FunctionFlags);
@@ -2478,7 +2575,7 @@ public:
 	 * @return	true if the passed in flag is set, false otherwise
 	 *			(including no flag passed in, unless the FlagsToCheck is CLASS_AllFlags)
 	 */
-	FORCEINLINE bool HasAnyFunctionFlags(EFunctionFlags FlagsToCheck) const
+	bool HasAnyFunctionFlags(EFunctionFlags FlagsToCheck) const
 	{
 #if UHT_ENABLE_ENGINE_TYPE_CHECKS
 		check(GetFunctionSafe() == nullptr || GetFunctionSafe()->FunctionFlags == FunctionData.FunctionFlags);
@@ -2492,7 +2589,7 @@ public:
 	 * @param FlagsToCheck	Function flags to check for
 	 * @return true if all of the passed in flags are set (including no flags passed in), false otherwise
 	 */
-	FORCEINLINE bool HasAllFunctionFlags(EFunctionFlags FlagsToCheck) const
+	bool HasAllFunctionFlags(EFunctionFlags FlagsToCheck) const
 	{
 #if UHT_ENABLE_ENGINE_TYPE_CHECKS
 		check(GetFunctionSafe() == nullptr || GetFunctionSafe()->FunctionFlags == FunctionData.FunctionFlags);
@@ -2507,7 +2604,7 @@ public:
 	 * @param ExpectedFlags The flags from the flags to check that should be set
 	 * @return true if specific of the passed in flags are set (including no flags passed in), false otherwise
 	 */
-	FORCEINLINE bool HasSpecificFunctionFlags(EFunctionFlags FlagsToCheck, EFunctionFlags ExpectedFlags) const
+	bool HasSpecificFunctionFlags(EFunctionFlags FlagsToCheck, EFunctionFlags ExpectedFlags) const
 	{
 #if UHT_ENABLE_ENGINE_TYPE_CHECKS
 		check(GetFunctionSafe() == nullptr || GetFunctionSafe()->FunctionFlags == FunctionData.FunctionFlags);
@@ -2583,7 +2680,7 @@ protected:
 	/**
 	 * Create UObject engine types
 	 */
-	virtual void CreateUObjectEngineTypesInternal() override;
+	virtual void CreateUObjectEngineTypesInternal(ECreateEngineTypesPhase Phase) override;
 
 	/**
 	 * Perform any post parsing finalization and validation
@@ -2613,9 +2710,11 @@ class FUnrealClassDefinitionInfo
 public:
 	FUnrealClassDefinitionInfo(FUnrealSourceFile& InSourceFile, int32 InLineNumber, FString&& InNameCPP, FName InName, bool bInIsInterface);
 
-	explicit FUnrealClassDefinitionInfo(UClass* Class) 
+	explicit FUnrealClassDefinitionInfo(UClass* Class)
 		: FUnrealStructDefinitionInfo(Class)
-	{ }
+	{
+		InitializeFromExistingUObject(Class);
+	}
 
 	/**
 	 * Attempts to find a class definition based on the given name
@@ -2647,7 +2746,6 @@ public:
 	{
 		return this;
 	}
-
 	/**
 	 * Return the engine class name
 	 */
@@ -2657,6 +2755,11 @@ public:
 		return ClassName;
 	}
 
+	/**
+	 * Return the CPP prefix
+	 */
+	virtual const TCHAR* GetPrefixCPP() const;
+
 	virtual uint32 GetHash(bool bIncludeNoExport = true) const override;
 
 	/**
@@ -2667,15 +2770,21 @@ public:
 		return static_cast<UClass*>(GetObject());
 	}
 
+	/**
+	 * Return the Engine instance associated with the compiler instance without null check
+	 */
+	UClass* GetClassSafe() const
+	{
+		return static_cast<UClass*>(GetObjectSafe());
+	}
+
 	/**
 	 * Set the Engine instance associated with the compiler instance
 	 */
 	virtual void SetObject(UObject* InObject) override
 	{
-		UClass* Class = static_cast<UClass*>(InObject);
-		InitialEngineClassFlags = Class->ClassFlags;
-		Class->ClassFlags |= ParsedClassFlags;
 		FUnrealStructDefinitionInfo::SetObject(InObject);
+		CastChecked<UClass>(InObject)->ClassFlags |= ClassFlags;
 	}
 
 	FUnrealClassDefinitionInfo* GetSuperClass() const;
@@ -2704,7 +2813,7 @@ public:
 		return EnclosingDefine;
 	}
 
-	/** 
+	/**
 	 * Set the enclosing define
 	 */
 	void SetEnclosingDefine(FString&& InEnclosingDefine)
@@ -2725,7 +2834,10 @@ public:
 	 */
 	EClassFlags GetClassFlags() const
 	{
-		return GetClass()->ClassFlags;
+#if UHT_ENABLE_ENGINE_TYPE_CHECKS
+		check(GetClassSafe() == nullptr || (GetClassSafe()->ClassFlags & CLASS_SaveInCompiledInClasses) == (ClassFlags & CLASS_SaveInCompiledInClasses));
+#endif
+		return ClassFlags;
 	}
 
 	/**
@@ -2733,7 +2845,14 @@ public:
 	 */
 	void SetClassFlags(EClassFlags FlagsToSet)
 	{
-		GetClass()->ClassFlags |= FlagsToSet;
+#if UHT_ENABLE_ENGINE_TYPE_CHECKS
+		check(GetClassSafe() == nullptr || (GetClassSafe()->ClassFlags & CLASS_SaveInCompiledInClasses) == (ClassFlags & CLASS_SaveInCompiledInClasses));
+#endif
+		ClassFlags |= FlagsToSet;
+		if (UClass* Class = GetClassSafe())
+		{
+			Class->ClassFlags |= FlagsToSet;
+		}
 	}
 
 	/**
@@ -2741,7 +2860,14 @@ public:
 	 */
 	void ClearClassFlags(EClassFlags FlagsToClear)
 	{
-		GetClass()->ClassFlags &= ~FlagsToClear;
+#if UHT_ENABLE_ENGINE_TYPE_CHECKS
+		check(GetClassSafe() == nullptr || (GetClassSafe()->ClassFlags & CLASS_SaveInCompiledInClasses) == (ClassFlags & CLASS_SaveInCompiledInClasses));
+#endif
+		ClassFlags &= ~FlagsToClear;
+		if (UClass* Class = GetClassSafe())
+		{
+			Class->ClassFlags &= ~FlagsToClear;
+		}
 	}
 
 	/**
@@ -2753,10 +2879,10 @@ public:
 	 */
 	bool HasAnyClassFlags(EClassFlags FlagsToCheck) const
 	{
-		bool bA = EnumHasAnyFlags(ParsedClassFlags, FlagsToCheck);
-		bool bB = GetClass()->HasAnyClassFlags(FlagsToCheck);
-		check(!bA || bB);
-		return EnumHasAnyFlags(ParsedClassFlags, FlagsToCheck) || GetClass()->HasAnyClassFlags(FlagsToCheck);
+#if UHT_ENABLE_ENGINE_TYPE_CHECKS
+		check(GetClassSafe() == nullptr || (GetClassSafe()->ClassFlags & CLASS_SaveInCompiledInClasses) == (ClassFlags & CLASS_SaveInCompiledInClasses));
+#endif
+		return EnumHasAnyFlags(ClassFlags, FlagsToCheck);
 	}
 
 	/**
@@ -2767,10 +2893,10 @@ public:
 	 */
 	bool HasAllClassFlags(EClassFlags FlagsToCheck) const
 	{
-		bool bA = EnumHasAllFlags(ParsedClassFlags, FlagsToCheck);
-		bool bB = GetClass()->HasAllClassFlags(FlagsToCheck);
-		check(!bA || bB);
-		return EnumHasAllFlags(ParsedClassFlags, FlagsToCheck) || GetClass()->HasAllClassFlags(FlagsToCheck);
+#if UHT_ENABLE_ENGINE_TYPE_CHECKS
+		check(GetClassSafe() == nullptr || (GetClassSafe()->ClassFlags & CLASS_SaveInCompiledInClasses) == (ClassFlags & CLASS_SaveInCompiledInClasses));
+#endif
+		return EnumHasAllFlags(ClassFlags, FlagsToCheck);
 	}
 
 	/**
@@ -2810,6 +2936,32 @@ public:
 		return false;
 	}
 
+	/**
+	 * Get the class cast flags
+	 */
+	EClassCastFlags GetClassCastFlags() const
+	{
+#if UHT_ENABLE_ENGINE_TYPE_CHECKS
+		check(GetClassSafe() == nullptr || GetClassSafe()->ClassCastFlags == ClassCastFlags);
+#endif
+		return ClassCastFlags;
+	}
+
+	/**
+	 * Set addition class cast flags
+	 */
+	void SetClassCastFlags(EClassCastFlags FlagsToSet)
+	{
+#if UHT_ENABLE_ENGINE_TYPE_CHECKS
+		check(GetClassSafe() == nullptr || GetClassSafe()->ClassCastFlags == ClassCastFlags);
+#endif
+		ClassCastFlags |= FlagsToSet;
+		if (UClass* Class = GetClassSafe())
+		{
+			Class->ClassCastFlags |= FlagsToSet;
+		}
+	}
+
 	/**
 	 * Return the flags that were parsed as part of the pre-parse phase
 	 */
@@ -2819,12 +2971,11 @@ public:
 	}
 
 	/**
-	 * Return the initial flags from the class. These will be set for class definitions created directly from
-	 * existing engine types.
+	 * Return the meta data from the preparse phase
 	 */
-	EClassFlags GetInitialEngineClassFlags() const
+	TMap<FName, FString>& GetParsedMetaData()
 	{
-		return InitialEngineClassFlags;
+		return ParsedMetaData;
 	}
 
 	/**
@@ -2847,7 +2998,7 @@ public:
 	* @param	DeclaredClassName Name this class was declared with (for validation)
 	* @param	PreviousClassFlags Class flags before resetting the class (for validation)
 	*/
-	void MergeAndValidateClassFlags(const FString& DeclaredClassName, EClassFlags PreviousClassFlags);
+	void MergeAndValidateClassFlags(const FString& DeclaredClassName);
 
 	/**
 	 * Add the category meta data
@@ -2870,6 +3021,42 @@ public:
 	void SetClassWithin(FUnrealClassDefinitionInfo* InClassWithin)
 	{
 		ClassWithin = InClassWithin;
+		if (UClass* Class = GetClassSafe())
+		{
+			Class->ClassWithin = ClassWithin->GetClass();
+		}
+	}
+
+	/**
+	 * Get the class repliaction properties.  This collection includes the parents
+	 */
+	const TArray<FUnrealPropertyDefinitionInfo*>& GetClassReps() const
+	{
+		return ClassReps;
+	}
+
+	/**
+	 * Get the index of the first class rep that is part of this class
+	 */
+	int32 GetFirstOwnedClassRep() const
+	{
+		return FirstOwnedClassRep;
+	}
+
+	/**
+	 * Return true if we have owned class reps
+	 */
+	bool HasOwnedClassReps() const
+	{
+		return FirstOwnedClassRep < ClassReps.Num();
+	}
+
+	/**
+	 * Return true if we have class reps
+	 */
+	bool HasClassReps() const
+	{
+		return !ClassReps.IsEmpty();
 	}
 
 	/**
@@ -2877,7 +3064,10 @@ public:
 	 */
 	FName GetClassConfigName() const
 	{
-		return GetClass()->ClassConfigName;
+#if UHT_ENABLE_ENGINE_TYPE_CHECKS
+		check(GetClassSafe() == nullptr || GetClassSafe()->ClassConfigName == ClassConfigName);
+#endif
+		return ClassConfigName;
 	}
 
 	/**
@@ -2885,20 +3075,71 @@ public:
 	 */
 	void SetClassConfigName(FName InClassConfigName)
 	{
-		GetClass()->ClassConfigName = InClassConfigName;
+#if UHT_ENABLE_ENGINE_TYPE_CHECKS
+		check(GetClassSafe() == nullptr || GetClassSafe()->ClassConfigName == ClassConfigName);
+#endif
+		ClassConfigName = InClassConfigName;
+		if (UClass* Class = GetClassSafe())
+		{
+			Class->ClassConfigName = InClassConfigName;
+		}
+	}
+
+	/**
+	 * Get the properties size
+	 */
+	int32 GetPropertiesSize() const
+	{
+#if UHT_ENABLE_ENGINE_TYPE_CHECKS
+		check(GetClassSafe() == nullptr || GetClassSafe()->PropertiesSize == PropertiesSize);
+#endif
+		return PropertiesSize;
+	}
+
+	/**
+	 * Set the properties size
+	 */
+	void SetPropertiesSize(int32 InPropertiesSize)
+	{
+#if UHT_ENABLE_ENGINE_TYPE_CHECKS
+		check(GetClassSafe() == nullptr || GetClassSafe()->PropertiesSize == PropertiesSize);
+#endif
+		PropertiesSize = InPropertiesSize;
+		if (UClass* Class = GetClassSafe())
+		{
+			Class->PropertiesSize = InPropertiesSize;
+		}
 	}
 
+	/**
+	 * Test to see if the class implements a specific interface
+	 */
+	bool ImplementsInterface(const FUnrealClassDefinitionInfo& SomeInterface) const;
+
+	/**
+	 * Invoked when the definition has an existing engine class
+	 */
+	void InitializeFromExistingUObject(UClass* Class);
+
+
 	FString GetNameWithPrefix(EEnforceInterfacePrefix EnforceInterfacePrefix = EEnforceInterfacePrefix::None) const;
 
+	virtual TMap<FName, FString>* GetMetaDataMap() const override
+	{
+		return const_cast<TMap<FName, FString>*>(&MetaData);
+	}
+
 protected:
+	/**
+	 * Create UObject engine types
+	 */
+	virtual void CreateUObjectEngineTypesInternal(ECreateEngineTypesPhase Phase) override;
+
 	/**
 	 * Perform any post parsing finalization and validation
 	 */
 	virtual void PostParseFinalizeInternal() override;
 
-public:
-	TMap<FName, FString> MetaData;
-
 private:
 	/** Merges all 'show' categories */
 	void MergeShowCategories();
@@ -2907,9 +3148,8 @@ private:
 	/** Sets and validates 'ConfigName' property */
 	void SetAndValidateConfigName();
 
-	void GetHideCategories(TArray<FString>& OutHideCategories) const;
-	void GetShowCategories(TArray<FString>& OutShowCategories) const;
-
+	TMap<FName, FString> MetaData; // Actual live meta data
+	TMap<FName, FString> ParsedMetaData; // Meta data from the preparse phase
 	TArray<FString> ShowCategories;
 	TArray<FString> ShowFunctions;
 	TArray<FString> DontAutoCollapseCategories;
@@ -2921,13 +3161,19 @@ private:
 	TArray<FString> DependsOn;
 	TArray<FString> ClassGroupNames;
 	TArray<FString> SparseClassDataTypes;
+	TArray<FUnrealPropertyDefinitionInfo*> ClassReps;
+	FName ClassConfigName = NAME_None;
 	FString EnclosingDefine;
 	FString ClassWithinStr;
 	FString ConfigName;
-	EClassFlags ParsedClassFlags = CLASS_None;
-	EClassFlags InitialEngineClassFlags = CLASS_None;
+	EClassFlags ClassFlags = CLASS_None; // Current class flags
+	EClassFlags ParsedClassFlags = CLASS_None; // Class flags from the preparse phase
+	EClassFlags InitialEngineClassFlags = CLASS_None; // Class flags from a preexisting UObject
+	EClassCastFlags ClassCastFlags = CASTCLASS_None;
 	FUnrealClassDefinitionInfo* ClassWithin = nullptr;
 	ESerializerArchiveType ArchiveType = ESerializerArchiveType::None;
+	int32 FirstOwnedClassRep = 0;
+	int32 PropertiesSize = 0;
 	bool bIsInterface = false;
 	bool bWantsToBePlaceable = false;
 };
@@ -3106,12 +3352,12 @@ public:
 	TUHTFieldIterator(const TUHTFieldIterator& rhs) = default;
 
 	/** conversion to "bool" returning true if the iterator is valid. */
-	FORCEINLINE explicit operator bool() const
+	explicit operator bool() const
 	{
 		return Field != NULL;
 	}
 	/** inverse of the "bool" operator */
-	FORCEINLINE bool operator !() const
+	bool operator !() const
 	{
 		return !(bool)*this;
 	}
