Fixing MAC build

#rnx

diff --git a/Engine/Source/Programs/UnrealHeaderTool/Private/PropertyTypes.cpp b/Engine/Source/Programs/UnrealHeaderTool/Private/PropertyTypes.cpp
--- a/Engine/Source/Programs/UnrealHeaderTool/Private/PropertyTypes.cpp
+++ b/Engine/Source/Programs/UnrealHeaderTool/Private/PropertyTypes.cpp
@@ -11,9 +11,9 @@
 #include "UObject/ObjectMacros.h"
 
 // Globals for common class definitions
-extern FUnrealClassDefinitionInfo* GUObjectDef = nullptr;
-extern FUnrealClassDefinitionInfo* GUClassDef = nullptr;
-extern FUnrealClassDefinitionInfo* GUInterfaceDef = nullptr;
+extern FUnrealClassDefinitionInfo* GUObjectDef;
+extern FUnrealClassDefinitionInfo* GUClassDef;
+extern FUnrealClassDefinitionInfo* GUInterfaceDef;
 
 void AddEditInlineMetaData(TMap<FName, FString>& MetaData);
 
