Cleanup of exception and logging.

#rnx

diff --git a/Engine/Source/Programs/UnrealHeaderTool/Private/BaseParser.cpp b/Engine/Source/Programs/UnrealHeaderTool/Private/BaseParser.cpp
--- a/Engine/Source/Programs/UnrealHeaderTool/Private/BaseParser.cpp
+++ b/Engine/Source/Programs/UnrealHeaderTool/Private/BaseParser.cpp
@@ -257,7 +257,7 @@ Loop:
 			if (!bInsideComment)
 			{
 				ClearComment();
-				FUHTException::Throwf(*this, TEXT("Unexpected '*/' outside of comment") );
+				Throwf(TEXT("Unexpected '*/' outside of comment") );
 			}
 
 			/** Asterisk and slash always end comment. */
@@ -276,7 +276,7 @@ Loop:
 		if (c == 0)
 		{
 			ClearComment();
-			FUHTException::Throwf(*this, TEXT("End of class header encountered inside comment") );
+			Throwf(TEXT("End of class header encountered inside comment") );
 		}
 		goto Loop;
 	}
@@ -401,7 +401,7 @@ bool FBaseParser::GetToken( FToken& Token, bool bNoConsts/*=false*/, ESymbolPars
 			Token.Identifier[Length++] = c;
 			if( Length >= NAME_SIZE )
 			{
-				FUHTException::Throwf(*this, TEXT("Identifer length exceeds maximum of %i"), (int32)NAME_SIZE);
+				Throwf(TEXT("Identifer length exceeds maximum of %i"), (int32)NAME_SIZE);
 				Length = ((int32)NAME_SIZE) - 1;
 				break;
 			}
@@ -453,7 +453,7 @@ bool FBaseParser::GetToken( FToken& Token, bool bNoConsts/*=false*/, ESymbolPars
 			Token.Identifier[Length++] = c;
 			if( Length >= NAME_SIZE )
 			{
-				FUHTException::Throwf(*this, TEXT("Number length exceeds maximum of %i "), (int32)NAME_SIZE );
+				Throwf(TEXT("Number length exceeds maximum of %i "), (int32)NAME_SIZE );
 				Length = ((int32)NAME_SIZE) - 1;
 				break;
 			}
@@ -505,7 +505,7 @@ bool FBaseParser::GetToken( FToken& Token, bool bNoConsts/*=false*/, ESymbolPars
 		c = GetChar(/*bLiteral=*/ true);
 		if (c != '\'')
 		{
-			FUHTException::Throwf(*this, TEXT("Unterminated character constant"));
+			Throwf(TEXT("Unterminated character constant"));
 			UngetChar();
 		}
 
@@ -536,7 +536,7 @@ bool FBaseParser::GetToken( FToken& Token, bool bNoConsts/*=false*/, ESymbolPars
 			Temp[Length++] = c;
 			if( Length >= MAX_STRING_CONST_SIZE )
 			{
-				FUHTException::Throwf(*this, TEXT("String constant exceeds maximum of %i characters"), (int32)MAX_STRING_CONST_SIZE );
+				Throwf(TEXT("String constant exceeds maximum of %i characters"), (int32)MAX_STRING_CONST_SIZE );
 				c = TEXT('\"');
 				Length = ((int32)MAX_STRING_CONST_SIZE) - 1;
 				break;
@@ -547,7 +547,7 @@ bool FBaseParser::GetToken( FToken& Token, bool bNoConsts/*=false*/, ESymbolPars
 
 		if( c != '"' )
 		{
-			FUHTException::Throwf(*this, TEXT("Unterminated string constant: %s"), Temp);
+			Throwf(TEXT("Unterminated string constant: %s"), Temp);
 			UngetChar();
 		}
 
@@ -634,7 +634,7 @@ bool FBaseParser::GetRawTokenRespectingQuotes( FToken& Token, TCHAR StopChar /*
 		Temp[Length++] = c;
 		if( Length >= MAX_STRING_CONST_SIZE )
 		{
-			FUHTException::Throwf(*this, TEXT("Identifier exceeds maximum of %i characters"), (int32)MAX_STRING_CONST_SIZE );
+			Throwf(TEXT("Identifier exceeds maximum of %i characters"), (int32)MAX_STRING_CONST_SIZE );
 			c = GetChar(true);
 			Length = ((int32)MAX_STRING_CONST_SIZE) - 1;
 			break;
@@ -645,7 +645,7 @@ bool FBaseParser::GetRawTokenRespectingQuotes( FToken& Token, TCHAR StopChar /*
 
 	if (bInQuote)
 	{
-		FUHTException::Throwf(*this, TEXT("Unterminated quoted string"));
+		Throwf(TEXT("Unterminated quoted string"));
 	}
 
 	// Get rid of trailing whitespace.
@@ -684,7 +684,7 @@ bool FBaseParser::GetRawToken( FToken& Token, TCHAR StopChar /* = TCHAR('\n') */
 		Temp[Length++] = c;
 		if( Length >= MAX_STRING_CONST_SIZE )
 		{
-			FUHTException::Throwf(*this, TEXT("Identifier exceeds maximum of %i characters"), (int32)MAX_STRING_CONST_SIZE );
+			Throwf(TEXT("Identifier exceeds maximum of %i characters"), (int32)MAX_STRING_CONST_SIZE );
 		}
 		c = GetChar(true);
 	}
@@ -769,7 +769,7 @@ bool FBaseParser::GetConstInt(int32& Result, const TCHAR* Tag)
 
 	if (Tag != NULL)
 	{
-		FUHTException::Throwf(*this, TEXT("%s: Missing constant integer"), Tag );
+		Throwf(TEXT("%s: Missing constant integer"), Tag );
 	}
 
 	return false;
@@ -792,7 +792,7 @@ bool FBaseParser::GetConstInt64(int64& Result, const TCHAR* Tag)
 
 	if (Tag != NULL)
 	{
-		FUHTException::Throwf(*this, TEXT("%s: Missing constant integer"), Tag );
+		Throwf(TEXT("%s: Missing constant integer"), Tag );
 	}
 
 	return false;
@@ -897,11 +897,11 @@ void FBaseParser::MatchSemi()
 		FToken Token;
 		if( GetToken(Token) )
 		{
-			FUHTException::Throwf(*this, TEXT("Missing ';' before '%s'"), Token.Identifier );
+			Throwf(TEXT("Missing ';' before '%s'"), Token.Identifier );
 		}
 		else
 		{
-			FUHTException::Throwf(*this, TEXT("Missing ';'") );
+			Throwf(TEXT("Missing ';'") );
 		}
 	}
 }
@@ -958,7 +958,7 @@ void FBaseParser::RequireSymbol( const TCHAR Match, const TCHAR* Tag, ESymbolPar
 {
 	if (!MatchSymbol(Match, bParseTemplateClosingBracket))
 	{
-		FUHTException::Throwf(*this, TEXT("Missing '%c' in %s"), Match, Tag );
+		Throwf(TEXT("Missing '%c' in %s"), Match, Tag );
 	}
 }
 
@@ -966,7 +966,7 @@ void FBaseParser::RequireSymbol(const TCHAR Match, TFunctionRef<FString()> Error
 {
 	if (!MatchSymbol(Match, bParseTemplateClosingBracket))
 	{
-		FUHTException::Throwf(*this, TEXT("Missing '%c' in %s"), Match, *ErrorGetter());
+		Throwf(TEXT("Missing '%c' in %s"), Match, *ErrorGetter());
 	}
 }
 
@@ -977,7 +977,7 @@ void FBaseParser::RequireConstInt( const TCHAR* Match, const TCHAR* Tag )
 {
 	if (!MatchConstInt(Match))
 	{
-		FUHTException::Throwf(*this, TEXT("Missing integer '%s' in %s"), Match, Tag );
+		Throwf(TEXT("Missing integer '%s' in %s"), Match, Tag );
 	}
 }
 
@@ -985,7 +985,7 @@ void FBaseParser::RequireAnyConstInt( const TCHAR* Tag )
 {
 	if (!MatchAnyConstInt())
 	{
-		FUHTException::Throwf(*this, TEXT("Missing integer in %s"), Tag );
+		Throwf(TEXT("Missing integer in %s"), Tag );
 	}
 }
 
@@ -996,7 +996,7 @@ void FBaseParser::RequireIdentifier( const TCHAR* Match, const ESearchCase::Type
 {
 	if (!MatchIdentifier(Match, SearchCase))
 	{
-		FUHTException::Throwf(*this, TEXT("Missing '%s' in %s"), Match, Tag );
+		Throwf(TEXT("Missing '%s' in %s"), Match, Tag );
 	}
 }
 
@@ -1013,7 +1013,7 @@ FString FBaseParser::ReadNewStyleValue(const TCHAR* TypeOfSpecifier)
 	FToken ValueToken;
 	if (!GetToken(ValueToken, false))
 	{
-		FUHTException::Throwf(*this, TEXT("Expected a value when handling a %s"), TypeOfSpecifier);
+		Throwf(TEXT("Expected a value when handling a %s"), TypeOfSpecifier);
 	}
 
 	FString Result;
@@ -1068,7 +1068,7 @@ void FBaseParser::ParseNameWithPotentialAPIMacroPrefix(FString& DeclaredName, FS
 	// Read an identifier
 	if (!GetIdentifier(NameToken))
 	{
-		FUHTException::Throwf(*this, TEXT("Missing %s name"), FailureMessage);
+		Throwf(TEXT("Missing %s name"), FailureMessage);
 	}
 
 	// Is the identifier the name or an DLL import/export API macro?
@@ -1080,7 +1080,7 @@ void FBaseParser::ParseNameWithPotentialAPIMacroPrefix(FString& DeclaredName, FS
 		// Read the real name
 		if (!GetIdentifier(NameToken))
 		{
-			FUHTException::Throwf(*this, TEXT("Missing %s name"), FailureMessage);
+			Throwf(TEXT("Missing %s name"), FailureMessage);
 		}
 		DeclaredName = NameToken.Identifier;
 	}
@@ -1112,7 +1112,7 @@ void FBaseParser::ReadSpecifierSetInsideMacro(TArray<FPropertySpecifier>& Specif
 		FToken Specifier;
 		if (!GetToken(Specifier))
 		{
-			FUHTException::Throwf(*this, TEXT("Expected %s"), *ErrorMessageGetter());
+			Throwf(TEXT("Expected %s"), *ErrorMessageGetter());
 		}
 
 		if (Specifier.Matches(TEXT("meta"), ESearchCase::IgnoreCase))
@@ -1127,7 +1127,7 @@ void FBaseParser::ReadSpecifierSetInsideMacro(TArray<FPropertySpecifier>& Specif
 				FToken MetaKeyToken;
 				if (!GetIdentifier(MetaKeyToken))
 				{
-					FUHTException::Throwf(*this, TEXT("Expected a metadata key"));
+					Throwf(TEXT("Expected a metadata key"));
 				}
 
 				FString Key = MetaKeyToken.Identifier;
@@ -1152,7 +1152,7 @@ void FBaseParser::ReadSpecifierSetInsideMacro(TArray<FPropertySpecifier>& Specif
 			{
 				if (MetadataKeyword->ValueArgument == EMetadataValueArgument::None)
 				{
-					FUHTException::Throwf(*this, TEXT("Incorrect = after metadata specifier '%s'"), Specifier.Identifier);
+					Throwf(TEXT("Incorrect = after metadata specifier '%s'"), Specifier.Identifier);
 				}
 
 				FString Value = ReadNewStyleValue(TypeOfSpecifier);
@@ -1162,7 +1162,7 @@ void FBaseParser::ReadSpecifierSetInsideMacro(TArray<FPropertySpecifier>& Specif
 			{
 				if (MetadataKeyword->ValueArgument == EMetadataValueArgument::Required)
 				{
-					FUHTException::Throwf(*this, TEXT("Missing = after metadata specifier '%s'"), Specifier.Identifier);
+					Throwf(TEXT("Missing = after metadata specifier '%s'"), Specifier.Identifier);
 				}
 
 				MetadataKeyword->ApplyToMetadata(*this, MetaData);
@@ -1191,7 +1191,7 @@ void FBaseParser::InsertMetaDataPair(TMap<FName, FString>& MetaData, FString Key
 	// make sure the key is valid
 	if (Key.Len() == 0)
 	{
-		FUHTException::Throwf(*this, TEXT("Invalid metadata"));
+		Throwf(TEXT("Invalid metadata"));
 	}
 
 	// trim extra white space and quotes
@@ -1208,7 +1208,7 @@ void FBaseParser::InsertMetaDataPair(TMap<FName, FString>& MetaData, FName KeyNa
 	FString* ExistingValue = MetaData.Find(KeyName);
 	if (ExistingValue && Value != *ExistingValue)
 	{
-		FUHTException::Throwf(*this, TEXT("Metadata key '%s' first seen with value '%s' then '%s'"), *KeyName.ToString(), **ExistingValue, *Value);
+		Throwf(TEXT("Metadata key '%s' first seen with value '%s' then '%s'"), *KeyName.ToString(), **ExistingValue, *Value);
 	}
 
 	// finally we have enough to put it into our metadata
diff --git a/Engine/Source/Programs/UnrealHeaderTool/Private/BaseParser.h b/Engine/Source/Programs/UnrealHeaderTool/Private/BaseParser.h
--- a/Engine/Source/Programs/UnrealHeaderTool/Private/BaseParser.h
+++ b/Engine/Source/Programs/UnrealHeaderTool/Private/BaseParser.h
@@ -84,7 +84,7 @@ public:
 //
 
 class FBaseParser 
-	: public FUHTExceptionContext
+	: public FUHTMessageProvider
 {
 protected:
 	FBaseParser(FUnrealSourceFile& InSourceFile);
@@ -276,7 +276,7 @@ void FBaseParser::ParseInheritance(const TCHAR* What, Lambda&& InLambda)
 		RequireIdentifier(TEXT("public"), ESearchCase::CaseSensitive, TEXT("inheritance"));
 		if (!GetIdentifier(Token))
 		{
-			FUHTException::Throwf(*this, TEXT("Missing %s name"), What);
+			Throwf(TEXT("Missing %s name"), What);
 		}
 		RedirectTypeIdentifier(Token);
 		InLambda(Token.Identifier, true);
@@ -294,7 +294,7 @@ void FBaseParser::ParseInheritance(const TCHAR* What, Lambda&& InLambda)
 			FToken Token;
 			if (!GetIdentifier(Token, true))
 			{
-				FUHTException::Throwf(*this, TEXT("Failed to get interface class identifier"));
+				Throwf(TEXT("Failed to get interface class identifier"));
 			}
 
 			InterfaceName += Token.Identifier;
@@ -309,7 +309,7 @@ void FBaseParser::ParseInheritance(const TCHAR* What, Lambda&& InLambda)
 				{
 					if (!GetToken(Token))
 					{
-						FUHTException::Throwf(*this, TEXT("Unexpected end of file"));
+						Throwf(TEXT("Unexpected end of file"));
 					}
 
 					if (Token.TokenType == TOKEN_Symbol)
diff --git a/Engine/Source/Programs/UnrealHeaderTool/Private/CodeGenerator.cpp b/Engine/Source/Programs/UnrealHeaderTool/Private/CodeGenerator.cpp
--- a/Engine/Source/Programs/UnrealHeaderTool/Private/CodeGenerator.cpp
+++ b/Engine/Source/Programs/UnrealHeaderTool/Private/CodeGenerator.cpp
@@ -279,7 +279,7 @@ namespace
 			}
 			else
 			{
-				FUHTException::Throwf(ClassDef, TEXT("Class %s has Net flagged properties and should declare member function: void GetLifetimeReplicatedProps(TArray<FLifetimeProperty>& OutLifetimeProps) const override"), ClassCPPName);
+				ClassDef.Throwf(TEXT("Class %s has Net flagged properties and should declare member function: void GetLifetimeReplicatedProps(TArray<FLifetimeProperty>& OutLifetimeProps) const override"), ClassCPPName);
 			}
 		}
 
@@ -1648,7 +1648,7 @@ FString GetEventStructParamsName(FUnrealObjectDefinitionInfo& OuterDef, const TC
 	}
 	else
 	{
-		FUHTException::Throwf(OuterDef, TEXT("Unrecognized outer type"));
+		OuterDef.Throwf(TEXT("Unrecognized outer type"));
 	}
 
 	FString Result = FString::Printf(TEXT("%s_event%s_Parms"), *OuterName, FunctionName);
@@ -1923,7 +1923,7 @@ void FNativeClassHeaderGenerator::ExportNativeGeneratedInitCode(FOutputDevice& O
 			// In a dynamic class the same function signature may be used for a Multi- and a Single-cast delegate.
 			if (!LocalFuncDef->IsDelegateFunction() || !bIsDynamic)
 			{
-				FUHTException::Throwf(*LocalFuncDef, TEXT("The same function linked twice. Function: %s Class: %s"), *LocalFuncDef->GetName(), *ClassDef.GetName());
+				LocalFuncDef->Throwf(TEXT("The same function linked twice. Function: %s Class: %s"), *LocalFuncDef->GetName(), *ClassDef.GetName());
 			}
 			continue;
 		}
@@ -4474,13 +4474,13 @@ void FNativeClassHeaderGenerator::CheckRPCFunctions(FReferenceGatherers& OutRefe
 	if (bNeedsImplementation && !bHasImplementation)
 	{
 		FString FunctionDecl = FString::Printf(TEXT("virtual %s %s::%s(%s) %s"), *FunctionReturnType, *ClassName, *FunctionData.CppImplName, *ParameterString, *ConstModifier);
-		FUHTException::Throwf(FunctionDef, TEXT("%s Declare function %s"), *AssertMessage, *FunctionDecl);
+		FunctionDef.Throwf(TEXT("%s Declare function %s"), *AssertMessage, *FunctionDecl);
 	}
 
 	if (bNeedsValidate && !bHasValidate)
 	{
 		FString FunctionDecl = FString::Printf(TEXT("virtual bool %s::%s(%s) %s"), *ClassName, *FunctionData.CppValidationImplName, *ParameterString, *ConstModifier);
-		FUHTException::Throwf(FunctionDef, TEXT("%s Declare function %s"), *AssertMessage, *FunctionDecl);
+		FunctionDef.Throwf(TEXT("%s Declare function %s"), *AssertMessage, *FunctionDecl);
 	}
 
 	//
@@ -4489,13 +4489,13 @@ void FNativeClassHeaderGenerator::CheckRPCFunctions(FReferenceGatherers& OutRefe
 	if (bNeedsImplementation && bHasImplementation && IsMissingVirtualSpecifier(FileContent, ImplementationPosition))
 	{
 		FString FunctionDecl = FString::Printf(TEXT("%s %s::%s(%s) %s"), *FunctionReturnType, *ClassName, *FunctionData.CppImplName, *ParameterString, *ConstModifier);
-		FUHTException::Throwf(FunctionDef, TEXT("Declared function %sis not marked as virtual."), *FunctionDecl);
+		FunctionDef.Throwf(TEXT("Declared function %sis not marked as virtual."), *FunctionDecl);
 	}
 
 	if (bNeedsValidate && bHasValidate && IsMissingVirtualSpecifier(FileContent, ValidatePosition))
 	{
 		FString FunctionDecl = FString::Printf(TEXT("bool %s::%s(%s) %s"), *ClassName, *FunctionData.CppValidationImplName, *ParameterString, *ConstModifier);
-		FUHTException::Throwf(FunctionDef, TEXT("Declared function %sis not marked as virtual."), *FunctionDecl);
+		FunctionDef.Throwf(TEXT("Declared function %sis not marked as virtual."), *FunctionDecl);
 	}
 }
 
@@ -5465,7 +5465,7 @@ void FNativeClassHeaderGenerator::GenerateSourceFiles(
 
 		auto GenerateSource = [&GeneratedCPP]()
 		{
-			FResults::Try(GeneratedCPP.SourceFile, [&GeneratedCPP]()
+			FResults::Try([&GeneratedCPP]()
 				{
 					FUnrealPackageDefinitionInfo& PackageDefLcl = GeneratedCPP.PackageDef;
 					const FManifestModule& Module = PackageDefLcl.GetModule();
@@ -5561,7 +5561,7 @@ void FNativeClassHeaderGenerator::GenerateSourceFiles(
 
 		auto WriteGenerated = [&GeneratedCPP]()
 		{
-			FResults::Try(GeneratedCPP.SourceFile, [&GeneratedCPP]()
+			FResults::Try([&GeneratedCPP]()
 				{
 					FUnrealPackageDefinitionInfo& PackageDefLcl = GeneratedCPP.PackageDef;
 					const FManifestModule& Module = PackageDefLcl.GetModule();
@@ -6004,7 +6004,7 @@ bool FNativeClassHeaderGenerator::SaveHeaderIfChanged(FGeneratedFileInfo& FileIn
 			FFileHelper::SaveStringToFile(InNewHeaderContents, *ConflictPath);
 
 			FResults::SetResult(ECompilationResult::FailedDueToHeaderChange);
-			FUHTException::Throwf(FString(FileInfo.GetFilename()), 1, TEXT("ERROR: '%s': Changes to generated code are not allowed - conflicts written to '%s'"), *HeaderPathStr, *ConflictPath);
+			FUHTMessage(FileInfo.GetFilename()).Throwf(TEXT("ERROR: '%s': Changes to generated code are not allowed - conflicts written to '%s'"), *HeaderPathStr, *ConflictPath);
 		}
 
 		// save the updated version to a tmp file so that the user can see what will be changing
@@ -6016,7 +6016,7 @@ bool FNativeClassHeaderGenerator::SaveHeaderIfChanged(FGeneratedFileInfo& FileIn
 			IFileManager::Get().Delete(*TmpHeaderFilename, false, true);
 			if (!FFileHelper::SaveStringToFile(InNewHeaderContents, *TmpHeaderFilename))
 			{
-				FResults::LogWarning(FString(FileInfo.GetFilename()), 1, FString::Printf(TEXT("Failed to save header export preview: '%s'"), *TmpHeaderFilename));
+				FUHTMessage(FileInfo.GetFilename()).LogWarning(TEXT("Failed to save header export preview: '%s'"), *TmpHeaderFilename);
 			}
 		};
 
@@ -6148,7 +6148,7 @@ void ResolveSuperClasses(const TCHAR* PackageName, FUnrealClassDefinitionInfo& C
 				if (FoundBaseClassDef == nullptr)
 				{
 					// Don't know its parent class. Raise error.
-					FUHTException::Throwf(ClassDef, TEXT("Couldn't find parent type for '%s' named '%s' in current module (Package: %s) or any other module parsed so far."), *ClassDef.GetName(), *BaseClassName, PackageName);
+					ClassDef.Throwf(TEXT("Couldn't find parent type for '%s' named '%s' in current module (Package: %s) or any other module parsed so far."), *ClassDef.GetName(), *BaseClassName, PackageName);
 				}
 
 				SuperClassInfo.Struct = FoundBaseClassDef;
@@ -6314,8 +6314,7 @@ void PrepareModules(TArray<FUnrealPackageDefinitionInfo*>& PackageDefs, const FS
 
 					if (NormalizedFullFilename != NormalizedExistingFilename)
 					{
-						FString AbsFilename = IFileManager::Get().ConvertToAbsolutePathForExternalAppForRead(*UnrealSourceFile->GetFilename());
-						FResults::LogError(MoveTemp(AbsFilename), 1, FString::Printf(TEXT("Duplicate leaf header name found: %s (original: %s)"), *NormalizedFullFilename, *NormalizedExistingFilename));
+						FUHTMessage(*UnrealSourceFile).LogError(TEXT("Duplicate leaf header name found: %s (original: %s)"), *NormalizedFullFilename, *NormalizedExistingFilename);
 					}
 				}
 
@@ -6382,7 +6381,7 @@ void PreparseSources(TArray<FUnrealPackageDefinitionInfo*>& PackageDefs, const F
 			// Phase #1: Load the file
 			auto LoadLambda = [&SourceFile = *SourceFile, &ModuleInfoPath]()
 			{
-				FResults::Try(SourceFile, [&SourceFile, &ModuleInfoPath]()
+				FResults::Try([&SourceFile, &ModuleInfoPath]()
 					{
 						FScopedDurationTimer SourceTimer(SourceFile.GetTime(ESourceFileTime::Load));
 
@@ -6391,7 +6390,7 @@ void PreparseSources(TArray<FUnrealPackageDefinitionInfo*>& PackageDefs, const F
 						FString Content;
 						if (!FFileHelper::LoadFileToString(Content, *FullFilename))
 						{
-							FUHTException::Throwf(SourceFile, 1, TEXT("UnrealHeaderTool was unable to load source file '%s'"), *FullFilename);
+							FUHTMessage(SourceFile).Throwf(TEXT("UnrealHeaderTool was unable to load source file '%s'"), *FullFilename);
 						}
 						SourceFile.SetContent(MoveTemp(Content));
 					}
@@ -6401,7 +6400,7 @@ void PreparseSources(TArray<FUnrealPackageDefinitionInfo*>& PackageDefs, const F
 			// Phase #2: Perform simplified class parse (can run concurrenrtly)
 			auto PreProcessLambda = [&SourceFile = *SourceFile]()
 			{
-				FResults::Try(SourceFile, [&SourceFile]()
+				FResults::Try([&SourceFile]()
 					{
 						FScopedDurationTimer SourceTimer(SourceFile.GetTime(ESourceFileTime::PreParse));
 
@@ -6432,7 +6431,7 @@ void DefineTypes(TArray<FUnrealPackageDefinitionInfo*>& PackageDefs)
 
 		for (TSharedRef<FUnrealSourceFile>& SourceFile : PackageDef->GetAllSourceFiles())
 		{
-			FResults::Try(*SourceFile, [PackageDef, &SourceFile = *SourceFile]()
+			FResults::Try([PackageDef, &SourceFile = *SourceFile]()
 			{
 				TArray<TSharedRef<FUnrealTypeDefinitionInfo>>& AllClasses = PackageDef->GetAllClasses();
 				UPackage* Package = PackageDef->GetPackage();
@@ -6643,7 +6642,7 @@ void ParseSourceFiles(TArray<FUnrealSourceFile*>& OrderedSourceFiles)
 		{
 			FResults::TryAlways([&ParsedCPP]()
 			{
-					FHeaderParser::Parse(ParsedCPP.PackageDef, ParsedCPP.SourceFile);
+				FHeaderParser::Parse(ParsedCPP.PackageDef, ParsedCPP.SourceFile);
 			});
 		};
 
@@ -6977,12 +6976,12 @@ void ProcessParsedClass(FUnrealClassDefinitionInfo& ClassDef)
 	// All classes must start with a valid unreal prefix
 	if (!FHeaderParser::ClassNameHasValidPrefix(ClassName, ClassNameStripped))
 	{
-		FUHTException::Throwf(ClassDef, TEXT("Invalid class name '%s'. The class name must have an appropriate prefix added (A for Actors, U for other classes)."), *ClassName);
+		ClassDef.Throwf(TEXT("Invalid class name '%s'. The class name must have an appropriate prefix added (A for Actors, U for other classes)."), *ClassName);
 	}
 
 	if(FHeaderParser::IsReservedTypeName(ClassNameStripped))
 	{
-		FUHTException::Throwf(ClassDef, TEXT("Invalid class name '%s'. Cannot use a reserved name ('%s')."), *ClassName, *ClassNameStripped);
+		ClassDef.Throwf(TEXT("Invalid class name '%s'. Cannot use a reserved name ('%s')."), *ClassName, *ClassNameStripped);
 	}
 
 	// Ensure the base class has any valid prefix and exists as a valid class. Checking for the 'correct' prefix will occur during compilation
@@ -6992,7 +6991,7 @@ void ProcessParsedClass(FUnrealClassDefinitionInfo& ClassDef)
 		BaseClassNameStripped = GetClassNameWithPrefixRemoved(BaseClassName);
 		if (!FHeaderParser::ClassNameHasValidPrefix(BaseClassName, BaseClassNameStripped))
 		{
-			FUHTException::Throwf(ClassDef, TEXT("No prefix or invalid identifier for base class %s.\nClass names must match Unreal prefix specifications (e.g., \"UObject\" or \"AActor\")"), *BaseClassName);
+			ClassDef.Throwf(TEXT("No prefix or invalid identifier for base class %s.\nClass names must match Unreal prefix specifications (e.g., \"UObject\" or \"AActor\")"), *BaseClassName);
 		}
 	}
 
@@ -7000,12 +6999,12 @@ void ProcessParsedClass(FUnrealClassDefinitionInfo& ClassDef)
 	// Handle failure and non-class headers.
 	if (BaseClassName.IsEmpty() && (ClassName != TEXT("UObject")))
 	{
-		FUHTException::Throwf(ClassDef, TEXT("Class '%s' must inherit UObject or a UObject-derived class"), *ClassName);
+		ClassDef.Throwf(TEXT("Class '%s' must inherit UObject or a UObject-derived class"), *ClassName);
 	}
 
 	if (ClassName == BaseClassName)
 	{
-		FUHTException::Throwf(ClassDef, TEXT("Class '%s' cannot inherit from itself"), *ClassName);
+		ClassDef.Throwf(TEXT("Class '%s' cannot inherit from itself"), *ClassName);
 	}
 
 	UClass* ResultClass = FEngineAPI::FindObject<UClass>(ANY_PACKAGE, *ClassNameStripped);
@@ -7018,7 +7017,7 @@ void ProcessParsedClass(FUnrealClassDefinitionInfo& ClassDef)
 	{
 		if (TSharedRef<FUnrealTypeDefinitionInfo>* Existing = GTypeDefinitionInfoMap.FindByName(*ClassNameStripped))
 		{
-			FUHTException::Throwf(ClassDef, TEXT("Duplicate class name: %s also exists in file %s"), *ClassName, *(*Existing)->GetFilename());
+			ClassDef.Throwf(TEXT("Duplicate class name: %s also exists in file %s"), *ClassName, *(*Existing)->GetFilename());
 		}
 
 		if (bVerboseOutput)
@@ -7040,13 +7039,13 @@ void ProcessParsedEnum(FUnrealEnumDefinitionInfo& EnumDef)
 
 	if (TSharedRef<FUnrealTypeDefinitionInfo>* Existing = GTypeDefinitionInfoMap.FindByName(*EnumName))
 	{
-		FUHTException::Throwf(EnumDef, TEXT("Duplicate enum name: %s also exists in file %s"), *EnumName, *(*Existing)->GetFilename());
+		EnumDef.Throwf(TEXT("Duplicate enum name: %s also exists in file %s"), *EnumName, *(*Existing)->GetFilename());
 	}
 
 	// Check if the enum name is using a reserved keyword
 	if (FHeaderParser::IsReservedTypeName(EnumName))
 	{
-		FUHTException::Throwf(EnumDef, TEXT("enum: '%s' uses a reserved type name."), *EnumName);
+		EnumDef.Throwf(TEXT("enum: '%s' uses a reserved type name."), *EnumName);
 	}
 }
 
@@ -7057,12 +7056,12 @@ void ProcessParsedStruct(FUnrealScriptStructDefinitionInfo& ScriptStructDef)
 
 	if (TSharedRef<FUnrealTypeDefinitionInfo>* Existing = GTypeDefinitionInfoMap.FindByName(*StructNameStripped))
 	{
-		FUHTException::Throwf(ScriptStructDef, TEXT("Duplicate struct name: %s also exists in file %s"), *StructNameStripped, *(*Existing)->GetFilename());
+		ScriptStructDef.Throwf(TEXT("Duplicate struct name: %s also exists in file %s"), *StructNameStripped, *(*Existing)->GetFilename());
 	}
 
 	// Check if the enum name is using a reserved keyword
 	if (FHeaderParser::IsReservedTypeName(StructNameStripped))
 	{
-		FUHTException::Throwf(ScriptStructDef, TEXT("struct: '%s' uses a reserved type name."), *StructNameStripped);
+		ScriptStructDef.Throwf(TEXT("struct: '%s' uses a reserved type name."), *StructNameStripped);
 	}
 }
diff --git a/Engine/Source/Programs/UnrealHeaderTool/Private/Exceptions.cpp b/Engine/Source/Programs/UnrealHeaderTool/Private/Exceptions.cpp
--- a/Engine/Source/Programs/UnrealHeaderTool/Private/Exceptions.cpp
+++ b/Engine/Source/Programs/UnrealHeaderTool/Private/Exceptions.cpp
@@ -17,23 +17,7 @@
 
 #include <atomic>
 
-FUHTException::FUHTException(ECompilationResult::Type InResult, FString&& InFilename, int32 InLine, FString&& InMessage)
-	: Result(InResult)
-	, Message(MoveTemp(InMessage))
-	, Filename(MoveTemp(InFilename))
-	, Line(InLine)
-{
-}
-
-FUHTException::FUHTException(ECompilationResult::Type InResult, const FUnrealSourceFile& SourceFile, int32 InLine, FString&& InMessage)
-	: Result(InResult)
-	, Message(MoveTemp(InMessage))
-	, Filename(SourceFile.GetFilename())
-	, Line(InLine)
-{
-}
-
-FUHTException::FUHTException(ECompilationResult::Type InResult, const FUHTExceptionContext& InContext, FString&& InMessage)
+FUHTException::FUHTException(ECompilationResult::Type InResult, const FUHTMessageProvider& InContext, FString&& InMessage)
 	: Result(InResult)
 	, Message(MoveTemp(InMessage))
 	, Filename(InContext.GetFilename())
@@ -70,6 +54,30 @@ namespace UE::UnrealHeaderTool::Exceptions::Private
 		FResults::SetResult(InResult);
 	}
 
+	void LogErrorOuter(ECompilationResult::Type InResult, const FString& InFilename, int32 Line, FString&& Message)
+	{
+		using namespace UE::UnrealHeaderTool::Exceptions::Private;
+
+		FString Filename = IFileManager::Get().ConvertToAbsolutePathForExternalAppForRead(*InFilename);
+
+		if (IsInGameThread())
+		{
+			LogErrorInternal(InResult, Filename, Line, Message);
+		}
+		else
+		{
+			auto LogWarningTask = [Filename = MoveTemp(Filename), Line, Message = MoveTemp(Message), InResult]()
+			{
+				LogErrorInternal(InResult, Filename, Line, Message);
+			};
+
+			FGraphEventRef EventRef = FFunctionGraphTask::CreateAndDispatchWhenReady(MoveTemp(LogWarningTask), TStatId(), nullptr, ENamedThreads::GameThread);
+
+			FScopeLock Lock(&ErrorTasksCS);
+			ErrorTasks.Add(EventRef);
+		}
+	}
+
 	void LogWarningInternal(const FString& Filename, int32 Line, const FString& Message)
 	{
 		TGuardValue<ELogTimes::Type> DisableLogTimes(GPrintLogTimes, ELogTimes::None);
@@ -91,81 +99,38 @@ namespace UE::UnrealHeaderTool::Exceptions::Private
 	}
 }
 
-void FResults::LogError(FString&& Filename, int32 Line, const FString& Message, ECompilationResult::Type InResult/* = ECompilationResult::OtherCompilationError*/)
+void FResults::LogError(const FUHTMessageProvider& Context, const TCHAR* ErrorMsg, ECompilationResult::Type InResult)
 {
 	using namespace UE::UnrealHeaderTool::Exceptions::Private;
-
-	if (IsInGameThread())
-	{
-		LogErrorInternal(InResult, Filename, Line, Message);
-	}
-	else
-	{
-		auto LogExceptionTask = [Filename = MoveTemp(Filename), Line, Message, InResult]()
-		{
-			LogErrorInternal(InResult, Filename, Line, Message);
-		};
-
-		FGraphEventRef EventRef = FFunctionGraphTask::CreateAndDispatchWhenReady(MoveTemp(LogExceptionTask), TStatId(), nullptr, ENamedThreads::GameThread);
-
-		FScopeLock Lock(&ErrorTasksCS);
-		ErrorTasks.Add(EventRef);
-	}
+	LogErrorOuter(InResult, Context.GetFilename(), Context.GetLineNumber(), FString(ErrorMsg));
 }
 
 void FResults::LogError(const FUHTException& Ex)
 {
-	FString AbsFilename;
-	const TCHAR* Filename = Ex.GetFilename();
-	if (Filename != nullptr)
-	{
-		AbsFilename = IFileManager::Get().ConvertToAbsolutePathForExternalAppForRead(Ex.GetFilename());
-	}
-	LogError(MoveTemp(AbsFilename), Ex.GetLine(), Ex.GetMessage());
-}
-
-void FResults::LogError(const FUnrealSourceFile& SourceFile, const FUHTException& Ex)
-{
-	FString AbsFilename;
-	const TCHAR* Filename = Ex.GetFilename();
-	if (Filename == nullptr)
-	{
-		Filename = *SourceFile.GetFilename();
-	}
-	if (Filename != nullptr)
-	{
-		AbsFilename = IFileManager::Get().ConvertToAbsolutePathForExternalAppForRead(Ex.GetFilename());
-	}
-	LogError(MoveTemp(AbsFilename), Ex.GetLine(), Ex.GetMessage());
-}
-
-void FResults::LogError(const FUnrealSourceFile& SourceFile, int32 Line, const TCHAR* ErrorMsg, ECompilationResult::Type InResult/* = ECompilationResult::OtherCompilationError*/)
-{
-	FString AbsFilename = IFileManager::Get().ConvertToAbsolutePathForExternalAppForRead(*SourceFile.GetFilename());
-	LogError(MoveTemp(AbsFilename), Line, ErrorMsg, InResult);
-}
-
-void FResults::LogError(const FUHTExceptionContext& Context, const TCHAR* ErrorMsg, ECompilationResult::Type InResult/* = ECompilationResult::OtherCompilationError*/)
-{
-	LogError(Context.GetFilename(), Context.GetLineNumber(), ErrorMsg, InResult);
+	using namespace UE::UnrealHeaderTool::Exceptions::Private;
+	LogErrorOuter(Ex.GetResult(), Ex.GetFilename(), Ex.GetLine(), FString(Ex.GetMessage()));
 }
 
-void FResults::LogError(const TCHAR* ErrorMsg, ECompilationResult::Type InResult/* = ECompilationResult::OtherCompilationError*/)
+void FResults::LogError(const TCHAR* Message)
 {
-	LogError(TEXT(""), 1, ErrorMsg, InResult);
+	using namespace UE::UnrealHeaderTool::Exceptions::Private;
+	LogErrorOuter(ECompilationResult::Type::OtherCompilationError, TEXT("UnknownFile"), 1, FString(Message));
 }
 
-void FResults::LogWarning(FString&& Filename, int32 Line, const FString& Message)
+void FResults::LogWarning(const FUHTMessageProvider& Context, const TCHAR* ErrorMsg)
 {
 	using namespace UE::UnrealHeaderTool::Exceptions::Private;
 
+	FString Filename = IFileManager::Get().ConvertToAbsolutePathForExternalAppForRead(*Context.GetFilename());
+	int32 Line = Context.GetLineNumber();
+
 	if (IsInGameThread())
 	{
-		LogWarningInternal(Filename, Line, Message);
+		LogWarningInternal(Filename, Line, ErrorMsg);
 	}
 	else
 	{
-		auto LogWarningTask = [Filename = MoveTemp(Filename), Line, Message]()
+		auto LogWarningTask = [Filename = MoveTemp(Filename), Line, Message = FString(ErrorMsg)]()
 		{
 			LogWarningInternal(Filename, Line, Message);
 		};
@@ -177,11 +142,6 @@ void FResults::LogWarning(FString&& Filename, int32 Line, const FString& Message
 	}
 }
 
-void FResults::LogWarning(const FUHTExceptionContext& Context, const TCHAR* ErrorMsg)
-{
-	LogWarning(Context.GetFilename(), Context.GetLineNumber(), ErrorMsg);
-}
-
 void FResults::WaitForErrorTasks()
 {
 	using namespace UE::UnrealHeaderTool::Exceptions::Private;
diff --git a/Engine/Source/Programs/UnrealHeaderTool/Private/Exceptions.h b/Engine/Source/Programs/UnrealHeaderTool/Private/Exceptions.h
--- a/Engine/Source/Programs/UnrealHeaderTool/Private/Exceptions.h
+++ b/Engine/Source/Programs/UnrealHeaderTool/Private/Exceptions.h
@@ -11,111 +11,19 @@ class FBaseParser;
 class FUnrealSourceFile;
 class FUnrealTypeDefinitionInfo;
 
-class FUHTExceptionContext
-{
-public:
-	virtual FString GetFilename() const = 0;
-	virtual int32 GetLineNumber() const = 0;
-};
+class FUHTMessageProvider;
 
 class FUHTException
 {
 public:
 
-	/**
-	 * Generate an exception
-	 * @param InFilename The file generating the exception
-	 * @param InLine The line number in the file
-	 * @param InFmt The format string
-	 * @param InArgs Arguments supplied to the format string
-	 */
-	template <typename FmtType, typename... Types>
-	UE_NORETURN static void VARARGS Throwf(FString&& InFilename, int32 InLine, const FmtType& InFmt, Types... InArgs)
-	{
-		FString ResultString = FString::Printf(InFmt, InArgs ...);
-		throw FUHTException(ECompilationResult::OtherCompilationError, MoveTemp(InFilename), InLine, MoveTemp(ResultString));
-	}
-
-	/**
-	 * Generate an exception
-	 * @param InResult Result code to be returned as the overall result of the compilation process
-	 * @param InFilename The file generating the exception
-	 * @param InLine The line number in the file
-	 * @param InFmt The format string
-	 * @param InArgs Arguments supplied to the format string
-	 */
-	template <typename FmtType, typename... Types>
-	UE_NORETURN static void VARARGS Throwf(ECompilationResult::Type InResult, FString&& InFilename, int32 InLine, const FmtType& InFmt, Types... InArgs)
-	{
-		FString ResultString = FString::Printf(InFmt, InArgs ...);
-		throw FUHTException(InResult, MoveTemp(InFilename), InLine, MoveTemp(ResultString));
-	}
-
-	/**
-	 * Generate an exception
-	 * @param InSourceFile The source file generating the exception
-	 * @param InLine The line number in the file
-	 * @param InFmt The format string
-	 * @param InArgs Arguments supplied to the format string
-	 */
-	template <typename FmtType, typename... Types>
-	UE_NORETURN static void VARARGS Throwf(const FUnrealSourceFile& InSourceFile, int32 InLine, const FmtType& InFmt, Types... InArgs)
-	{
-		FString ResultString = FString::Printf(InFmt, InArgs ...);
-		throw FUHTException(ECompilationResult::OtherCompilationError, InSourceFile, InLine, MoveTemp(ResultString));
-	}
-
-	/**
-	 * Generate an exception
-	 * @param InResult Result code to be returned as the overall result of the compilation process
-	 * @param InSourceFile The source file generating the exception
-	 * @param InLine The line number in the file
-	 * @param InFmt The format string
-	 * @param InArgs Arguments supplied to the format string
-	 */
-	template <typename FmtType, typename... Types>
-	UE_NORETURN static void VARARGS Throwf(ECompilationResult::Type InResult, const FUnrealSourceFile& InSourceFile, int32 InLine, const FmtType& InFmt, Types... InArgs)
-	{
-		FString ResultString = FString::Printf(InFmt, InArgs ...);
-		throw FUHTException(InResult, InSourceFile, InLine, MoveTemp(ResultString));
-	}
-
-	/**
-	 * Generate an exception
-	 * @param InContext The context of the exception
-	 * @param InText The text of the error
-	 */
-	UE_NORETURN static void Throwf(const FUHTExceptionContext& InContext, FString&& InText)
-	{
-		throw FUHTException(ECompilationResult::OtherCompilationError, InContext, MoveTemp(InText));
-	}
-
 	/**
 	 * Generate an exception
 	 * @param InContext The context of the exception
 	 * @param InFmt The format string
-	 * @param InArgs Arguments supplied to the format string
+	 * @param InArgs InMessage The text of the message
 	 */
-	template <typename FmtType, typename... Types>
-	UE_NORETURN static void VARARGS Throwf(const FUHTExceptionContext& InContext, const FmtType& InFmt, Types... InArgs)
-	{
-		FString ResultString = FString::Printf(InFmt, InArgs ...);
-		throw FUHTException(ECompilationResult::OtherCompilationError, InContext, MoveTemp(ResultString));
-	}
-
-	/**
-	 * Generate an exception
-	 * @param InResult Result code to be returned as the overall result of the compilation process
-	 * @param InContext The context of the exception
-	 * @param InFmt The format string
-	 * @param InArgs Arguments supplied to the format string
-	 */
-	template <typename FmtType, typename... Types>
-	UE_NORETURN static void VARARGS Throwf(ECompilationResult::Type InResult, const FUHTExceptionContext& InContext, const FmtType& InFmt, Types... InArgs)
-	{
-		FString ResultString = FString::Printf(InFmt, InArgs ...);
-		throw FUHTException(InResult, InContext, MoveTemp(ResultString));
-	}
+	FUHTException(ECompilationResult::Type InResult, const FUHTMessageProvider& InContext, FString&& InMessage);
 
 	/**
 	 * Return the result code of the exception
@@ -154,15 +62,12 @@ private:
 	FString Message;
 	FString Filename;
 	int32 Line;
-
-	FUHTException(ECompilationResult::Type InResult, FString&& InFilename, int32 InLine, FString&& InMessage);
-	FUHTException(ECompilationResult::Type InResult, const FUnrealSourceFile& SourceFile, int32 InLine, FString&& InMessage);
-	FUHTException(ECompilationResult::Type InResult, const FUHTExceptionContext& InContext, FString&& InMessage);
 };
 
 /** Helper methods for working with exceptions and compilation results */
 struct FResults
 {
+	friend class FUHTMessageProvider;
 
 	/**
 	 * Wait for any pending error tasks to complete.
@@ -199,35 +104,16 @@ struct FResults
 	static ECompilationResult::Type GetOverallResults();
 
 	/**
-	 * Log an error.
-	 * @param Filename The filename generating the error.  If empty, then no file and line number are included in the error.
-	 * @param Line Line number of the error
-	 * @param Message Message body of the error
-	 * @param Result Compilation result of the error
-	 */
-	static void LogError(FString&& Filename, int32 Line, const FString& Message, ECompilationResult::Type Result = ECompilationResult::OtherCompilationError);
-
-	/**
-	 * Log an error from an exception with possible override of the source file.
-	 * @param SourceFile The source file currently being processed.  If the exception does not include filename information, the source file will be used.
-	 * @param Ex The exception generating the error.
-	 */
-	static void LogError(const FUnrealSourceFile& SourceFile, const FUHTException& Ex);
-
-	/**
-	 * Log an error from an exception.
-	 * @param Ex The exception generating the error.
+	 * Log an error
+	 * @param Ex The exception generating the error
 	 */
 	static void LogError(const FUHTException& Ex);
 
 	/**
-	 * Log an error for the given source file
-	 * @param SourceFile The source file generating the error
-	 * @param Line The line number generating the error
-	 * @param ErrorMsg The text of the error
-	 * @param Result Compilation result of the error
+	 * Log an error
+	 * @param Text
 	 */
-	static void LogError(const FUnrealSourceFile& SourceFile, int32 Line, const TCHAR* ErrorMsg, ECompilationResult::Type Result = ECompilationResult::OtherCompilationError);
+	static void LogError(const TCHAR* Text);
 
 	/**
 	 * Log an error for the given source file where the type is defined
@@ -235,23 +121,7 @@ struct FResults
 	 * @param ErrorMsg The text of the error
 	 * @param Result Compilation result of the error
 	 */
-	static void LogError(const FUHTExceptionContext& Context, const TCHAR* ErrorMsg, ECompilationResult::Type Result = ECompilationResult::OtherCompilationError);
-
-	/**
-	 * Log an error without any source file information
-	 * @param ErrorMsg The text of the error
-	 * @param Result Compilation result of the error
-	 */
-	static void LogError(const TCHAR* ErrorMsg, ECompilationResult::Type Result = ECompilationResult::OtherCompilationError);
-
-	/**
-	 * Log a warning.
-	 * @param Filename The filename generating the warning.  If empty, then no file and line number are included in the warning.
-	 * @param Line Line number of the warning
-	 * @param Message Message body of the warning
-	 * @param Result Compilation result of the warning
-	 */
-	static void LogWarning(FString&& Filename, int32 Line, const FString& Message);
+	static void LogError(const FUHTMessageProvider& Context, const TCHAR* ErrorMsg, ECompilationResult::Type Result);
 
 	/**
 	 * Log a warning for the given source file where the type is defined
@@ -259,36 +129,7 @@ struct FResults
 	 * @param ErrorMsg The text of the warning
 	 * @param Result Compilation result of the warning
 	 */
-	static void LogWarning(const FUHTExceptionContext& Context, const TCHAR* ErrorMsg);
-
-	/**
-	 * Invoke the given lambda in a try block catching all supported exception types.
-	 * @param SourceFile The source file being processed
-	 * @param InLambda The code to be executed in the try block
-	 */
-	template<typename Lambda>
-	static void Try(FUnrealSourceFile& SourceFile, Lambda&& InLambda)
-	{
-		if (IsSucceeding())
-		{
-#if !PLATFORM_EXCEPTIONS_DISABLED
-			try
-#endif
-			{
-				InLambda();
-			}
-#if !PLATFORM_EXCEPTIONS_DISABLED
-			catch (const FUHTException& Ex)
-			{
-				LogError(SourceFile, Ex);
-			}
-			catch (const TCHAR* ErrorMsg)
-			{
-				LogError(SourceFile, 1, ErrorMsg);
-			}
-#endif
-		}
-	}
+	static void LogWarning(const FUHTMessageProvider& Context, const TCHAR* ErrorMsg);
 
 	/**
 	 * Invoke the given lambda in a try block catching all supported exception types.
@@ -360,5 +201,155 @@ struct FResults
 	}
 };
 
-#define UE_LOG_WARNING_UHT(Context, Format, ...) { FResults::LogWarning(Context, *FString::Printf(Format, ##__VA_ARGS__)); }
-#define UE_LOG_ERROR_UHT(Context, Format, ...) { FResults::LogError(Context, *FString::Printf(Format, ##__VA_ARGS__)); }
+class FUHTMessageProvider
+{
+public:
+	virtual ~FUHTMessageProvider() = default;
+	virtual FString GetFilename() const = 0;
+	virtual int32 GetLineNumber() const = 0;
+
+	/**
+	 * Generate an exception
+	 * @param InText The text of the error
+	 */
+	UE_NORETURN void Throwf(FString&& InText) const
+	{
+		throw FUHTException(ECompilationResult::OtherCompilationError, *this, MoveTemp(InText));
+	}
+
+	/**
+	 * Generate an exception
+	 * @param InResult Result code to be returned as the overall result of the compilation process
+	 * @param InText The text of the error
+	 */
+	UE_NORETURN void Throwf(ECompilationResult::Type InResult, FString&& InText) const
+	{
+		throw FUHTException(InResult, *this, MoveTemp(InText));
+	}
+
+	/**
+	 * Generate an exception
+	 * @param InFmt The format string
+	 * @param InArgs Arguments supplied to the format string
+	 */
+	template <typename FmtType, typename... Types>
+	UE_NORETURN void VARARGS Throwf(const FmtType& InFmt, Types... InArgs) const
+	{
+		FString ResultString = FString::Printf(InFmt, InArgs ...);
+		throw FUHTException(ECompilationResult::OtherCompilationError, *this, MoveTemp(ResultString));
+	}
+
+	/**
+	 * Generate an exception
+	 * @param InResult Result code to be returned as the overall result of the compilation process
+	 * @param InFmt The format string
+	 * @param InArgs Arguments supplied to the format string
+	 */
+	template <typename FmtType, typename... Types>
+	UE_NORETURN void VARARGS Throwf(ECompilationResult::Type InResult, const FmtType& InFmt, Types... InArgs) const
+	{
+		FString ResultString = FString::Printf(InFmt, InArgs ...);
+		throw FUHTException(InResult, *this, MoveTemp(ResultString));
+	}
+
+	/**
+	 * Log an error
+	 * @param InText The text of the error
+	 */
+	void LogError(const FString& InText) const
+	{
+		FResults::LogError(*this, *InText, ECompilationResult::OtherCompilationError);
+	}
+
+	/**
+	 * Log an error
+	 * @param InResult Result code to be returned as the overall result of the compilation process
+	 * @param InText The text of the error
+	 */
+	void LogError(ECompilationResult::Type InResult, const FString& InText) const
+	{
+		FResults::LogError(*this, *InText, InResult);
+	}
+
+	/**
+	 * Log an error
+	 * @param InFmt The format string
+	 * @param InArgs Arguments supplied to the format string
+	 */
+	template <typename FmtType, typename... Types>
+	void VARARGS LogError(const FmtType& InFmt, Types... InArgs) const
+	{
+		FString ResultString = FString::Printf(InFmt, InArgs ...);
+		FResults::LogError(*this, *ResultString, ECompilationResult::OtherCompilationError);
+	}
+
+	/**
+	 * Log an error
+	 * @param InResult Result code to be returned as the overall result of the compilation process
+	 * @param InFmt The format string
+	 * @param InArgs Arguments supplied to the format string
+	 */
+	template <typename FmtType, typename... Types>
+	void VARARGS LogError(ECompilationResult::Type InResult, const FmtType& InFmt, Types... InArgs) const
+	{
+		FString ResultString = FString::Printf(InFmt, InArgs ...);
+		FResults::LogError(*this, *ResultString, InResult);
+	}
+
+	/**
+	 * Log an warning
+	 * @param InText The text of the warning
+	 */
+	void LogWarning(const FString& InText) const
+	{
+		FResults::LogWarning(*this, *InText);
+	}
+
+	/**
+	 * Log an warning
+	 * @param InFmt The format string
+	 * @param InArgs Arguments supplied to the format string
+	 */
+	template <typename FmtType, typename... Types>
+	void VARARGS LogWarning(const FmtType& InFmt, Types... InArgs) const
+	{
+		FString ResultString = FString::Printf(InFmt, InArgs ...);
+		FResults::LogWarning(*this, *ResultString);
+	}
+};
+
+template <typename T>
+FString GetMessageFilename(const T& Source);
+
+template <>
+inline FString GetMessageFilename<FString>(const FString& Source)
+{
+	return Source;
+}
+
+class FUHTMessage
+	: public FUHTMessageProvider
+{
+public:
+
+	template <typename T>
+	explicit FUHTMessage(const T& Source, int32 InLineNumber = 1)
+		: Filename(GetMessageFilename(Source))
+		, LineNumber(InLineNumber)
+	{
+	}
+
+	virtual FString GetFilename() const override
+	{
+		return Filename;
+	}
+
+	virtual int32 GetLineNumber() const override
+	{
+		return LineNumber;
+	}
+
+private:
+	FString Filename;
+	int32 LineNumber;
+};
diff --git a/Engine/Source/Programs/UnrealHeaderTool/Private/HeaderParser.cpp b/Engine/Source/Programs/UnrealHeaderTool/Private/HeaderParser.cpp
--- a/Engine/Source/Programs/UnrealHeaderTool/Private/HeaderParser.cpp
+++ b/Engine/Source/Programs/UnrealHeaderTool/Private/HeaderParser.cpp
@@ -229,7 +229,7 @@ namespace
 					int32 TempInt = FCString::Atoi(Equals + 1);
 					if (TempInt <= 0 || TempInt > MAX_uint16)
 					{
-						FUHTException::Throwf(HeaderParser, TEXT("Invalid network identifier %s for function"), IdentifierPtr);
+						HeaderParser.Throwf(TEXT("Invalid network identifier %s for function"), IdentifierPtr);
 					}
 					FuncInfo.RPCId = (uint16)TempInt;
 				}
@@ -239,7 +239,7 @@ namespace
 					int32 TempInt = FCString::Atoi(Equals + 1);
 					if (TempInt <= 0 || TempInt > MAX_uint16)
 					{
-						FUHTException::Throwf(HeaderParser, TEXT("Invalid network identifier %s for function"), IdentifierPtr);
+						HeaderParser.Throwf(TEXT("Invalid network identifier %s for function"), IdentifierPtr);
 					}
 					FuncInfo.RPCResponseId = (uint16)TempInt;
 				}
@@ -250,7 +250,7 @@ namespace
 
 				if (FuncInfo.EndpointName.Len())
 				{
-					FUHTException::Throwf(HeaderParser, TEXT("Function should not specify multiple endpoints - '%s' found but already using '%s'"), *Identifier);
+					HeaderParser.Throwf(TEXT("Function should not specify multiple endpoints - '%s' found but already using '%s'"), *Identifier);
 				}
 
 				FuncInfo.EndpointName = Identifier;
@@ -275,7 +275,7 @@ namespace
 			{
 				default:
 				{
-					FUHTException::Throwf(HeaderParser, TEXT("Unknown function specifier '%s'"), *Specifier.Key);
+					HeaderParser.Throwf(TEXT("Unknown function specifier '%s'"), *Specifier.Key);
 				}
 				break;
 
@@ -283,20 +283,20 @@ namespace
 				{
 					if (FuncInfo.FunctionFlags & FUNC_Net)
 					{
-						UE_LOG_ERROR_UHT(HeaderParser, TEXT("BlueprintNativeEvent functions cannot be replicated!") );
+						HeaderParser.LogError(TEXT("BlueprintNativeEvent functions cannot be replicated!") );
 					}
 					else if ( (FuncInfo.FunctionFlags & FUNC_BlueprintEvent) && !(FuncInfo.FunctionFlags & FUNC_Native) )
 					{
 						// already a BlueprintImplementableEvent
-						UE_LOG_ERROR_UHT(HeaderParser, TEXT("A function cannot be both BlueprintNativeEvent and BlueprintImplementableEvent!") );
+						HeaderParser.LogError(TEXT("A function cannot be both BlueprintNativeEvent and BlueprintImplementableEvent!") );
 					}
 					else if (bSawPropertyAccessor)
 					{
-						UE_LOG_ERROR_UHT(HeaderParser, TEXT("A function cannot be both BlueprintNativeEvent and a Blueprint Property accessor!"));
+						HeaderParser.LogError(TEXT("A function cannot be both BlueprintNativeEvent and a Blueprint Property accessor!"));
 					}
 					else if ( (FuncInfo.FunctionFlags & FUNC_Private) )
 					{
-						UE_LOG_ERROR_UHT(HeaderParser, TEXT("A Private function cannot be a BlueprintNativeEvent!") );
+						HeaderParser.LogError(TEXT("A Private function cannot be a BlueprintNativeEvent!") );
 					}
 
 					FuncInfo.FunctionFlags |= FUNC_Event;
@@ -308,20 +308,20 @@ namespace
 				{
 					if (FuncInfo.FunctionFlags & FUNC_Net)
 					{
-						UE_LOG_ERROR_UHT(HeaderParser, TEXT("BlueprintImplementableEvent functions cannot be replicated!") );
+						HeaderParser.LogError(TEXT("BlueprintImplementableEvent functions cannot be replicated!") );
 					}
 					else if ( (FuncInfo.FunctionFlags & FUNC_BlueprintEvent) && (FuncInfo.FunctionFlags & FUNC_Native) )
 					{
 						// already a BlueprintNativeEvent
-						UE_LOG_ERROR_UHT(HeaderParser, TEXT("A function cannot be both BlueprintNativeEvent and BlueprintImplementableEvent!") );
+						HeaderParser.LogError(TEXT("A function cannot be both BlueprintNativeEvent and BlueprintImplementableEvent!") );
 					}
 					else if (bSawPropertyAccessor)
 					{
-						UE_LOG_ERROR_UHT(HeaderParser, TEXT("A function cannot be both BlueprintImplementableEvent and a Blueprint Property accessor!"));
+						HeaderParser.LogError(TEXT("A function cannot be both BlueprintImplementableEvent and a Blueprint Property accessor!"));
 					}
 					else if ( (FuncInfo.FunctionFlags & FUNC_Private) )
 					{
-						UE_LOG_ERROR_UHT(HeaderParser, TEXT("A Private function cannot be a BlueprintImplementableEvent!") );
+						HeaderParser.LogError(TEXT("A Private function cannot be a BlueprintImplementableEvent!") );
 					}
 
 					FuncInfo.FunctionFlags |= FUNC_Event;
@@ -335,7 +335,7 @@ namespace
 					FuncInfo.FunctionFlags |= FUNC_Exec;
 					if( FuncInfo.FunctionFlags & FUNC_Net )
 					{
-						UE_LOG_ERROR_UHT(HeaderParser, TEXT("Exec functions cannot be replicated!") );
+						HeaderParser.LogError(TEXT("Exec functions cannot be replicated!") );
 					}
 				}
 				break;
@@ -350,7 +350,7 @@ namespace
 				{
 					if ((FuncInfo.FunctionFlags & FUNC_BlueprintEvent) != 0)
 					{
-						FUHTException::Throwf(HeaderParser, TEXT("BlueprintImplementableEvent or BlueprintNativeEvent functions cannot be declared as Client or Server"));
+						HeaderParser.Throwf(TEXT("BlueprintImplementableEvent or BlueprintNativeEvent functions cannot be declared as Client or Server"));
 					}
 
 					FuncInfo.FunctionFlags |= FUNC_Net;
@@ -363,7 +363,7 @@ namespace
 
 					if( FuncInfo.FunctionFlags & FUNC_Exec )
 					{
-						UE_LOG_ERROR_UHT(HeaderParser, TEXT("Exec functions cannot be replicated!") );
+						HeaderParser.LogError(TEXT("Exec functions cannot be replicated!") );
 					}
 				}
 				break;
@@ -372,7 +372,7 @@ namespace
 				{
 					if ((FuncInfo.FunctionFlags & FUNC_BlueprintEvent) != 0)
 					{
-						FUHTException::Throwf(HeaderParser, TEXT("BlueprintImplementableEvent or BlueprintNativeEvent functions cannot be declared as Client or Server"));
+						HeaderParser.Throwf(TEXT("BlueprintImplementableEvent or BlueprintNativeEvent functions cannot be declared as Client or Server"));
 					}
 
 					FuncInfo.FunctionFlags |= FUNC_Net;
@@ -389,7 +389,7 @@ namespace
 				{
 					if ((FuncInfo.FunctionFlags & FUNC_BlueprintEvent) != 0)
 					{
-						FUHTException::Throwf(HeaderParser, TEXT("BlueprintImplementableEvent or BlueprintNativeEvent functions cannot be declared as Multicast"));
+						HeaderParser.Throwf(TEXT("BlueprintImplementableEvent or BlueprintNativeEvent functions cannot be declared as Multicast"));
 					}
 
 					FuncInfo.FunctionFlags |= FUNC_Net;
@@ -401,7 +401,7 @@ namespace
 				{
 					if ((FuncInfo.FunctionFlags & FUNC_BlueprintEvent) != 0)
 					{
-						FUHTException::Throwf(HeaderParser, TEXT("BlueprintImplementableEvent or BlueprintNativeEvent functions cannot be declared as a ServiceRequest"));
+						HeaderParser.Throwf(TEXT("BlueprintImplementableEvent or BlueprintNativeEvent functions cannot be declared as a ServiceRequest"));
 					}
 
 					FuncInfo.FunctionFlags |= FUNC_Net;
@@ -413,7 +413,7 @@ namespace
 
 					if (FuncInfo.EndpointName.Len() == 0)
 					{
-						FUHTException::Throwf(HeaderParser, TEXT("ServiceRequest needs to specify an endpoint name"));
+						HeaderParser.Throwf(TEXT("ServiceRequest needs to specify an endpoint name"));
 					}
 				}
 				break;
@@ -422,7 +422,7 @@ namespace
 				{
 					if ((FuncInfo.FunctionFlags & FUNC_BlueprintEvent) != 0)
 					{
-						FUHTException::Throwf(HeaderParser, TEXT("BlueprintImplementableEvent or BlueprintNativeEvent functions cannot be declared as a ServiceResponse"));
+						HeaderParser.Throwf(TEXT("BlueprintImplementableEvent or BlueprintNativeEvent functions cannot be declared as a ServiceResponse"));
 					}
 
 					FuncInfo.FunctionFlags |= FUNC_Net;
@@ -433,7 +433,7 @@ namespace
 
 					if (FuncInfo.EndpointName.Len() == 0)
 					{
-						FUHTException::Throwf(HeaderParser, TEXT("ServiceResponse needs to specify an endpoint name"));
+						HeaderParser.Throwf(TEXT("ServiceResponse needs to specify an endpoint name"));
 					}
 				}
 				break;
@@ -466,7 +466,7 @@ namespace
 				{
 					if (FuncInfo.FunctionFlags & FUNC_Event)
 					{
-						UE_LOG_ERROR_UHT(HeaderParser, TEXT("Function cannot be a blueprint event and a blueprint getter."));
+						HeaderParser.LogError(TEXT("Function cannot be a blueprint event and a blueprint getter."));
 					}
 
 					bSawPropertyAccessor = true;
@@ -480,7 +480,7 @@ namespace
 				{
 					if (FuncInfo.FunctionFlags & FUNC_Event)
 					{
-						UE_LOG_ERROR_UHT(HeaderParser, TEXT("Function cannot be a blueprint event and a blueprint setter."));
+						HeaderParser.LogError(TEXT("Function cannot be a blueprint event and a blueprint setter."));
 					}
 
 					bSawPropertyAccessor = true;
@@ -549,41 +549,41 @@ namespace
 
 			if (FuncInfo.FunctionFlags & FUNC_Static)
 			{
-				UE_LOG_ERROR_UHT(HeaderParser, TEXT("Static functions can't be replicated"));
+				HeaderParser.LogError(TEXT("Static functions can't be replicated"));
 			}
 
 			if (!bIsNetReliable && !bSpecifiedUnreliable && !bIsNetService)
 			{
-				UE_LOG_ERROR_UHT(HeaderParser, TEXT("Replicated function: 'reliable' or 'unreliable' is required"));
+				HeaderParser.LogError(TEXT("Replicated function: 'reliable' or 'unreliable' is required"));
 			}
 
 			if (bIsNetReliable && bSpecifiedUnreliable && !bIsNetService)
 			{
-				UE_LOG_ERROR_UHT(HeaderParser, TEXT("'reliable' and 'unreliable' are mutually exclusive"));
+				HeaderParser.LogError(TEXT("'reliable' and 'unreliable' are mutually exclusive"));
 			}
 		}
 		else if (FuncInfo.FunctionFlags & FUNC_NetReliable)
 		{
-			UE_LOG_ERROR_UHT(HeaderParser, TEXT("'reliable' specified without 'client' or 'server'"));
+			HeaderParser.LogError(TEXT("'reliable' specified without 'client' or 'server'"));
 		}
 		else if (bSpecifiedUnreliable)
 		{
-			UE_LOG_ERROR_UHT(HeaderParser, TEXT("'unreliable' specified without 'client' or 'server'"));
+			HeaderParser.LogError(TEXT("'unreliable' specified without 'client' or 'server'"));
 		}
 
 		if (FuncInfo.bSealedEvent && !(FuncInfo.FunctionFlags & FUNC_Event))
 		{
-			UE_LOG_ERROR_UHT(HeaderParser, TEXT("SealedEvent may only be used on events"));
+			HeaderParser.LogError(TEXT("SealedEvent may only be used on events"));
 		}
 
 		if (FuncInfo.bSealedEvent && FuncInfo.FunctionFlags & FUNC_BlueprintEvent)
 		{
-			UE_LOG_ERROR_UHT(HeaderParser, TEXT("SealedEvent cannot be used on Blueprint events"));
+			HeaderParser.LogError(TEXT("SealedEvent cannot be used on Blueprint events"));
 		}
 
 		if (FuncInfo.bForceBlueprintImpure && (FuncInfo.FunctionFlags & FUNC_BlueprintPure) != 0)
 		{
-			UE_LOG_ERROR_UHT(HeaderParser, TEXT("BlueprintPure (or BlueprintPure=true) and BlueprintPure=false should not both appear on the same function, they are mutually exclusive"));
+			HeaderParser.LogError(TEXT("BlueprintPure (or BlueprintPure=true) and BlueprintPure=false should not both appear on the same function, they are mutually exclusive"));
 		}
 	}
 
@@ -602,7 +602,7 @@ namespace
 				return TEXT("Member variable declaration");
 
 			default:
-				FUHTException::Throwf(HeaderParser, TEXT("Unknown variable category"));
+				HeaderParser.Throwf(TEXT("Unknown variable category"));
 		}
 
 		// Unreachable
@@ -641,13 +641,13 @@ namespace
 		FToken Token;
 		if (Parser.GetToken(Token) && (Token.ConstantType != CPT_Float || Token.TokenType != TOKEN_Const))
 		{
-			FUHTException::Throwf(Parser, TEXT("Expected engine version in %s macro"), MacroToken.Identifier);
+			Parser.Throwf(TEXT("Expected engine version in %s macro"), MacroToken.Identifier);
 		}
 
 		Parser.RequireSymbol(TEXT(','), ErrorMessageGetter);
 		if (Parser.GetToken(Token) && (Token.ConstantType != CPT_String || Token.TokenType != TOKEN_Const))
 		{
-			FUHTException::Throwf(Parser, TEXT("Expected deprecation message in %s macro"), MacroToken.Identifier);
+			Parser.Throwf(TEXT("Expected deprecation message in %s macro"), MacroToken.Identifier);
 		}
 
 		Parser.RequireSymbol(TEXT(')'), ErrorMessageGetter);
@@ -679,7 +679,7 @@ namespace
 			return EPointerMemberBehavior::AllowAndLog;
 		}
 
-		FUHTException::Throwf(TEXT("Configuration"), 1, TEXT("Unrecognized native pointer member behavior: %s"), *InString );
+		FUHTMessage(FString(TEXT("Configuration"))).Throwf(TEXT("Unrecognized native pointer member behavior: %s"), *InString);
 		return EPointerMemberBehavior::Disallow;
 	}
 
@@ -787,7 +787,7 @@ FUnrealClassDefinitionInfo* FHeaderParser::GetQualifiedClass(const TCHAR* Thing)
 
 	if (!ClassName[0])
 	{
-		FUHTException::Throwf(*this, TEXT("%s: Missing class name"), Thing);
+		Throwf(TEXT("%s: Missing class name"), Thing);
 	}
 
 	return FUnrealClassDefinitionInfo::FindScriptClassOrThrow(*this, ClassName);
@@ -829,7 +829,7 @@ FUnrealFunctionDefinitionInfo* FHeaderParser::FindFunction
 				{
 					if (Thing)
 					{
-						FUHTException::Throwf(InScope, TEXT("%s: expecting function or delegate, got a property"), Thing);
+						InScope.Throwf(TEXT("%s: expecting function or delegate, got a property"), Thing);
 					}
 					return nullptr;
 				}
@@ -866,7 +866,7 @@ FUnrealPropertyDefinitionInfo* FHeaderParser::FindProperty(const FUnrealStructDe
 				{
 					if (Thing)
 					{
-						FUHTException::Throwf(InScope, TEXT("%s: expecting a property, got a function or delegate"), Thing);
+						InScope.Throwf(TEXT("%s: expecting a property, got a function or delegate"), Thing);
 					}
 					return nullptr;
 				}
@@ -951,7 +951,7 @@ FUnrealEnumDefinitionInfo& FHeaderParser::CompileEnum()
 	EEnumFlags      Flags         = EEnumFlags::None;
 	if (!GetIdentifier(EnumToken))
 	{
-		FUHTException::Throwf(*this, TEXT("Missing identifier after UENUM()") );
+		Throwf(TEXT("Missing identifier after UENUM()") );
 	}
 
 	if (EnumToken.Matches(TEXT("namespace"), ESearchCase::CaseSensitive))
@@ -966,7 +966,7 @@ FUnrealEnumDefinitionInfo& FHeaderParser::CompileEnum()
 	{
 		if (!GetIdentifier(EnumToken))
 		{
-			FUHTException::Throwf(*this, TEXT("Missing identifier after enum") );
+			Throwf(TEXT("Missing identifier after enum") );
 		}
 
 		if (EnumToken.Matches(TEXT("class"), ESearchCase::CaseSensitive) || EnumToken.Matches(TEXT("struct"), ESearchCase::CaseSensitive))
@@ -987,13 +987,13 @@ FUnrealEnumDefinitionInfo& FHeaderParser::CompileEnum()
 	}
 	else
 	{
-		FUHTException::Throwf(*this, TEXT("UENUM() should be followed by \'enum\' or \'namespace\' keywords.") );
+		Throwf(TEXT("UENUM() should be followed by \'enum\' or \'namespace\' keywords.") );
 	}
 
 	// Get enumeration name.
 	if (!bReadEnumName)
 	{
-		FUHTException::Throwf(*this, TEXT("Missing enumeration name") );
+		Throwf(TEXT("Missing enumeration name") );
 	}
 
 	ParseFieldMetaData(EnumMetaData, EnumToken.Identifier);
@@ -1008,7 +1008,7 @@ FUnrealEnumDefinitionInfo& FHeaderParser::CompileEnum()
 		switch ((EEnumSpecifier)Algo::FindSortedStringCaseInsensitive(*Specifier.Key, GEnumSpecifierStrings))
 		{
 		default:
-			FUHTException::Throwf(*this, TEXT("Unknown enum specifier '%s'"), *Specifier.Key);
+			Throwf(TEXT("Unknown enum specifier '%s'"), *Specifier.Key);
 
 		case EEnumSpecifier::Flags:
 			Flags |= EEnumFlags::Flags;
@@ -1033,7 +1033,7 @@ FUnrealEnumDefinitionInfo& FHeaderParser::CompileEnum()
 			FToken BaseToken;
 			if (!GetIdentifier(BaseToken))
 			{
-				FUHTException::Throwf(*this, TEXT("Missing enum base") );
+				Throwf(TEXT("Missing enum base") );
 			}
 
 			if (!FCString::Strcmp(BaseToken.Identifier, TEXT("uint8")))
@@ -1070,7 +1070,7 @@ FUnrealEnumDefinitionInfo& FHeaderParser::CompileEnum()
 			}
 			else
 			{
-				FUHTException::Throwf(*this, TEXT("Unsupported enum class base type: %s"), BaseToken.Identifier);
+				Throwf(TEXT("Unsupported enum class base type: %s"), BaseToken.Identifier);
 			}
 		}
 		else
@@ -1084,13 +1084,13 @@ FUnrealEnumDefinitionInfo& FHeaderParser::CompileEnum()
 	{
 		if (EnumHasAnyFlags(Flags, EEnumFlags::Flags))
 		{
-			FUHTException::Throwf(*this, TEXT("The 'Flags' specifier can only be used on enum classes"));
+			Throwf(TEXT("The 'Flags' specifier can only be used on enum classes"));
 		}
 	}
 
 	if (UnderlyingType != EUnderlyingEnumType::uint8 && EnumMetaData.Contains(FHeaderParserNames::NAME_BlueprintType))
 	{
-		FUHTException::Throwf(*this, TEXT("Invalid BlueprintType enum base - currently only uint8 supported"));
+		Throwf(TEXT("Invalid BlueprintType enum base - currently only uint8 supported"));
 	}
 
 	// Get opening brace.
@@ -1108,7 +1108,7 @@ FUnrealEnumDefinitionInfo& FHeaderParser::CompileEnum()
 			FToken InnerEnumToken;
 			if (!GetIdentifier(InnerEnumToken))
 			{
-				FUHTException::Throwf(*this, TEXT("Missing enumeration name") );
+				Throwf(TEXT("Missing enumeration name") );
 			}
 
 			EnumDef.SetCppType(FString::Printf(TEXT("%s::%s"), EnumToken.Identifier, InnerEnumToken.Identifier));
@@ -1148,7 +1148,7 @@ FUnrealEnumDefinitionInfo& FHeaderParser::CompileEnum()
 			FToken InitToken;
 			if (!GetToken(InitToken))
 			{
-				FUHTException::Throwf(*this, TEXT("UENUM: missing enumerator initializer"));
+				Throwf(TEXT("UENUM: missing enumerator initializer"));
 			}
 
 			int64 NewEnumValue = -1;
@@ -1163,7 +1163,7 @@ FUnrealEnumDefinitionInfo& FHeaderParser::CompileEnum()
 			{
 				if (!GetToken(InitToken))
 				{
-					FUHTException::Throwf(*this, TEXT("Enumerator: end of file encountered while parsing the initializer"));
+					Throwf(TEXT("Enumerator: end of file encountered while parsing the initializer"));
 				}
 
 				if (InitToken.TokenType == TOKEN_Symbol)
@@ -1241,7 +1241,7 @@ FUnrealEnumDefinitionInfo& FHeaderParser::CompileEnum()
 			FToken ClosingBrace;
 			if (!GetToken(ClosingBrace))
 			{
-				FUHTException::Throwf(*this, TEXT("UENUM: end of file encountered"));
+				Throwf(TEXT("UENUM: end of file encountered"));
 			}
 
 			if (ClosingBrace.TokenType == TOKEN_Symbol && ClosingBrace.Matches(TEXT('}')))
@@ -1271,17 +1271,17 @@ FUnrealEnumDefinitionInfo& FHeaderParser::CompileEnum()
 		const int32 MaxEnumItemIndex = EnumDef.GetIndexByName(MaxEnumItem);
 		if (MaxEnumItemIndex != INDEX_NONE)
 		{
-			FUHTException::Throwf(*this, TEXT("Illegal enumeration tag specified.  Conflicts with auto-generated tag '%s'"), *MaxEnumItem.ToString());
+			Throwf(TEXT("Illegal enumeration tag specified.  Conflicts with auto-generated tag '%s'"), *MaxEnumItem.ToString());
 		}
 
-		FUHTException::Throwf(*this, TEXT("Unable to generate enum MAX entry '%s' due to name collision"), *MaxEnumItem.ToString());
+		Throwf(TEXT("Unable to generate enum MAX entry '%s' due to name collision"), *MaxEnumItem.ToString());
 	}
 
 	CheckDocumentationPolicyForEnum(EnumDef, EnumValueMetaData, EntryMetaData);
 
 	if (!EnumDef.IsValidEnumValue(0) && EnumMetaData.Contains(FHeaderParserNames::NAME_BlueprintType))
 	{
-		UE_LOG_WARNING_UHT(EnumDef, TEXT("'%s' does not have a 0 entry! (This is a problem when the enum is initalized by default)"), *EnumDef.GetName());
+		EnumDef.LogWarning(TEXT("'%s' does not have a 0 entry! (This is a problem when the enum is initalized by default)"), *EnumDef.GetName());
 	}
 
 	// Add the metadata gathered for the enum to the package
@@ -1762,7 +1762,7 @@ FUnrealScriptStructDefinitionInfo& FHeaderParser::CompileStructDeclaration()
 		{
 			default:
 			{
-				FUHTException::Throwf(*this, TEXT("Unknown struct specifier '%s'"), *Specifier.Key);
+				Throwf(TEXT("Unknown struct specifier '%s'"), *Specifier.Key);
 			}
 			break;
 
@@ -1786,7 +1786,7 @@ FUnrealScriptStructDefinitionInfo& FHeaderParser::CompileStructDeclaration()
 
 				if (!FPaths::IsSamePath(Filename, GUObjectDef->GetUnrealSourceFile().GetFilename()))
 				{
-					UE_LOG_ERROR_UHT(*this, TEXT("Immutable is being phased out in favor of SerializeNative, and is only legal on the mirror structs declared in UObject"));
+					LogError(TEXT("Immutable is being phased out in favor of SerializeNative, and is only legal on the mirror structs declared in UObject"));
 				}
 			}
 			break;
@@ -1840,7 +1840,7 @@ FUnrealScriptStructDefinitionInfo& FHeaderParser::CompileStructDeclaration()
 		// If the struct still wasn't found, throw an error
 		if (BaseStructDef == nullptr)
 		{
-			FUHTException::Throwf(*this, TEXT("'struct': Can't find struct '%s'"), *ParentStructNameInScript);
+			Throwf(TEXT("'struct': Can't find struct '%s'"), *ParentStructNameInScript);
 		}
 
 		// If the struct was found, confirm it adheres to the correct syntax. This should always fail if we were expecting an override that was not found.
@@ -1849,7 +1849,7 @@ FUnrealScriptStructDefinitionInfo& FHeaderParser::CompileStructDeclaration()
 			const TCHAR* PrefixCPP = UHTConfig.StructsWithTPrefix.Contains(ParentStructNameStripped) ? TEXT("T") : BaseStructDef->GetPrefixCPP();
 			if (ParentStructNameInScript != FString::Printf(TEXT("%s%s"), PrefixCPP, *ParentStructNameStripped))
 			{
-				FUHTException::Throwf(*this, TEXT("Parent Struct '%s' is missing a valid Unreal prefix, expecting '%s'"), *ParentStructNameInScript, *FString::Printf(TEXT("%s%s"), PrefixCPP, *BaseStructDef->GetName()));
+				Throwf(TEXT("Parent Struct '%s' is missing a valid Unreal prefix, expecting '%s'"), *ParentStructNameInScript, *FString::Printf(TEXT("%s%s"), PrefixCPP, *BaseStructDef->GetName()));
 			}
 		}
 
@@ -1875,14 +1875,14 @@ FUnrealScriptStructDefinitionInfo& FHeaderParser::CompileStructDeclaration()
 		FString ExpectedStructName = FString::Printf(TEXT("%s%s"), ExpectedPrefixCPP, *StructNameStripped);
 		if (StructNameInScript != ExpectedStructName)
 		{
-			FUHTException::Throwf(*this, TEXT("Struct '%s' has an invalid Unreal prefix, expecting '%s'"), *StructNameInScript, *ExpectedStructName);
+			Throwf(TEXT("Struct '%s' has an invalid Unreal prefix, expecting '%s'"), *StructNameInScript, *ExpectedStructName);
 		}
 	}
 	else
 	{
 		const TCHAR* ExpectedPrefixCPP = UHTConfig.StructsWithTPrefix.Contains(StructNameInScript) ? TEXT("T") : StructDef.GetPrefixCPP();
 		FString ExpectedStructName = FString::Printf(TEXT("%s%s"), ExpectedPrefixCPP, *StructNameInScript);
-		FUHTException::Throwf(*this, TEXT("Struct '%s' is missing a valid Unreal prefix, expecting '%s'"), *StructNameInScript, *ExpectedStructName);
+		Throwf(TEXT("Struct '%s' is missing a valid Unreal prefix, expecting '%s'"), *StructNameInScript, *ExpectedStructName);
 	}
 
 	StructDef.SetStructFlags(EStructFlags(StructFlags));
@@ -1921,7 +1921,7 @@ FUnrealScriptStructDefinitionInfo& FHeaderParser::CompileStructDeclaration()
 		}
 		else if (Token.Matches(TEXT("UFUNCTION"), ESearchCase::CaseSensitive))
 		{
-			FUHTException::Throwf(*this, TEXT("USTRUCTs cannot contain UFUNCTIONs."));
+			Throwf(TEXT("USTRUCTs cannot contain UFUNCTIONs."));
 		}
 		else if (Token.Matches(TEXT("RIGVM_METHOD"), ESearchCase::CaseSensitive))
 		{
@@ -1932,12 +1932,12 @@ FUnrealScriptStructDefinitionInfo& FHeaderParser::CompileStructDeclaration()
 			// Match 'GENERATED_USTRUCT_BODY' '(' [StructName] ')' or 'GENERATED_BODY' '(' [StructName] ')'
 			if (CurrentAccessSpecifier != ACCESS_Public)
 			{
-				FUHTException::Throwf(*this, TEXT("%s must be in the public scope of '%s', not private or protected."), Token.Identifier, *StructNameInScript);
+				Throwf(TEXT("%s must be in the public scope of '%s', not private or protected."), Token.Identifier, *StructNameInScript);
 			}
 
 			if (StructDef.GetMacroDeclaredLineNumber() != INDEX_NONE)
 			{
-				FUHTException::Throwf(*this, TEXT("Multiple %s declarations found in '%s'"), Token.Identifier, *StructNameInScript);
+				Throwf(TEXT("Multiple %s declarations found in '%s'"), Token.Identifier, *StructNameInScript);
 			}
 
 			StructDef.SetMacroDeclaredLineNumber(InputLine);
@@ -1962,7 +1962,7 @@ FUnrealScriptStructDefinitionInfo& FHeaderParser::CompileStructDeclaration()
 		{
 			if (CompilerDirectiveStack.Num() < 1)
 			{
-				FUHTException::Throwf(*this, TEXT("Unmatched '#endif' in class or global scope"));
+				Throwf(TEXT("Unmatched '#endif' in class or global scope"));
 			}
 			CompilerDirectiveStack.Pop();
 			// Do nothing and hope that the if code below worked out OK earlier
@@ -1977,7 +1977,7 @@ FUnrealScriptStructDefinitionInfo& FHeaderParser::CompileStructDeclaration()
 			{
 				if (bInvertConditional)
 				{
-					FUHTException::Throwf(*this, TEXT("Cannot use !WITH_EDITORONLY_DATA"));
+					Throwf(TEXT("Cannot use !WITH_EDITORONLY_DATA"));
 				}
 
 				PushCompilerDirective(ECompilerDirective::WithEditorOnlyData);
@@ -1986,7 +1986,7 @@ FUnrealScriptStructDefinitionInfo& FHeaderParser::CompileStructDeclaration()
 			{
 				if (bInvertConditional)
 				{
-					FUHTException::Throwf(*this, TEXT("Cannot use !WITH_EDITOR"));
+					Throwf(TEXT("Cannot use !WITH_EDITOR"));
 				}
 				PushCompilerDirective(ECompilerDirective::WithEditor);
 			}
@@ -1997,7 +1997,7 @@ FUnrealScriptStructDefinitionInfo& FHeaderParser::CompileStructDeclaration()
 			}
 			else
 			{
-				FUHTException::Throwf(*this, TEXT("'struct': Unsupported preprocessor directive inside a struct.") );
+				Throwf(TEXT("'struct': Unsupported preprocessor directive inside a struct.") );
 			}
 
 			if (bConsumeAsCppText)
@@ -2010,7 +2010,7 @@ FUnrealScriptStructDefinitionInfo& FHeaderParser::CompileStructDeclaration()
 
 					if ( ch==0 )
 					{
-						FUHTException::Throwf(*this, TEXT("Unexpected end of struct definition %s"), *StructDef.GetName());
+						Throwf(TEXT("Unexpected end of struct definition %s"), *StructDef.GetName());
 					}
 					else if ( ch=='{' || (ch=='#' && (PeekIdentifier(TEXT("if"), ESearchCase::CaseSensitive) || PeekIdentifier(TEXT("ifdef"), ESearchCase::CaseSensitive))) )
 					{
@@ -2050,7 +2050,7 @@ FUnrealScriptStructDefinitionInfo& FHeaderParser::CompileStructDeclaration()
 
 				if (!SkipDeclaration(Token))
 				{
-					FUHTException::Throwf(*this, TEXT("'struct': Unexpected '%s'"), FirstTokenIdentifier);
+					Throwf(TEXT("'struct': Unexpected '%s'"), FirstTokenIdentifier);
 				}	
 			}
 			else
@@ -2068,7 +2068,7 @@ FUnrealScriptStructDefinitionInfo& FHeaderParser::CompileStructDeclaration()
 	{
 		// Roll the line number back to the start of the struct body and error out
 		InputLine = SavedLineNumber;
-		FUHTException::Throwf(*this, TEXT("Expected a GENERATED_BODY() at the start of struct"));
+		Throwf(TEXT("Expected a GENERATED_BODY() at the start of struct"));
 	}
 
 	// Validate sparse class data
@@ -2153,11 +2153,11 @@ void FHeaderParser::CheckAllow( const TCHAR* Thing, ENestAllowFlags AllowFlags )
 	{
 		if (TopNest->NestType == ENestType::GlobalScope)
 		{
-			FUHTException::Throwf(*this, TEXT("%s is not allowed before the Class definition"), Thing );
+			Throwf(TEXT("%s is not allowed before the Class definition"), Thing );
 		}
 		else
 		{
-			FUHTException::Throwf(*this, TEXT("%s is not allowed here"), Thing );
+			Throwf(TEXT("%s is not allowed here"), Thing );
 		}
 	}
 }
@@ -2176,7 +2176,7 @@ void FHeaderParser::PushNest(ENestType NestType, FUnrealStructDefinitionInfo* In
 	// Prevent overnesting.
 	if (NestLevel >= MAX_NEST_LEVELS)
 	{
-		FUHTException::Throwf(*this, TEXT("Maximum nesting limit exceeded"));
+		Throwf(TEXT("Maximum nesting limit exceeded"));
 	}
 
 	// Inherit info from stack node above us.
@@ -2208,7 +2208,7 @@ void FHeaderParser::PushNest(ENestType NestType, FUnrealStructDefinitionInfo* In
 		break;
 
 	default:
-		FUHTException::Throwf(*this, TEXT("Internal error in PushNest, type %i"), (uint8)NestType);
+		Throwf(TEXT("Internal error in PushNest, type %i"), (uint8)NestType);
 		break;
 	}
 }
@@ -2224,16 +2224,16 @@ void FHeaderParser::PopNest(ENestType NestType, const TCHAR* Descr)
 	// Validate the nesting state.
 	if (NestLevel <= 0)
 	{
-		FUHTException::Throwf(*this, TEXT("Unexpected '%s' at global scope"), Descr, NestTypeName(NestType));
+		Throwf(TEXT("Unexpected '%s' at global scope"), Descr, NestTypeName(NestType));
 	}
 	else if (TopNest->NestType != NestType)
 	{
-		FUHTException::Throwf(*this, TEXT("Unexpected end of %s in '%s' block"), Descr, NestTypeName(TopNest->NestType));
+		Throwf(TEXT("Unexpected end of %s in '%s' block"), Descr, NestTypeName(TopNest->NestType));
 	}
 
 	if (NestType != ENestType::GlobalScope && NestType != ENestType::Class && NestType != ENestType::Interface && NestType != ENestType::NativeInterface && NestType != ENestType::FunctionDeclaration)
 	{
-		FUHTException::Throwf(*this, TEXT("Bad first pass NestType %i"), (uint8)NestType);
+		Throwf(TEXT("Bad first pass NestType %i"), (uint8)NestType);
 	}
 
 	// Pop the nesting level.
@@ -2310,21 +2310,21 @@ void FHeaderParser::FixupDelegateProperties(FUnrealStructDefinitionInfo& StructD
 				if ( DelegateClassName.Len() == 0 )
 				{
 					UngetToken(PropertyDef->GetLineNumber(), PropertyDef->GetParsePosition());
-					FUHTException::Throwf(*this, TEXT("Invalid scope specified in delegate property function reference: '%s'"), *NameOfDelegateFunction);
+					Throwf(TEXT("Invalid scope specified in delegate property function reference: '%s'"), *NameOfDelegateFunction);
 				}
 
 				// verify that we got a valid string for the name of the function
 				if ( DelegateName.Len() == 0 )
 				{
 					UngetToken(PropertyDef->GetLineNumber(), PropertyDef->GetParsePosition());
-					FUHTException::Throwf(*this, TEXT("Invalid delegate name specified in delegate property function reference '%s'"), *NameOfDelegateFunction);
+					Throwf(TEXT("Invalid delegate name specified in delegate property function reference '%s'"), *NameOfDelegateFunction);
 				}
 
 				// make sure that the class that contains the delegate can be referenced here
 				FUnrealClassDefinitionInfo* DelegateOwnerClassDef = FUnrealClassDefinitionInfo::FindScriptClassOrThrow(*this, DelegateClassName);
 				if (DelegateOwnerClassDef->GetScope()->FindTypeByName(*DelegateName) != nullptr)
 				{
-					FUHTException::Throwf(*this, TEXT("Inaccessible type: '%s'"), *DelegateOwnerClassDef->GetPathName());
+					Throwf(TEXT("Inaccessible type: '%s'"), *DelegateOwnerClassDef->GetPathName());
 				}
 				SourceDelegateFunctionDef = FindFunction(*DelegateOwnerClassDef, *DelegateName, false, nullptr);
 			}
@@ -2332,12 +2332,12 @@ void FHeaderParser::FixupDelegateProperties(FUnrealStructDefinitionInfo& StructD
 			if (SourceDelegateFunctionDef == NULL )
 			{
 				UngetToken(PropertyDef->GetLineNumber(), PropertyDef->GetParsePosition());
-				FUHTException::Throwf(*this, TEXT("Failed to find delegate function '%s'"), *NameOfDelegateFunction);
+				Throwf(TEXT("Failed to find delegate function '%s'"), *NameOfDelegateFunction);
 			}
 			else if (!SourceDelegateFunctionDef->HasAnyFunctionFlags(FUNC_Delegate))
 			{
 				UngetToken(PropertyDef->GetLineNumber(), PropertyDef->GetParsePosition());
-				FUHTException::Throwf(*this, TEXT("Only delegate functions can be used as the type for a delegate property; '%s' is not a delegate."), *NameOfDelegateFunction);
+				Throwf(TEXT("Only delegate functions can be used as the type for a delegate property; '%s' is not a delegate."), *NameOfDelegateFunction);
 			}
 		}
 
@@ -2355,7 +2355,7 @@ void FHeaderParser::FixupDelegateProperties(FUnrealStructDefinitionInfo& StructD
 			}
 			else
 			{
-				FUHTException::Throwf(*this, TEXT("Unable to declare a single-cast delegate property for a multi-cast delegate type '%s'.  Either add a 'multicast' qualifier to the property or change the delegate type to be single-cast as well."), *SourceDelegateFunctionDef->GetName());
+				Throwf(TEXT("Unable to declare a single-cast delegate property for a multi-cast delegate type '%s'.  Either add a 'multicast' qualifier to the property or change the delegate type to be single-cast as well."), *SourceDelegateFunctionDef->GetName());
 			}
 		}
 		else if (PropertyBase.Type == CPT_MulticastDelegate)
@@ -2372,7 +2372,7 @@ void FHeaderParser::FixupDelegateProperties(FUnrealStructDefinitionInfo& StructD
 						{
 							FString ExtendedCPPType;
 							FString CPPType = FuncParamDef->GetCPPType(&ExtendedCPPType);
-							UE_LOG_ERROR_UHT(*this, TEXT("Type '%s%s' is not supported by blueprint. %s.%s"), *CPPType, *ExtendedCPPType, *SourceDelegateFunctionDef->GetName(), *FuncParamDef->GetName());
+							LogError(TEXT("Type '%s%s' is not supported by blueprint. %s.%s"), *CPPType, *ExtendedCPPType, *SourceDelegateFunctionDef->GetName(), *FuncParamDef->GetName());
 						}
 
 						if(FuncParamDef->HasAllPropertyFlags(CPF_OutParm) && !FuncParamDef->HasAllPropertyFlags(CPF_ConstParm)  )
@@ -2381,7 +2381,7 @@ void FHeaderParser::FixupDelegateProperties(FUnrealStructDefinitionInfo& StructD
 							const bool bAllowedArrayRefFromBP = bClassGeneratedFromBP && FuncParamDef->IsDynamicArray();
 							if (!bAllowedArrayRefFromBP)
 							{
-								UE_LOG_ERROR_UHT(*this, TEXT("BlueprintAssignable delegates do not support non-const references at the moment. Function: %s Parameter: '%s'"), *SourceDelegateFunctionDef->GetName(), *FuncParamDef->GetName());
+								LogError(TEXT("BlueprintAssignable delegates do not support non-const references at the moment. Function: %s Parameter: '%s'"), *SourceDelegateFunctionDef->GetName(), *FuncParamDef->GetName());
 							}
 						}
 					}
@@ -2390,7 +2390,7 @@ void FHeaderParser::FixupDelegateProperties(FUnrealStructDefinitionInfo& StructD
 			}
 			else
 			{
-				FUHTException::Throwf(*this, TEXT("Unable to declare a multi-cast delegate property for a single-cast delegate type '%s'.  Either remove the 'multicast' qualifier from the property or change the delegate type to be 'multicast' as well."), *SourceDelegateFunctionDef->GetName());
+				Throwf(TEXT("Unable to declare a multi-cast delegate property for a single-cast delegate type '%s'.  Either remove the 'multicast' qualifier from the property or change the delegate type to be 'multicast' as well."), *SourceDelegateFunctionDef->GetName());
 			}
 		}
 	}
@@ -2416,7 +2416,7 @@ void FHeaderParser::CheckSparseClassData(const FUnrealStructDefinitionInfo& Stru
 	{
 		if (ClassDef == nullptr)
 		{
-			FUHTException::Throwf(StructDef, TEXT("%s contains sparse class data but is not a class."), *StructDef.GetName());
+			StructDef.Throwf(TEXT("%s contains sparse class data but is not a class."), *StructDef.GetName());
 		}
 	}
 	else
@@ -2430,12 +2430,12 @@ void FHeaderParser::CheckSparseClassData(const FUnrealStructDefinitionInfo& Stru
 	// for now we only support one sparse class data structure per class
 	if (SparseClassDataTypes.Num() > 1)
 	{
-		FUHTException::Throwf(StructDef, TEXT("Class %s contains multiple sparse class data types."), *ClassDef->GetName());
+		StructDef.Throwf(TEXT("Class %s contains multiple sparse class data types."), *ClassDef->GetName());
 		return;
 	}
 	if (SparseClassDataTypes.Num() == 0)
 	{
-		FUHTException::Throwf(StructDef, TEXT("Class %s has sparse class metadata but does not specify a type."), *ClassDef->GetName());
+		StructDef.Throwf(TEXT("Class %s has sparse class metadata but does not specify a type."), *ClassDef->GetName());
 		return;
 	}
 
@@ -2446,7 +2446,7 @@ void FHeaderParser::CheckSparseClassData(const FUnrealStructDefinitionInfo& Stru
 		// make sure the sparse class data struct actually exists
 		if (!SparseClassDataStructDef)
 		{
-			FUHTException::Throwf(StructDef, TEXT("Unable to find sparse data type %s for class %s."), *SparseClassDataTypeName, *ClassDef->GetName());
+			StructDef.Throwf(TEXT("Unable to find sparse data type %s for class %s."), *SparseClassDataTypeName, *ClassDef->GetName());
 			return;
 		}
 
@@ -2455,19 +2455,19 @@ void FHeaderParser::CheckSparseClassData(const FUnrealStructDefinitionInfo& Stru
 		{
 			if (PropertyDef->HasAnyPropertyFlags(CPF_BlueprintAssignable))
 			{
-				FUHTException::Throwf(StructDef, TEXT("Sparse class data types can not contain blueprint assignable delegates. Type '%s' Delegate '%s'"), *SparseClassDataStructDef->GetName(), *PropertyDef->GetName());
+				StructDef.Throwf(TEXT("Sparse class data types can not contain blueprint assignable delegates. Type '%s' Delegate '%s'"), *SparseClassDataStructDef->GetName(), *PropertyDef->GetName());
 			}
 
 			// all sparse properties should have EditDefaultsOnly
 			if (!PropertyDef->HasAllPropertyFlags(CPF_Edit | CPF_DisableEditOnInstance))
 			{
-				FUHTException::Throwf(StructDef, TEXT("Sparse class data types must be EditDefaultsOnly. Type '%s' Property '%s'"), *SparseClassDataStructDef->GetName(), *PropertyDef->GetName());
+				StructDef.Throwf(TEXT("Sparse class data types must be EditDefaultsOnly. Type '%s' Property '%s'"), *SparseClassDataStructDef->GetName(), *PropertyDef->GetName());
 			}
 
 			// no sparse properties should have BlueprintReadWrite
 			if (PropertyDef->HasAllPropertyFlags(CPF_BlueprintVisible) && !PropertyDef->HasAllPropertyFlags(CPF_BlueprintReadOnly))
 			{
-				FUHTException::Throwf(StructDef, TEXT("Sparse class data types must not be BlueprintReadWrite. Type '%s' Property '%s'"), *SparseClassDataStructDef->GetName(), *PropertyDef->GetName());
+				StructDef.Throwf(TEXT("Sparse class data types must not be BlueprintReadWrite. Type '%s' Property '%s'"), *SparseClassDataStructDef->GetName(), *PropertyDef->GetName());
 			}
 		}
 
@@ -2482,7 +2482,7 @@ void FHeaderParser::CheckSparseClassData(const FUnrealStructDefinitionInfo& Stru
 			{
 				if (!SparseClassDataStructDef->IsChildOf(*ParentSparseClassDataStructDef))
 				{
-					FUHTException::Throwf(StructDef, TEXT("Class %s is a child of %s but its sparse class data struct, %s, does not inherit from %s."), *ClassDef->GetName(), *ParentClassDef->GetName(), *SparseClassDataStructDef->GetName(), *ParentSparseClassDataStructDef->GetName());
+					StructDef.Throwf(TEXT("Class %s is a child of %s but its sparse class data struct, %s, does not inherit from %s."), *ClassDef->GetName(), *ParentClassDef->GetName(), *SparseClassDataStructDef->GetName(), *ParentSparseClassDataStructDef->GetName());
 				}
 			}
 		}
@@ -2494,7 +2494,7 @@ void FHeaderParser::ValidateClassFlags(const FUnrealClassDefinitionInfo& ToValid
 	if (ToValidate.HasAnyClassFlags(CLASS_NeedsDeferredDependencyLoading) && !ToValidate.IsChildOf(*GUClassDef))
 	{
 		// CLASS_NeedsDeferredDependencyLoading can only be set on classes derived from UClass
-		FUHTException::Throwf(ToValidate, TEXT("NeedsDeferredDependencyLoading is set on %s but the flag can only be used with classes derived from UClass."), *ToValidate.GetClass()->GetName());
+		ToValidate.Throwf(TEXT("NeedsDeferredDependencyLoading is set on %s but the flag can only be used with classes derived from UClass."), *ToValidate.GetClass()->GetName());
 	}
 }
 
@@ -2506,23 +2506,23 @@ void FHeaderParser::VerifyBlueprintPropertyGetter(FUnrealPropertyDefinitionInfo&
 	FUnrealPropertyDefinitionInfo* ReturnPropDef = TargetFuncDef->GetReturn();
 	if (Properties.Num() > 1 || (Properties.Num() == 1 && ReturnPropDef == nullptr))
 	{
-		UE_LOG_ERROR_UHT(*this, TEXT("Blueprint Property getter function %s must not have parameters."), *TargetFuncDef->GetName());
+		LogError(TEXT("Blueprint Property getter function %s must not have parameters."), *TargetFuncDef->GetName());
 	}
 
 	if (ReturnPropDef == nullptr || !PropertyDef.SameType(*ReturnPropDef))
 	{
 		FString ExtendedCPPType;
 		FString CPPType = PropertyDef.GetCPPType(&ExtendedCPPType);
-		UE_LOG_ERROR_UHT(*this, TEXT("Blueprint Property getter function %s must have return value of type %s%s."), *TargetFuncDef->GetName(), *CPPType, *ExtendedCPPType);
+		LogError(TEXT("Blueprint Property getter function %s must have return value of type %s%s."), *TargetFuncDef->GetName(), *CPPType, *ExtendedCPPType);
 	}
 
 	if (TargetFuncDef->HasAnyFunctionFlags(FUNC_Event))
 	{
-		UE_LOG_ERROR_UHT(*this, TEXT("Blueprint Property setter function cannot be a blueprint event."));
+		LogError(TEXT("Blueprint Property setter function cannot be a blueprint event."));
 	}
 	else if (!TargetFuncDef->HasAnyFunctionFlags(FUNC_BlueprintPure))
 	{
-		UE_LOG_ERROR_UHT(*this, TEXT("Blueprint Property getter function must be pure."));
+		LogError(TEXT("Blueprint Property getter function must be pure."));
 	}
 }
 
@@ -2533,7 +2533,7 @@ void FHeaderParser::VerifyBlueprintPropertySetter(FUnrealPropertyDefinitionInfo&
 
 	if (ReturnPropDef)
 	{
-		UE_LOG_ERROR_UHT(*this, TEXT("Blueprint Property setter function %s must not have a return value."), *TargetFuncDef->GetName());
+		LogError(TEXT("Blueprint Property setter function %s must not have a return value."), *TargetFuncDef->GetName());
 	}
 	else
 	{
@@ -2542,21 +2542,21 @@ void FHeaderParser::VerifyBlueprintPropertySetter(FUnrealPropertyDefinitionInfo&
 		{
 			FString ExtendedCPPType;
 			FString CPPType = PropertyDef.GetCPPType(&ExtendedCPPType);
-			UE_LOG_ERROR_UHT(*this, TEXT("Blueprint Property setter function %s must have exactly one parameter of type %s%s."), *TargetFuncDef->GetName(), *CPPType, *ExtendedCPPType);
+			LogError(TEXT("Blueprint Property setter function %s must have exactly one parameter of type %s%s."), *TargetFuncDef->GetName(), *CPPType, *ExtendedCPPType);
 		}
 	}
 
 	if (TargetFuncDef->HasAnyFunctionFlags(FUNC_Event))
 	{
-		UE_LOG_ERROR_UHT(*this, TEXT("Blueprint Property setter function cannot be a blueprint event."));
+		LogError(TEXT("Blueprint Property setter function cannot be a blueprint event."));
 	}
 	else if (!TargetFuncDef->HasAnyFunctionFlags(FUNC_BlueprintCallable))
 	{
-		UE_LOG_ERROR_UHT(*this, TEXT("Blueprint Property setter function must be blueprint callable."));
+		LogError(TEXT("Blueprint Property setter function must be blueprint callable."));
 	}
 	else if (TargetFuncDef->HasAnyFunctionFlags(FUNC_BlueprintPure))
 	{
-		UE_LOG_ERROR_UHT(*this, TEXT("Blueprint Property setter function must not be pure."));
+		LogError(TEXT("Blueprint Property setter function must not be pure."));
 	}
 }
 
@@ -2566,7 +2566,7 @@ void FHeaderParser::VerifyRepNotifyCallback(FUnrealPropertyDefinitionInfo& Prope
 	{
 		if (TargetFuncDef->GetReturn())
 		{
-			UE_LOG_ERROR_UHT(*this, TEXT("Replication notification function %s must not have return value."), *TargetFuncDef->GetName());
+			LogError(TEXT("Replication notification function %s must not have return value."), *TargetFuncDef->GetName());
 		}
 
 		const bool bIsArrayProperty = PropertyDef.IsStaticArray() || PropertyDef.IsDynamicArray();
@@ -2575,7 +2575,7 @@ void FHeaderParser::VerifyRepNotifyCallback(FUnrealPropertyDefinitionInfo& Prope
 		const TArray<TSharedRef<FUnrealPropertyDefinitionInfo>>& Properties = TargetFuncDef->GetProperties();
 		if (Properties.Num() > MaxParms)
 		{
-			UE_LOG_ERROR_UHT(*this, TEXT("Replication notification function %s has too many parameters."), *TargetFuncDef->GetName());
+			LogError(TEXT("Replication notification function %s has too many parameters."), *TargetFuncDef->GetName());
 		}
 
 		if (Properties.Num() >= 1)
@@ -2586,7 +2586,7 @@ void FHeaderParser::VerifyRepNotifyCallback(FUnrealPropertyDefinitionInfo& Prope
 			{
 				FString ExtendedCPPType;
 				FString CPPType = ParmDef->GetCPPType(&ExtendedCPPType);
-				UE_LOG_ERROR_UHT(*this, TEXT("Replication notification function %s has invalid parameter for property %s. First (optional) parameter must be of type %s%s."), *TargetFuncDef->GetName(), *ParmDef->GetName(), *CPPType, *ExtendedCPPType);
+				LogError(TEXT("Replication notification function %s has invalid parameter for property %s. First (optional) parameter must be of type %s%s."), *TargetFuncDef->GetName(), *ParmDef->GetName(), *CPPType, *ExtendedCPPType);
 			}
 		}
 
@@ -2597,14 +2597,14 @@ void FHeaderParser::VerifyRepNotifyCallback(FUnrealPropertyDefinitionInfo& Prope
 			// A 2nd parameter for arrays can be specified as a const TArray<uint8>&. This is a list of element indices that have changed
 			if (!(ParmBase.Type == CPT_Byte && ParmBase.ArrayType == EArrayType::Dynamic && ParmDef->HasAllPropertyFlags(CPF_ConstParm | CPF_ReferenceParm)))
 			{
-				UE_LOG_ERROR_UHT(*this, TEXT("Replication notification function %s (optional) second parameter must be of type 'const TArray<uint8>&'"), *TargetFuncDef->GetName());
+				LogError(TEXT("Replication notification function %s (optional) second parameter must be of type 'const TArray<uint8>&'"), *TargetFuncDef->GetName());
 			}
 		}
 	}
 	else
 	{
 		// Couldn't find a valid function...
-		UE_LOG_ERROR_UHT(*this, TEXT("Replication notification function %s not found"), *PropertyDef.GetRepNotifyFunc().ToString() );
+		LogError(TEXT("Replication notification function %s not found"), *PropertyDef.GetRepNotifyFunc().ToString() );
 	}
 }
 void FHeaderParser::VerifyPropertyMarkups(FUnrealClassDefinitionInfo& TargetClassDef)
@@ -2651,7 +2651,7 @@ void FHeaderParser::VerifyPropertyMarkups(FUnrealClassDefinitionInfo& TargetClas
 				else
 				{
 					// Couldn't find a valid function...
-					UE_LOG_ERROR_UHT(*this, TEXT("Blueprint Property getter function %s not found"), *GetterFuncName);
+					LogError(TEXT("Blueprint Property getter function %s not found"), *GetterFuncName);
 				}
 			}
 
@@ -2667,7 +2667,7 @@ void FHeaderParser::VerifyPropertyMarkups(FUnrealClassDefinitionInfo& TargetClas
 					else
 					{
 						// Couldn't find a valid function...
-						UE_LOG_ERROR_UHT(*this, TEXT("Blueprint Property setter function %s not found"), *SetterFuncName);
+						LogError(TEXT("Blueprint Property setter function %s not found"), *SetterFuncName);
 					}
 				}
 			}
@@ -2693,11 +2693,11 @@ void FHeaderParser::CompileDirective()
 
 	if (!GetIdentifier(Directive))
 	{
-		FUHTException::Throwf(*this, TEXT("Missing compiler directive after '#'") );
+		Throwf(TEXT("Missing compiler directive after '#'") );
 	}
 	else if (Directive.Matches(TEXT("error"), ESearchCase::CaseSensitive))
 	{
-		FUHTException::Throwf(*this, TEXT("#error directive encountered") );
+		Throwf(TEXT("#error directive encountered") );
 	}
 	else if (Directive.Matches(TEXT("pragma"), ESearchCase::CaseSensitive))
 	{
@@ -2708,7 +2708,7 @@ void FHeaderParser::CompileDirective()
 		FToken Number;
 		if (!GetToken(Number) || (Number.TokenType != TOKEN_Const) || (Number.ConstantType != CPT_Int && Number.ConstantType != CPT_Int64))
 		{
-			FUHTException::Throwf(*this, TEXT("Missing line number in line number directive"));
+			Throwf(TEXT("Missing line number in line number directive"));
 		}
 
 		int32 newInputLine;
@@ -2745,7 +2745,7 @@ void FHeaderParser::CompileDirective()
 			FToken Define;
 			if (!GetIdentifier(Define))
 			{
-				FUHTException::Throwf(*this, TEXT("Missing define name '#if'") );
+				Throwf(TEXT("Missing define name '#if'") );
 			}
 
 			if ( Define.Matches(TEXT("WITH_EDITORONLY_DATA"), ESearchCase::CaseSensitive) )
@@ -2766,7 +2766,7 @@ void FHeaderParser::CompileDirective()
 			}
 			else
 			{
-				FUHTException::Throwf(*this, TEXT("Unknown define '#if %s' in class or global scope"), Define.Identifier);
+				Throwf(TEXT("Unknown define '#if %s' in class or global scope"), Define.Identifier);
 			}
 		}
 	}
@@ -2774,7 +2774,7 @@ void FHeaderParser::CompileDirective()
 	{
 		if (CompilerDirectiveStack.Num() < 1)
 		{
-			FUHTException::Throwf(*this, TEXT("Unmatched '#endif' in class or global scope"));
+			Throwf(TEXT("Unmatched '#endif' in class or global scope"));
 		}
 		CompilerDirectiveStack.Pop();
 	}
@@ -2793,7 +2793,7 @@ void FHeaderParser::CompileDirective()
 	}
 	else
 	{
-		FUHTException::Throwf(*this, TEXT("Unrecognized compiler directive %s"), Directive.Identifier );
+		Throwf(TEXT("Unrecognized compiler directive %s"), Directive.Identifier );
 	}
 
 	// Skip to end of line (or end of multiline #define).
@@ -2912,7 +2912,7 @@ void FHeaderParser::GetVarType(
 				{
 					if (bSeenEditSpecifier)
 					{
-						UE_LOG_ERROR_UHT(*this, TEXT("Found more than one edit/visibility specifier (%s), only one is allowed"), *Specifier.Key);
+						LogError(TEXT("Found more than one edit/visibility specifier (%s), only one is allowed"), *Specifier.Key);
 					}
 					Flags |= CPF_Edit;
 					bSeenEditSpecifier = true;
@@ -2923,7 +2923,7 @@ void FHeaderParser::GetVarType(
 				{
 					if (bSeenEditSpecifier)
 					{
-						UE_LOG_ERROR_UHT(*this, TEXT("Found more than one edit/visibility specifier (%s), only one is allowed"), *Specifier.Key);
+						LogError(TEXT("Found more than one edit/visibility specifier (%s), only one is allowed"), *Specifier.Key);
 					}
 					Flags |= CPF_Edit | CPF_DisableEditOnTemplate;
 					bSeenEditSpecifier = true;
@@ -2934,7 +2934,7 @@ void FHeaderParser::GetVarType(
 				{
 					if (bSeenEditSpecifier)
 					{
-						UE_LOG_ERROR_UHT(*this, TEXT("Found more than one edit/visibility specifier (%s), only one is allowed"), *Specifier.Key);
+						LogError(TEXT("Found more than one edit/visibility specifier (%s), only one is allowed"), *Specifier.Key);
 					}
 					Flags |= CPF_Edit | CPF_DisableEditOnInstance;
 					bSeenEditSpecifier = true;
@@ -2945,7 +2945,7 @@ void FHeaderParser::GetVarType(
 				{
 					if (bSeenEditSpecifier)
 					{
-						UE_LOG_ERROR_UHT(*this, TEXT("Found more than one edit/visibility specifier (%s), only one is allowed"), *Specifier.Key);
+						LogError(TEXT("Found more than one edit/visibility specifier (%s), only one is allowed"), *Specifier.Key);
 					}
 					Flags |= CPF_Edit | CPF_EditConst;
 					bSeenEditSpecifier = true;
@@ -2956,7 +2956,7 @@ void FHeaderParser::GetVarType(
 				{
 					if (bSeenEditSpecifier)
 					{
-						UE_LOG_ERROR_UHT(*this, TEXT("Found more than one edit/visibility specifier (%s), only one is allowed"), *Specifier.Key);
+						LogError(TEXT("Found more than one edit/visibility specifier (%s), only one is allowed"), *Specifier.Key);
 					}
 					Flags |= CPF_Edit | CPF_EditConst | CPF_DisableEditOnTemplate;
 					bSeenEditSpecifier = true;
@@ -2967,7 +2967,7 @@ void FHeaderParser::GetVarType(
 				{
 					if (bSeenEditSpecifier)
 					{
-						UE_LOG_ERROR_UHT(*this, TEXT("Found more than one edit/visibility specifier (%s), only one is allowed"), *Specifier.Key);
+						LogError(TEXT("Found more than one edit/visibility specifier (%s), only one is allowed"), *Specifier.Key);
 					}
 					Flags |= CPF_Edit | CPF_EditConst | CPF_DisableEditOnInstance;
 					bSeenEditSpecifier = true;
@@ -2978,19 +2978,19 @@ void FHeaderParser::GetVarType(
 				{
 					if (bSeenBlueprintReadOnlySpecifier)
 					{
-						UE_LOG_ERROR_UHT(*this, TEXT("Cannot specify a property as being both BlueprintReadOnly and BlueprintReadWrite."));
+						LogError(TEXT("Cannot specify a property as being both BlueprintReadOnly and BlueprintReadWrite."));
 					}
 
 					const FString* PrivateAccessMD = MetaDataFromNewStyle.Find(NAME_AllowPrivateAccess);  // FBlueprintMetadata::MD_AllowPrivateAccess
 					const bool bAllowPrivateAccess = PrivateAccessMD ? (*PrivateAccessMD != TEXT("false")) : false;
 					if (CurrentAccessSpecifier == ACCESS_Private && !bAllowPrivateAccess)
 					{
-						UE_LOG_ERROR_UHT(*this, TEXT("BlueprintReadWrite should not be used on private members"));
+						LogError(TEXT("BlueprintReadWrite should not be used on private members"));
 					}
 
 					if ((Flags & CPF_EditorOnly) != 0 && OwnerScriptStructDef != nullptr)
 					{
-						UE_LOG_ERROR_UHT(*this, TEXT("Blueprint exposed struct members cannot be editor only"));
+						LogError(TEXT("Blueprint exposed struct members cannot be editor only"));
 					}
 
 					Flags |= CPF_BlueprintVisible;
@@ -3002,12 +3002,12 @@ void FHeaderParser::GetVarType(
 				{
 					if (bSeenBlueprintReadOnlySpecifier)
 					{
-						UE_LOG_ERROR_UHT(*this, TEXT("Cannot specify a property as being both BlueprintReadOnly and having a BlueprintSetter."));
+						LogError(TEXT("Cannot specify a property as being both BlueprintReadOnly and having a BlueprintSetter."));
 					}
 
 					if (OwnerScriptStructDef != nullptr)
 					{
-						UE_LOG_ERROR_UHT(*this, TEXT("Cannot specify BlueprintSetter for a struct member."))
+						LogError(TEXT("Cannot specify BlueprintSetter for a struct member."));
 					}
 
 					const FString BlueprintSetterFunction = RequireExactlyOneSpecifierValue(*this, Specifier);
@@ -3022,19 +3022,19 @@ void FHeaderParser::GetVarType(
 				{
 					if (bSeenBlueprintWriteSpecifier)
 					{
-						UE_LOG_ERROR_UHT(*this, TEXT("Cannot specify both BlueprintReadOnly and BlueprintReadWrite or BlueprintSetter."), *Specifier.Key);
+						LogError(TEXT("Cannot specify both BlueprintReadOnly and BlueprintReadWrite or BlueprintSetter."), *Specifier.Key);
 					}
 
 					const FString* PrivateAccessMD = MetaDataFromNewStyle.Find(NAME_AllowPrivateAccess);  // FBlueprintMetadata::MD_AllowPrivateAccess
 					const bool bAllowPrivateAccess = PrivateAccessMD ? (*PrivateAccessMD != TEXT("false")) : false;
 					if (CurrentAccessSpecifier == ACCESS_Private && !bAllowPrivateAccess)
 					{
-						UE_LOG_ERROR_UHT(*this, TEXT("BlueprintReadOnly should not be used on private members"));
+						LogError(TEXT("BlueprintReadOnly should not be used on private members"));
 					}
 
 					if ((Flags & CPF_EditorOnly) != 0 && OwnerScriptStructDef != nullptr)
 					{
-						UE_LOG_ERROR_UHT(*this, TEXT("Blueprint exposed struct members cannot be editor only"));
+						LogError(TEXT("Blueprint exposed struct members cannot be editor only"));
 					}
 
 					Flags        |= CPF_BlueprintVisible | CPF_BlueprintReadOnly;
@@ -3047,7 +3047,7 @@ void FHeaderParser::GetVarType(
 				{
 					if (OwnerScriptStructDef != nullptr)
 					{
-						UE_LOG_ERROR_UHT(*this, TEXT("Cannot specify BlueprintGetter for a struct member."))
+						LogError(TEXT("Cannot specify BlueprintGetter for a struct member."));
 					}
 
 					const FString BlueprintGetterFunction = RequireExactlyOneSpecifierValue(*this, Specifier);
@@ -3072,7 +3072,7 @@ void FHeaderParser::GetVarType(
 
 				case EVariableSpecifier::Localized:
 				{
-					UE_LOG_ERROR_UHT(*this, TEXT("The Localized specifier is deprecated"));
+					LogError(TEXT("The Localized specifier is deprecated"));
 				}
 				break;
 
@@ -3096,7 +3096,7 @@ void FHeaderParser::GetVarType(
 
 				case EVariableSpecifier::NonPIETransient:
 				{
-					UE_LOG_WARNING_UHT(*this, TEXT("NonPIETransient is deprecated - NonPIEDuplicateTransient should be used instead"));
+					LogWarning(TEXT("NonPIETransient is deprecated - NonPIEDuplicateTransient should be used instead"));
 					Flags |= CPF_NonPIEDuplicateTransient;
 				}
 				break;
@@ -3115,7 +3115,7 @@ void FHeaderParser::GetVarType(
 
 				case EVariableSpecifier::EditInline:
 				{
-					UE_LOG_ERROR_UHT(*this, TEXT("EditInline is deprecated. Remove it, or use Instanced instead."));
+					LogError(TEXT("EditInline is deprecated. Remove it, or use Instanced instead."));
 				}
 				break;
 
@@ -3136,7 +3136,7 @@ void FHeaderParser::GetVarType(
 				{
 					if (OwnerScriptStructDef != nullptr)
 					{
-						UE_LOG_ERROR_UHT(*this, TEXT("Struct members cannot be replicated"));
+						LogError(TEXT("Struct members cannot be replicated"));
 					}
 
 					Flags |= CPF_Net;
@@ -3154,7 +3154,7 @@ void FHeaderParser::GetVarType(
 				{
 					if (OwnerScriptStructDef == nullptr)
 					{
-						UE_LOG_ERROR_UHT(*this, TEXT("Only Struct members can be marked NotReplicated"));
+						LogError(TEXT("Only Struct members can be marked NotReplicated"));
 					}
 
 					Flags |= CPF_RepSkip;
@@ -3163,7 +3163,7 @@ void FHeaderParser::GetVarType(
 
 				case EVariableSpecifier::RepRetry:
 				{
-					UE_LOG_ERROR_UHT(*this, TEXT("'RepRetry' is deprecated."));
+					LogError(TEXT("'RepRetry' is deprecated."));
 				}
 				break;
 
@@ -3238,7 +3238,7 @@ void FHeaderParser::GetVarType(
 
 				default:
 				{
-					UE_LOG_ERROR_UHT(*this, TEXT("Unknown variable specifier '%s'"), *Specifier.Key);
+					LogError(TEXT("Unknown variable specifier '%s'"), *Specifier.Key);
 				}
 				break;
 			}
@@ -3268,14 +3268,14 @@ void FHeaderParser::GetVarType(
 					}
 					else
 					{
-						UE_LOG_ERROR_UHT(*this, TEXT("Only parameters in service request functions can be marked NotReplicated"));
+						LogError(TEXT("Only parameters in service request functions can be marked NotReplicated"));
 					}
 				}
 				break;
 
 				default:
 				{
-					UE_LOG_ERROR_UHT(*this, TEXT("Unknown variable specifier '%s'"), *Specifier.Key);
+					LogError(TEXT("Unknown variable specifier '%s'"), *Specifier.Key);
 				}
 				break;
 			}
@@ -3297,11 +3297,11 @@ void FHeaderParser::GetVarType(
 		{
 			if (0 != (CPF_DisableEditOnInstance & Flags))
 			{
-				UE_LOG_WARNING_UHT(*this, TEXT("Property cannot have both 'DisableEditOnInstance' and 'ExposeOnSpawn' flags"));
+				LogWarning(TEXT("Property cannot have both 'DisableEditOnInstance' and 'ExposeOnSpawn' flags"));
 			}
 			if (0 == (CPF_BlueprintVisible & Flags))
 			{
-				UE_LOG_WARNING_UHT(*this, TEXT("Property cannot have 'ExposeOnSpawn' without 'BlueprintVisible' flag."));
+				LogWarning(TEXT("Property cannot have 'ExposeOnSpawn' without 'BlueprintVisible' flag."));
 			}
 			Flags |= CPF_ExposeOnSpawn;
 		}
@@ -3336,7 +3336,7 @@ void FHeaderParser::GetVarType(
 	}
 	else
 	{
-		FUHTException::Throwf(*this, TEXT("Unknown access level"));
+		Throwf(TEXT("Unknown access level"));
 	}
 
 	// Swallow inline keywords
@@ -3345,7 +3345,7 @@ void FHeaderParser::GetVarType(
 		FToken InlineToken;
 		if (!GetIdentifier(InlineToken, true))
 		{
-			FUHTException::Throwf(*this, TEXT("%s: Missing variable type"), GetHintText(*this, VariableCategory));
+			Throwf(TEXT("%s: Missing variable type"), GetHintText(*this, VariableCategory));
 		}
 
 		if (FCString::Strcmp(InlineToken.Identifier, TEXT("inline")) != 0
@@ -3430,7 +3430,7 @@ void FHeaderParser::GetVarType(
 	{
 		if (VariableCategory == EVariableCategory::Member)
 		{
-			FUHTException::Throwf(*this, TEXT("%s: Cannot declare enum at variable declaration"), GetHintText(*this, VariableCategory));
+			Throwf(TEXT("%s: Cannot declare enum at variable declaration"), GetHintText(*this, VariableCategory));
 		}
 
 		bUnconsumedEnumKeyword = true;
@@ -3440,7 +3440,7 @@ void FHeaderParser::GetVarType(
 	FToken VarType;
 	if ( !GetIdentifier(VarType,1) )
 	{
-		FUHTException::Throwf(*this, TEXT("%s: Missing variable type"), GetHintText(*this, VariableCategory));
+		Throwf(TEXT("%s: Missing variable type"), GetHintText(*this, VariableCategory));
 	}
 
 	RedirectTypeIdentifier(VarType);
@@ -3503,7 +3503,7 @@ void FHeaderParser::GetVarType(
 	{
 		if (IsBitfieldProperty(LayoutMacroType))
 		{
-			UE_LOG_ERROR_UHT(*this, TEXT("bool bitfields are not supported."));
+			LogError(TEXT("bool bitfields are not supported."));
 		}
 		// C++ bool type
 		VarProperty = FPropertyBase(CPT_Bool);
@@ -3539,7 +3539,7 @@ void FHeaderParser::GetVarType(
 	{
 		if (!SourceFile.GetFilename().EndsWith(TEXT("NoExportTypes.h")))
 		{
-			FUHTException::Throwf(*this, TEXT("FLargeWorldCoordinatesReal is intended for LWC support only and should not be used outside of NoExportTypes.h"));
+			Throwf(TEXT("FLargeWorldCoordinatesReal is intended for LWC support only and should not be used outside of NoExportTypes.h"));
 		}
 		VarProperty = FPropertyBase(CPT_FLargeWorldCoordinatesReal);
 	}
@@ -3557,7 +3557,7 @@ void FHeaderParser::GetVarType(
 		GetVarType(Scope, VarProperty, Disallow, &VarType, &VarTypePropertyFlags, EPropertyDeclarationStyle::None, VariableCategory);
 		if (VarProperty.IsContainer())
 		{
-			FUHTException::Throwf(*this, TEXT("Nested containers are not supported.") );
+			Throwf(TEXT("Nested containers are not supported.") );
 		}
 		// TODO: Prevent sparse delegate types from being used in a container
 
@@ -3572,7 +3572,7 @@ void FHeaderParser::GetVarType(
 		FToken CloseTemplateToken;
 		if (!GetToken(CloseTemplateToken, /*bNoConsts=*/ true, ESymbolParseOption::CloseTemplateBracket))
 		{
-			FUHTException::Throwf(*this, TEXT("Missing token while parsing TArray."));
+			Throwf(TEXT("Missing token while parsing TArray."));
 		}
 
 		if (CloseTemplateToken.TokenType != TOKEN_Symbol || !CloseTemplateToken.Matches(TEXT('>')))
@@ -3580,26 +3580,26 @@ void FHeaderParser::GetVarType(
 			// If we didn't find a comma, report it
 			if (!CloseTemplateToken.Matches(TEXT(',')))
 			{
-				FUHTException::Throwf(*this, TEXT("Expected '>' but found '%s'"), CloseTemplateToken.Identifier);
+				Throwf(TEXT("Expected '>' but found '%s'"), CloseTemplateToken.Identifier);
 			}
 
 			// If we found a comma, read the next thing, assume it's an allocator, and report that
 			FToken AllocatorToken;
 			if (!GetToken(AllocatorToken, /*bNoConsts=*/ true, ESymbolParseOption::CloseTemplateBracket))
 			{
-				FUHTException::Throwf(*this, TEXT("Unexpected end of file when parsing TArray allocator."));
+				Throwf(TEXT("Unexpected end of file when parsing TArray allocator."));
 			}
 
 			if (AllocatorToken.TokenType != TOKEN_Identifier)
 			{
-				FUHTException::Throwf(*this, TEXT("Found '%s' - expected a '>' or ','."), AllocatorToken.Identifier);
+				Throwf(TEXT("Found '%s' - expected a '>' or ','."), AllocatorToken.Identifier);
 			}
 
 			if (FCString::Strcmp(AllocatorToken.Identifier, TEXT("FMemoryImageAllocator")) == 0)
 			{
 				if (EnumHasAnyFlags(Flags, CPF_Net))
 				{
-					FUHTException::Throwf(*this, TEXT("Replicated arrays with MemoryImageAllocators are not yet supported"));
+					Throwf(TEXT("Replicated arrays with MemoryImageAllocators are not yet supported"));
 				}
 
 				RequireSymbol(TEXT('>'), TEXT("TArray template arguments"), ESymbolParseOption::CloseTemplateBracket);
@@ -3610,7 +3610,7 @@ void FHeaderParser::GetVarType(
 			{
 				if (EnumHasAnyFlags(Flags, CPF_Net))
 				{
-					FUHTException::Throwf(*this, TEXT("Replicated arrays with MemoryImageAllocators are not yet supported"));
+					Throwf(TEXT("Replicated arrays with MemoryImageAllocators are not yet supported"));
 				}
 
 				RequireSymbol(TEXT('<'), TEXT("TMemoryImageAllocator template arguments"));
@@ -3620,7 +3620,7 @@ void FHeaderParser::GetVarType(
 				{
 					if (!GetToken(SkipToken, /*bNoConsts=*/ false, ESymbolParseOption::CloseTemplateBracket))
 					{
-						FUHTException::Throwf(*this, TEXT("Unexpected end of file when parsing TMemoryImageAllocator template arguments."));
+						Throwf(TEXT("Unexpected end of file when parsing TMemoryImageAllocator template arguments."));
 					}
 
 					if (SkipToken.TokenType == TOKEN_Symbol && FCString::Strcmp(SkipToken.Identifier, TEXT(">")) == 0)
@@ -3633,7 +3633,7 @@ void FHeaderParser::GetVarType(
 			}
 			else
 			{
-				FUHTException::Throwf(*this, TEXT("Found '%s' - explicit allocators are not supported in TArray properties."), AllocatorToken.Identifier);
+				Throwf(TEXT("Found '%s' - explicit allocators are not supported in TArray properties."), AllocatorToken.Identifier);
 			}
 		}
 	}
@@ -3647,35 +3647,35 @@ void FHeaderParser::GetVarType(
 		GetVarType(Scope, MapKeyType, Disallow, &VarType, &VarTypePropertyFlags, EPropertyDeclarationStyle::None, VariableCategory);
 		if (MapKeyType.IsContainer())
 		{
-			FUHTException::Throwf(*this, TEXT("Nested containers are not supported.") );
+			Throwf(TEXT("Nested containers are not supported.") );
 		}
 		// TODO: Prevent sparse delegate types from being used in a container
 
 		if (MapKeyType.Type == CPT_Interface)
 		{
-			FUHTException::Throwf(*this, TEXT("UINTERFACEs are not currently supported as key types."));
+			Throwf(TEXT("UINTERFACEs are not currently supported as key types."));
 		}
 
 		if (MapKeyType.Type == CPT_Text)
 		{
-			FUHTException::Throwf(*this, TEXT("FText is not currently supported as a key type."));
+			Throwf(TEXT("FText is not currently supported as a key type."));
 		}
 		
 		if (EnumHasAnyFlags(Flags, CPF_Net))
 		{
-			UE_LOG_ERROR_UHT(*this, TEXT("Replicated maps are not supported."));
+			LogError(TEXT("Replicated maps are not supported."));
 		}
 
 		FToken CommaToken;
 		if (!GetToken(CommaToken, /*bNoConsts=*/ true) || CommaToken.TokenType != TOKEN_Symbol || !CommaToken.Matches(TEXT(',')))
 		{
-			FUHTException::Throwf(*this, TEXT("Missing value type while parsing TMap."));
+			Throwf(TEXT("Missing value type while parsing TMap."));
 		}
 
 		GetVarType(Scope, VarProperty, Disallow, &VarType, &VarTypePropertyFlags, EPropertyDeclarationStyle::None, VariableCategory);
 		if (VarProperty.IsContainer())
 		{
-			FUHTException::Throwf(*this, TEXT("Nested containers are not supported.") );
+			Throwf(TEXT("Nested containers are not supported.") );
 		}
 		// TODO: Prevent sparse delegate types from being used in a container
 
@@ -3687,7 +3687,7 @@ void FHeaderParser::GetVarType(
 		FToken CloseTemplateToken;
 		if (!GetToken(CloseTemplateToken, /*bNoConsts=*/ true, ESymbolParseOption::CloseTemplateBracket))
 		{
-			FUHTException::Throwf(*this, TEXT("Missing token while parsing TMap."));
+			Throwf(TEXT("Missing token while parsing TMap."));
 		}
 
 		if (CloseTemplateToken.TokenType != TOKEN_Symbol || !CloseTemplateToken.Matches(TEXT('>')))
@@ -3695,19 +3695,19 @@ void FHeaderParser::GetVarType(
 			// If we didn't find a comma, report it
 			if (!CloseTemplateToken.Matches(TEXT(',')))
 			{
-				FUHTException::Throwf(*this, TEXT("Expected '>' but found '%s'"), CloseTemplateToken.Identifier);
+				Throwf(TEXT("Expected '>' but found '%s'"), CloseTemplateToken.Identifier);
 			}
 
 			// If we found a comma, read the next thing, assume it's an allocator, and report that
 			FToken AllocatorToken;
 			if (!GetToken(AllocatorToken, /*bNoConsts=*/ true, ESymbolParseOption::CloseTemplateBracket))
 			{
-				FUHTException::Throwf(*this, TEXT("Unexpected end of file when parsing TArray allocator."));
+				Throwf(TEXT("Unexpected end of file when parsing TArray allocator."));
 			}
 
 			if (AllocatorToken.TokenType != TOKEN_Identifier)
 			{
-				FUHTException::Throwf(*this, TEXT("Found '%s' - expected a '>' or ','."), AllocatorToken.Identifier);
+				Throwf(TEXT("Found '%s' - expected a '>' or ','."), AllocatorToken.Identifier);
 			}
 
 			if (FCString::Strcmp(AllocatorToken.Identifier, TEXT("FMemoryImageSetAllocator")) == 0)
@@ -3718,7 +3718,7 @@ void FHeaderParser::GetVarType(
 			}
 			else
 			{
-				FUHTException::Throwf(*this, TEXT("Found '%s' - explicit allocators are not supported in TMap properties."), AllocatorToken.Identifier);
+				Throwf(TEXT("Found '%s' - explicit allocators are not supported in TMap properties."), AllocatorToken.Identifier);
 			}
 		}
 	}
@@ -3731,23 +3731,23 @@ void FHeaderParser::GetVarType(
 		GetVarType(Scope, VarProperty, Disallow, &VarType, &VarTypePropertyFlags, EPropertyDeclarationStyle::None, VariableCategory);
 		if (VarProperty.IsContainer())
 		{
-			FUHTException::Throwf(*this, TEXT("Nested containers are not supported.") );
+			Throwf(TEXT("Nested containers are not supported.") );
 		}
 		// TODO: Prevent sparse delegate types from being used in a container
 
 		if (VarProperty.Type == CPT_Interface)
 		{
-			FUHTException::Throwf(*this, TEXT("UINTERFACEs are not currently supported as element types."));
+			Throwf(TEXT("UINTERFACEs are not currently supported as element types."));
 		}
 
 		if (VarProperty.Type == CPT_Text)
 		{
-			FUHTException::Throwf(*this, TEXT("FText is not currently supported as an element type."));
+			Throwf(TEXT("FText is not currently supported as an element type."));
 		}
 
 		if (EnumHasAnyFlags(Flags, CPF_Net))
 		{
-			UE_LOG_ERROR_UHT(*this, TEXT("Replicated sets are not supported."));
+			LogError(TEXT("Replicated sets are not supported."));
 		}
 
 		VarTypePropertyFlags = VarProperty.PropertyFlags & (CPF_ContainsInstancedReference | CPF_InstancedReference); // propagate these to the set, we will fix them later
@@ -3756,7 +3756,7 @@ void FHeaderParser::GetVarType(
 		FToken CloseTemplateToken;
 		if (!GetToken(CloseTemplateToken, /*bNoConsts=*/ true, ESymbolParseOption::CloseTemplateBracket))
 		{
-			FUHTException::Throwf(*this, TEXT("Missing token while parsing TArray."));
+			Throwf(TEXT("Missing token while parsing TArray."));
 		}
 
 		if (CloseTemplateToken.TokenType != TOKEN_Symbol || !CloseTemplateToken.Matches(TEXT('>')))
@@ -3764,17 +3764,17 @@ void FHeaderParser::GetVarType(
 			// If we didn't find a comma, report it
 			if (!CloseTemplateToken.Matches(TEXT(',')))
 			{
-				FUHTException::Throwf(*this, TEXT("Expected '>' but found '%s'"), CloseTemplateToken.Identifier);
+				Throwf(TEXT("Expected '>' but found '%s'"), CloseTemplateToken.Identifier);
 			}
 
 			// If we found a comma, read the next thing, assume it's a keyfuncs, and report that
 			FToken AllocatorToken;
 			if (!GetToken(AllocatorToken, /*bNoConsts=*/ true, ESymbolParseOption::CloseTemplateBracket))
 			{
-				FUHTException::Throwf(*this, TEXT("Expected '>' but found '%s'"), CloseTemplateToken.Identifier);
+				Throwf(TEXT("Expected '>' but found '%s'"), CloseTemplateToken.Identifier);
 			}
 
-			FUHTException::Throwf(*this, TEXT("Found '%s' - explicit KeyFuncs are not supported in TSet properties."), AllocatorToken.Identifier);
+			Throwf(TEXT("Found '%s' - explicit KeyFuncs are not supported in TSet properties."), AllocatorToken.Identifier);
 		}
 	}
 	else if ( VarType.Matches(TEXT("FString"), ESearchCase::CaseSensitive) || VarType.Matches(TEXT("FMemoryImageString"), ESearchCase::CaseSensitive))
@@ -3806,7 +3806,7 @@ void FHeaderParser::GetVarType(
 	}
 	else if ( VarType.Matches(TEXT("Text"), ESearchCase::IgnoreCase) )
 	{
-		FUHTException::Throwf(*this, TEXT("%s' is missing a prefix, expecting 'FText'"), VarType.Identifier);
+		Throwf(TEXT("%s' is missing a prefix, expecting 'FText'"), VarType.Identifier);
 	}
 	else if ( VarType.Matches(TEXT("FText"), ESearchCase::CaseSensitive) )
 	{
@@ -3839,13 +3839,13 @@ void FHeaderParser::GetVarType(
 			FToken ScopedTrueEnumName;
 			if (!GetIdentifier(ScopedTrueEnumName, true))
 			{
-				FUHTException::Throwf(*this, TEXT("Expected a namespace scoped enum name.") );
+				Throwf(TEXT("Expected a namespace scoped enum name.") );
 			}
 		}
 
 		if (!bFoundEnum)
 		{
-			FUHTException::Throwf(*this, TEXT("Expected the name of a previously defined enum"));
+			Throwf(TEXT("Expected the name of a previously defined enum"));
 		}
 
 		RequireSymbol(TEXT('>'), VarType.Identifier, ESymbolParseOption::CloseTemplateBracket);
@@ -3858,14 +3858,14 @@ void FHeaderParser::GetVarType(
 		FToken PropertyTypeToken;
 		if (!GetToken(PropertyTypeToken, /* bNoConsts = */ true))
 		{
-			FUHTException::Throwf(*this, TEXT("Expected the property type"));
+			Throwf(TEXT("Expected the property type"));
 		}
 		else
 		{
 			FieldClassName = FName(PropertyTypeToken.Identifier + 1, FNAME_Add);
 			if (!FPropertyTraits::IsValidFieldClass(FieldClassName))
 			{
-				FUHTException::Throwf(*this, TEXT("Undefined property type: %s"), PropertyTypeToken.Identifier);
+				Throwf(TEXT("Undefined property type: %s"), PropertyTypeToken.Identifier);
 			}
 		}
 
@@ -3881,7 +3881,7 @@ void FHeaderParser::GetVarType(
 		{
 			if (EnumDef->GetCppForm() != UEnum::ECppForm::EnumClass)
 			{
-				FUHTException::Throwf(*this, TEXT("You cannot use the raw enum name as a type for member variables, instead use TEnumAsByte or a C++11 enum class with an explicit underlying type."), *EnumDef->GetCppType());
+				Throwf(TEXT("You cannot use the raw enum name as a type for member variables, instead use TEnumAsByte or a C++11 enum class with an explicit underlying type."), *EnumDef->GetCppType());
 			}
 		}
 
@@ -3892,7 +3892,7 @@ void FHeaderParser::GetVarType(
 			FToken ScopedTrueEnumName;
 			if (!GetIdentifier(ScopedTrueEnumName, true))
 			{
-				FUHTException::Throwf(*this, TEXT("Expected a namespace scoped enum name.") );
+				Throwf(TEXT("Expected a namespace scoped enum name.") );
 			}
 		}
 
@@ -3946,7 +3946,7 @@ void FHeaderParser::GetVarType(
 				}
 				else
 				{
-					FUHTException::Throwf(*this, TEXT("Cannot find class '%s', to resolve delegate '%s'"), *IdentifierStripped, DelegateName.Identifier);
+					Throwf(TEXT("Cannot find class '%s', to resolve delegate '%s'"), *IdentifierStripped, DelegateName.Identifier);
 				}
 			}
 		}
@@ -3962,13 +3962,13 @@ void FHeaderParser::GetVarType(
 				FString ExpectedStructName = FString::Printf(TEXT("%s%s"), PrefixCPP, *StructDef->GetName());
 				if (FString(VarType.Identifier) != ExpectedStructName)
 				{
-					FUHTException::Throwf(*this, TEXT("Struct '%s' is missing or has an incorrect prefix, expecting '%s'"), VarType.Identifier, *ExpectedStructName);
+					Throwf(TEXT("Struct '%s' is missing or has an incorrect prefix, expecting '%s'"), VarType.Identifier, *ExpectedStructName);
 				}
 			}
 			else if (!UHTConfig.StructsWithNoPrefix.Contains(VarType.Identifier))
 			{
 				const TCHAR* PrefixCPP = UHTConfig.StructsWithTPrefix.Contains(VarType.Identifier) ? TEXT("T") : StructDef->GetPrefixCPP();
-				FUHTException::Throwf(*this, TEXT("Struct '%s' is missing a prefix, expecting '%s'"), VarType.Identifier, *FString::Printf(TEXT("%s%s"), PrefixCPP, *StructDef->GetName()));
+				Throwf(TEXT("Struct '%s' is missing a prefix, expecting '%s'"), VarType.Identifier, *FString::Printf(TEXT("%s%s"), PrefixCPP, *StructDef->GetName()));
 			}
 
 			bHandledType = true;
@@ -4047,7 +4047,7 @@ void FHeaderParser::GetVarType(
 							TempClassDef = FUnrealClassDefinitionInfo::FindScriptClass(InnerClass.Identifier);
 							if (TempClassDef == nullptr)
 							{
-								FUHTException::Throwf(*this, TEXT("Unrecognized type '%s' (in expression %s<%s>) - type must be a UCLASS"), InnerClass.Identifier, VarType.Identifier, InnerClass.Identifier);
+								Throwf(TEXT("Unrecognized type '%s' (in expression %s<%s>) - type must be a UCLASS"), InnerClass.Identifier, VarType.Identifier, InnerClass.Identifier);
 							}
 
 							if (bIsAutoweakPtrTemplate)
@@ -4076,7 +4076,7 @@ void FHeaderParser::GetVarType(
 						}
 						else
 						{
-							FUHTException::Throwf(*this, TEXT("%s: Missing template type"), VarType.Identifier);
+							Throwf(TEXT("%s: Missing template type"), VarType.Identifier);
 						}
 
 						// Const after template argument type but before end of template symbol
@@ -4119,7 +4119,7 @@ void FHeaderParser::GetVarType(
 							FToken Limitor;
 							if (!GetIdentifier(Limitor))
 							{
-								FUHTException::Throwf(*this, TEXT("'class': Missing class limitor"));
+								Throwf(TEXT("'class': Missing class limitor"));
 							}
 
 							RedirectTypeIdentifier(Limitor);
@@ -4135,16 +4135,16 @@ void FHeaderParser::GetVarType(
 
 						if (PropertyType == CPT_WeakObjectReference)
 						{
-							FUHTException::Throwf(*this, TEXT("Class variables cannot be weak, they are always strong."));
+							Throwf(TEXT("Class variables cannot be weak, they are always strong."));
 						}
 						else if (PropertyType == CPT_LazyObjectReference)
 						{
-							FUHTException::Throwf(*this, TEXT("Class variables cannot be lazy, they are always strong."));
+							Throwf(TEXT("Class variables cannot be lazy, they are always strong."));
 						}
 
 						if (bIsSoftObjectPtrTemplate)
 						{
-							FUHTException::Throwf(*this, TEXT("Class variables cannot be stored in TSoftObjectPtr, use TSoftClassPtr instead."));
+							Throwf(TEXT("Class variables cannot be stored in TSoftObjectPtr, use TSoftClassPtr instead."));
 						}
 					}
 
@@ -4199,7 +4199,7 @@ void FHeaderParser::GetVarType(
 
 				if (!bHandledType)
 				{
-					FUHTException::Throwf(*this, TEXT("Unrecognized type '%s' - type must be a UCLASS, USTRUCT or UENUM"), VarType.Identifier);
+					Throwf(TEXT("Unrecognized type '%s' - type must be a UCLASS, USTRUCT or UENUM"), VarType.Identifier);
 				}
 			}
 		}
@@ -4219,32 +4219,32 @@ void FHeaderParser::GetVarType(
 	{
 		if (VariableCategory == EVariableCategory::Member)
 		{
-			FUHTException::Throwf(*this, TEXT("Const properties are not supported."));
+			Throwf(TEXT("Const properties are not supported."));
 		}
 		else
 		{
-			FUHTException::Throwf(*this, TEXT("Inappropriate keyword 'const' on variable of type '%s'"), VarType.Identifier);
+			Throwf(TEXT("Inappropriate keyword 'const' on variable of type '%s'"), VarType.Identifier);
 		}
 	}
 
 	if (bUnconsumedClassKeyword)
 	{
-		FUHTException::Throwf(*this, TEXT("Inappropriate keyword 'class' on variable of type '%s'"), VarType.Identifier );
+		Throwf(TEXT("Inappropriate keyword 'class' on variable of type '%s'"), VarType.Identifier );
 	}
 
 	if (bUnconsumedStructKeyword)
 	{
-		FUHTException::Throwf(*this, TEXT("Inappropriate keyword 'struct' on variable of type '%s'"), VarType.Identifier );
+		Throwf(TEXT("Inappropriate keyword 'struct' on variable of type '%s'"), VarType.Identifier );
 	}
 
 	if (bUnconsumedEnumKeyword)
 	{
-		FUHTException::Throwf(*this, TEXT("Inappropriate keyword 'enum' on variable of type '%s'"), VarType.Identifier );
+		Throwf(TEXT("Inappropriate keyword 'enum' on variable of type '%s'"), VarType.Identifier );
 	}
 
 	if (MatchSymbol(TEXT('*')))
 	{
-		FUHTException::Throwf(*this, TEXT("Inappropriate '*' on variable of type '%s', cannot have an exposed pointer to this type."), VarType.Identifier );
+		Throwf(TEXT("Inappropriate '*' on variable of type '%s', cannot have an exposed pointer to this type."), VarType.Identifier );
 	}
 
 	//@TODO: UCREMOVAL: 'const' member variables that will get written post-construction by defaultproperties
@@ -4276,7 +4276,7 @@ void FHeaderParser::GetVarType(
 			{
 				if (!(Flags & CPF_ConstParm))
 				{
-					FUHTException::Throwf(*this, TEXT("Replicated %s parameters cannot be passed by non-const reference"), VarType.Identifier);
+					Throwf(TEXT("Replicated %s parameters cannot be passed by non-const reference"), VarType.Identifier);
 				}
 
 				Flags |= CPF_ReferenceParm;
@@ -4314,7 +4314,7 @@ void FHeaderParser::GetVarType(
 		}
 		else
 		{
-			FUHTException::Throwf(*this, TEXT("Must specify a valid function name for replication notifications"));
+			Throwf(TEXT("Must specify a valid function name for replication notifications"));
 		}
 	}
 
@@ -4325,33 +4325,33 @@ void FHeaderParser::GetVarType(
 		{
 			if (VarProperty.ClassDef->IsChildOf(*GUClassDef))
 			{
-				FUHTException::Throwf(*this, TEXT("'Instanced' cannot be applied to class properties (UClass* or TSubclassOf<>)"));
+				Throwf(TEXT("'Instanced' cannot be applied to class properties (UClass* or TSubclassOf<>)"));
 			}
 		}
 		else
 		{
-			FUHTException::Throwf(*this, TEXT("'Instanced' is only allowed on an object property, an array of objects, a set of objects, or a map with an object value type."));
+			Throwf(TEXT("'Instanced' is only allowed on an object property, an array of objects, a set of objects, or a map with an object value type."));
 		}
 	}
 
 	if ( VarProperty.IsObjectOrInterface() && VarProperty.Type != CPT_SoftObjectReference && VarProperty.MetaClassDef == nullptr && (VarProperty.PropertyFlags&CPF_Config) != 0 )
 	{
-		FUHTException::Throwf(*this, TEXT("Not allowed to use 'config' with object variables"));
+		Throwf(TEXT("Not allowed to use 'config' with object variables"));
 	}
 
 	if ((VarProperty.PropertyFlags & CPF_BlueprintAssignable) && VarProperty.Type != CPT_MulticastDelegate)
 	{
-		FUHTException::Throwf(*this, TEXT("'BlueprintAssignable' is only allowed on multicast delegate properties"));
+		Throwf(TEXT("'BlueprintAssignable' is only allowed on multicast delegate properties"));
 	}
 
 	if ((VarProperty.PropertyFlags & CPF_BlueprintCallable) && VarProperty.Type != CPT_MulticastDelegate)
 	{
-		FUHTException::Throwf(*this, TEXT("'BlueprintCallable' is only allowed on a property when it is a multicast delegate"));
+		Throwf(TEXT("'BlueprintCallable' is only allowed on a property when it is a multicast delegate"));
 	}
 
 	if ((VarProperty.PropertyFlags & CPF_BlueprintAuthorityOnly) && VarProperty.Type != CPT_MulticastDelegate)
 	{
-		FUHTException::Throwf(*this, TEXT("'BlueprintAuthorityOnly' is only allowed on a property when it is a multicast delegate"));
+		Throwf(TEXT("'BlueprintAuthorityOnly' is only allowed on a property when it is a multicast delegate"));
 	}
 	
 	if (VariableCategory != EVariableCategory::Member)
@@ -4365,13 +4365,13 @@ void FHeaderParser::GetVarType(
 	if (Transients && OwnerClassDef == nullptr)
 	{
 		TArray<const TCHAR*> FlagStrs = ParsePropertyFlags(Transients);
-		FUHTException::Throwf(*this, TEXT("'%s' specifier(s) are only allowed on class member variables"), *FString::Join(FlagStrs, TEXT(", ")));
+		Throwf(TEXT("'%s' specifier(s) are only allowed on class member variables"), *FString::Join(FlagStrs, TEXT(", ")));
 	}
 
 	// Make sure the overrides are allowed here.
 	if( VarProperty.PropertyFlags & Disallow )
 	{
-		FUHTException::Throwf(*this, TEXT("Specified type modifiers not allowed here") );
+		Throwf(TEXT("Specified type modifiers not allowed here") );
 	}
 
 	// For now, copy the flags that a TMap value has to the key
@@ -4421,7 +4421,7 @@ FUnrealPropertyDefinitionInfo& FHeaderParser::GetVarNameAndDim
 		FToken VarToken;
 		if (!GetIdentifier(VarToken))
 		{
-			FUHTException::Throwf(*this, TEXT("Missing variable name") );
+			Throwf(TEXT("Missing variable name") );
 		}
 
 		switch (LayoutMacroType)
@@ -4452,7 +4452,7 @@ FUnrealPropertyDefinitionInfo& FHeaderParser::GetVarNameAndDim
 		{
 			if (DeprecatedIndex != VarName.Len() - 11)
 			{
-				FUHTException::Throwf(*this, TEXT("Deprecated variables must end with _DEPRECATED"));
+				Throwf(TEXT("Deprecated variables must end with _DEPRECATED"));
 			}
 
 			// We allow deprecated properties in blueprints that have getters and setters assigned as they may be part of a backwards compatibility path
@@ -4462,12 +4462,12 @@ FUnrealPropertyDefinitionInfo& FHeaderParser::GetVarNameAndDim
 
 			if (bWarnOnGetter)
 			{
-				UE_LOG_WARNING_UHT(*this, TEXT("%s: Deprecated property '%s' should not be marked as blueprint visible without having a BlueprintGetter"), HintText, *VarName);
+				LogWarning(TEXT("%s: Deprecated property '%s' should not be marked as blueprint visible without having a BlueprintGetter"), HintText, *VarName);
 			}
 
 			if (bWarnOnSetter)
 			{
-				UE_LOG_WARNING_UHT(*this, TEXT("%s: Deprecated property '%s' should not be marked as blueprint writeable without having a BlueprintSetter"), HintText, *VarName);
+				LogWarning(TEXT("%s: Deprecated property '%s' should not be marked as blueprint writeable without having a BlueprintSetter"), HintText, *VarName);
 			}
 
 
@@ -4475,7 +4475,7 @@ FUnrealPropertyDefinitionInfo& FHeaderParser::GetVarNameAndDim
 			if (VarProperty.PropertyFlags & (CPF_Edit | CPF_EditConst) || // Property is marked as editable
 				(!bBlueprintVisible && (VarProperty.PropertyFlags & CPF_BlueprintReadOnly) && !(VarProperty.ImpliedPropertyFlags & CPF_BlueprintReadOnly)) ) // Is BPRO, but not via Implied Flags and not caught by Getter/Setter path above
 			{
-				UE_LOG_WARNING_UHT(*this, TEXT("%s: Deprecated property '%s' should not be marked as visible or editable"), HintText, *VarName);
+				LogWarning(TEXT("%s: Deprecated property '%s' should not be marked as visible or editable"), HintText, *VarName);
 			}
 
 			VarProperty.PropertyFlags |= CPF_Deprecated;
@@ -4515,7 +4515,7 @@ FUnrealPropertyDefinitionInfo& FHeaderParser::GetVarNameAndDim
 
 		if (bErrorDueToShadowing)
 		{
-			FUHTException::Throwf(*this, TEXT("%s: '%s' cannot be defined in '%s' as it is already defined in scope '%s' (shadowing is not allowed)"),
+			Throwf(TEXT("%s: '%s' cannot be defined in '%s' as it is already defined in scope '%s' (shadowing is not allowed)"),
 				HintText,
 				Identifier,
 				*ParentStruct.GetName(),
@@ -4531,24 +4531,24 @@ FUnrealPropertyDefinitionInfo& FHeaderParser::GetVarNameAndDim
 		{
 			case EVariableCategory::Return:
 			{
-				FUHTException::Throwf(*this, TEXT("Arrays aren't allowed as return types"));
+				Throwf(TEXT("Arrays aren't allowed as return types"));
 			}
 
 			case EVariableCategory::RegularParameter:
 			case EVariableCategory::ReplicatedParameter:
 			{
-				FUHTException::Throwf(*this, TEXT("Arrays aren't allowed as function parameters"));
+				Throwf(TEXT("Arrays aren't allowed as function parameters"));
 			}
 		}
 
 		if (VarProperty.IsContainer())
 		{
-			FUHTException::Throwf(*this, TEXT("Static arrays of containers are not allowed"));
+			Throwf(TEXT("Static arrays of containers are not allowed"));
 		}
 
 		if (VarProperty.IsBool())
 		{
-			FUHTException::Throwf(*this, TEXT("Bool arrays are not allowed") );
+			Throwf(TEXT("Bool arrays are not allowed") );
 		}
 
 		if (LayoutMacroType == ELayoutMacroType::None)
@@ -4556,7 +4556,7 @@ FUnrealPropertyDefinitionInfo& FHeaderParser::GetVarNameAndDim
 			// Ignore how the actual array dimensions are actually defined - we'll calculate those with the compiler anyway.
 			if (!GetRawToken(Dimensions, TEXT(']')))
 			{
-				FUHTException::Throwf(*this, TEXT("%s %s: Missing ']'"), HintText, Identifier);
+				Throwf(TEXT("%s %s: Missing ']'"), HintText, Identifier);
 			}
 		}
 		else
@@ -4564,7 +4564,7 @@ FUnrealPropertyDefinitionInfo& FHeaderParser::GetVarNameAndDim
 			// Ignore how the actual array dimensions are actually defined - we'll calculate those with the compiler anyway.
 			if (!GetRawToken(Dimensions, TEXT(')')))
 			{
-				FUHTException::Throwf(*this, TEXT("%s %s: Missing ']'"), HintText, Identifier);
+				Throwf(TEXT("%s %s: Missing ']'"), HintText, Identifier);
 			}
 		}
 
@@ -4591,12 +4591,12 @@ FUnrealPropertyDefinitionInfo& FHeaderParser::GetVarNameAndDim
 		// WeakPointer is supported by VM as return type (see UObject::execLetWeakObjPtr), but there is no P_GET_... macro for WeakPointer.
 		if (VarProperty.Type == CPT_LazyObjectReference)
 		{
-			FUHTException::Throwf(*this, TEXT("UFunctions cannot take a lazy pointer as a parameter."));
+			Throwf(TEXT("UFunctions cannot take a lazy pointer as a parameter."));
 		}
 		else if (VarProperty.Type == CPT_ObjectPtrReference)
 		{
 			// @TODO: OBJPTR: Investigate TObjectPtr support for UFunction parameters.
-			FUHTException::Throwf(*this, TEXT("UFunctions cannot take a TObjectPtr as a parameter."));
+			Throwf(TEXT("UFunctions cannot take a TObjectPtr as a parameter."));
 		}
 	}
 
@@ -4623,7 +4623,7 @@ bool FHeaderParser::CompileDeclaration(TArray<FUnrealFunctionDefinitionInfo*>& D
 	{
 		if (!IsAllowedInThisNesting(ENestAllowFlags::VarDecl) && !IsAllowedInThisNesting(ENestAllowFlags::Function))
 		{
-			FUHTException::Throwf(*this, TEXT("Access specifier %s not allowed here."), Token.Identifier);
+			Throwf(TEXT("Access specifier %s not allowed here."), Token.Identifier);
 		}
 		check(TopNest->NestType == ENestType::Class || TopNest->NestType == ENestType::Interface || TopNest->NestType == ENestType::NativeInterface);
 		CurrentAccessSpecifier = AccessSpecifier;
@@ -4635,7 +4635,7 @@ bool FHeaderParser::CompileDeclaration(TArray<FUnrealFunctionDefinitionInfo*>& D
 		// Make sure the previous class ended with valid nesting.
 		if (bEncounteredNewStyleClass_UnmatchedBrackets)
 		{
-			FUHTException::Throwf(*this, TEXT("Missing } at end of class"));
+			Throwf(TEXT("Missing } at end of class"));
 		}
 
 		// Start parsing the second class
@@ -4655,7 +4655,7 @@ bool FHeaderParser::CompileDeclaration(TArray<FUnrealFunctionDefinitionInfo*>& D
 	{
 		if (TopNest->NestType != ENestType::NativeInterface)
 		{
-			FUHTException::Throwf(*this, TEXT("%s must occur inside the native interface definition"), Token.Identifier);
+			Throwf(TEXT("%s must occur inside the native interface definition"), Token.Identifier);
 		}
 		RequireSymbol(TEXT('('), Token.Identifier);
 		CompileVersionDeclaration(GetCurrentClassDef());
@@ -4665,13 +4665,13 @@ bool FHeaderParser::CompileDeclaration(TArray<FUnrealFunctionDefinitionInfo*>& D
 		if (StructData.ParsedInterface == EParsedInterface::NotAnInterface)
 		{
 			FString CurrentClassName = GetCurrentClassDef().GetName();
-			FUHTException::Throwf(*this, TEXT("Could not find the associated 'U%s' class while parsing 'I%s' - it could be missing or malformed"), *CurrentClassName, *CurrentClassName);
+			Throwf(TEXT("Could not find the associated 'U%s' class while parsing 'I%s' - it could be missing or malformed"), *CurrentClassName, *CurrentClassName);
 		}
 
 		if (StructData.ParsedInterface == EParsedInterface::ParsedIInterface)
 		{
 			FString CurrentClassName = GetCurrentClassDef().GetName();
-			FUHTException::Throwf(*this, TEXT("Duplicate IInterface definition found while parsing 'I%s'"), *CurrentClassName);
+			Throwf(TEXT("Duplicate IInterface definition found while parsing 'I%s'"), *CurrentClassName);
 		}
 
 		check(StructData.ParsedInterface == EParsedInterface::ParsedUInterface);
@@ -4698,7 +4698,7 @@ bool FHeaderParser::CompileDeclaration(TArray<FUnrealFunctionDefinitionInfo*>& D
 	{
 		if (TopNest->NestType != ENestType::Interface)
 		{
-			FUHTException::Throwf(*this, TEXT("%s must occur inside the interface definition"), Token.Identifier);
+			Throwf(TEXT("%s must occur inside the interface definition"), Token.Identifier);
 		}
 		RequireSymbol(TEXT('('), Token.Identifier);
 		CompileVersionDeclaration(GetCurrentClassDef());
@@ -4722,7 +4722,7 @@ bool FHeaderParser::CompileDeclaration(TArray<FUnrealFunctionDefinitionInfo*>& D
 	{
 		if (TopNest->NestType != ENestType::Class)
 		{
-			FUHTException::Throwf(*this, TEXT("%s must occur inside the class definition"), Token.Identifier);
+			Throwf(TEXT("%s must occur inside the class definition"), Token.Identifier);
 		}
 
 		FStructMetaData& StructData = GetCurrentClassData();
@@ -4836,13 +4836,13 @@ bool FHeaderParser::CompileDeclaration(TArray<FUnrealFunctionDefinitionInfo*>& D
 			// Ensure the UINTERFACE classes have a GENERATED_BODY declaration
 			if (bHaveSeenUClass && !bClassHasGeneratedUInterfaceBody)
 			{
-				FUHTException::Throwf(*this, TEXT("Expected a GENERATED_BODY() at the start of class"));
+				Throwf(TEXT("Expected a GENERATED_BODY() at the start of class"));
 			}
 
 			// Ensure the non-UINTERFACE interface classes have a GENERATED_BODY declaration
 			if (!bHaveSeenUClass && !bClassHasGeneratedIInterfaceBody)
 			{
-				FUHTException::Throwf(*this, TEXT("Expected a GENERATED_BODY() at the start of class"));
+				Throwf(TEXT("Expected a GENERATED_BODY() at the start of class"));
 			}
 		}
 		else
@@ -4853,7 +4853,7 @@ bool FHeaderParser::CompileDeclaration(TArray<FUnrealFunctionDefinitionInfo*>& D
 			// Ensure classes have a GENERATED_BODY declaration
 			if (bHaveSeenUClass && !bClassHasGeneratedBody)
 			{
-				FUHTException::Throwf(*this, TEXT("Expected a GENERATED_BODY() at the start of class"));
+				Throwf(TEXT("Expected a GENERATED_BODY() at the start of class"));
 			}
 		}
 
@@ -4870,11 +4870,11 @@ bool FHeaderParser::CompileDeclaration(TArray<FUnrealFunctionDefinitionInfo*>& D
 	{
 		if (GetToken(Token))
 		{
-			FUHTException::Throwf(*this, TEXT("Extra ';' before '%s'"), Token.Identifier);
+			Throwf(TEXT("Extra ';' before '%s'"), Token.Identifier);
 		}
 		else
 		{
-			FUHTException::Throwf(*this, TEXT("Extra ';' before end of file"));
+			Throwf(TEXT("Extra ';' before end of file"));
 		}
 	}
 
@@ -5014,7 +5014,7 @@ bool FHeaderParser::CompileDeclaration(TArray<FUnrealFunctionDefinitionInfo*>& D
 						}
 						else
 						{
-							FUHTException::Throwf(*this, TEXT("Serialize functions must not be inside preprocessor blocks, except for WITH_EDITORONLY_DATA"));
+							Throwf(TEXT("Serialize functions must not be inside preprocessor blocks, except for WITH_EDITORONLY_DATA"));
 						}
 					}
 				}
@@ -5097,7 +5097,7 @@ bool FHeaderParser::SkipDeclaration(FToken& Token)
 
 			if (NestedScopes < 0)
 			{
-				FUHTException::Throwf(*this, TEXT("Unexpected '}'. Did you miss a semi-colon?"));
+				Throwf(TEXT("Unexpected '}'. Did you miss a semi-colon?"));
 			}
 		}
 		else if (bMacroDeclaration && NestedScopes == 0)
@@ -5127,7 +5127,7 @@ bool FHeaderParser::SkipDeclaration(FToken& Token)
 			}
 			else if (!SafeMatchSymbol(TEXT(';')))
 			{
-				FUHTException::Throwf(*this, TEXT("Unexpected '%s'. Did you miss a semi-colon?"), VariableName.Identifier);
+				Throwf(TEXT("Unexpected '%s'. Did you miss a semi-colon?"), VariableName.Identifier);
 			}
 		}
 
@@ -5193,7 +5193,7 @@ FUnrealClassDefinitionInfo& FHeaderParser::ParseClassNameDeclaration(FString& De
 		{
 			if (!BaseClassDef->HasAnyClassFlags(CLASS_Interface))
 			{
-				FUHTException::Throwf(*this, TEXT("Implements: Class %s is not an interface; Can only inherit from non-UObjects or UInterface derived interfaces"), *BaseClassDef->GetName());
+				Throwf(TEXT("Implements: Class %s is not an interface; Can only inherit from non-UObjects or UInterface derived interfaces"), *BaseClassDef->GetName());
 			}
 
 			// Propagate the inheritable ClassFlags
@@ -5239,7 +5239,7 @@ void PostParsingClassSetup(FUnrealClassDefinitionInfo& ClassDef)
 		ClassDef.SetClassConfigName(ClassDef.GetSuperClass() ? ClassDef.GetSuperClass()->GetClassConfigName() : NAME_None);
 		if (ClassDef.GetClassConfigName() == NAME_None)
 		{
-			FUHTException::Throwf(ClassDef, TEXT("Classes with config / globalconfig member variables need to specify config file."));
+			ClassDef.Throwf(TEXT("Classes with config / globalconfig member variables need to specify config file."));
 			ClassDef.SetClassConfigName(NAME_Engine);
 		}
 	}
@@ -5295,7 +5295,7 @@ FUnrealClassDefinitionInfo& FHeaderParser::CompileClassDeclaration()
 		bool bIsIntrinsic = TempDef->HasAnyClassFlags(CLASS_Intrinsic);
 		if (!(bIsParsed || bIsIntrinsic))
 		{
-			FUHTException::Throwf(*this, TEXT("'%s' can't be compiled: Parent class '%s' has errors"), *ClassDef.GetName(), *TempDef->GetName());
+			Throwf(TEXT("'%s' can't be compiled: Parent class '%s' has errors"), *ClassDef.GetName(), *TempDef->GetName());
 		}
 	}
 
@@ -5361,7 +5361,7 @@ FUnrealClassDefinitionInfo* FHeaderParser::ParseInterfaceNameDeclaration(FString
 		// UInterface is special and actually extends from UObject, which isn't an interface
 		if (DeclaredInterfaceName != TEXT("UInterface"))
 		{
-			FUHTException::Throwf(*this, TEXT("Interface class '%s' cannot inherit from non-interface class '%s'"), *DeclaredInterfaceName, *TempClassDef->GetName());
+			Throwf(TEXT("Interface class '%s' cannot inherit from non-interface class '%s'"), *DeclaredInterfaceName, *TempClassDef->GetName());
 		}
 	}
 
@@ -5370,7 +5370,7 @@ FUnrealClassDefinitionInfo* FHeaderParser::ParseInterfaceNameDeclaration(FString
 	check(SuperClassDef);
 	if (SuperClassDef != TempClassDef)
 	{
-		FUHTException::Throwf(*this, TEXT("%s's superclass must be %s, not %s"), *ClassDef->GetPathName(), *SuperClassDef->GetPathName(), *TempClassDef->GetPathName());
+		Throwf(TEXT("%s's superclass must be %s, not %s"), *ClassDef->GetPathName(), *SuperClassDef->GetPathName(), *TempClassDef->GetPathName());
 	}
 
 	return ClassDef;
@@ -5455,7 +5455,7 @@ void FHeaderParser::CompileInterfaceDeclaration()
 		// All classes that are parsed are expected to be native
 		if (!InterfaceSuperClassDef->HasAnyClassFlags(CLASS_Native))
 		{
-			FUHTException::Throwf(*this, TEXT("Native classes cannot extend non-native classes"));
+			Throwf(TEXT("Native classes cannot extend non-native classes"));
 		}
 		InterfaceClassDef->SetClassWithin(InterfaceSuperClassDef->GetClassWithin());
 	}
@@ -5473,13 +5473,13 @@ void FHeaderParser::CompileInterfaceDeclaration()
 		{
 			default:
 			{
-				FUHTException::Throwf(*this, TEXT("Unknown interface specifier '%s'"), *Specifier.Key);
+				Throwf(TEXT("Unknown interface specifier '%s'"), *Specifier.Key);
 			}
 			break;
 
 			case EInterfaceSpecifier::DependsOn:
 			{
-				FUHTException::Throwf(*this, TEXT("The dependsOn specifier is deprecated. Please use #include \"ClassHeaderFilename.h\" instead."));
+				Throwf(TEXT("The dependsOn specifier is deprecated. Please use #include \"ClassHeaderFilename.h\" instead."));
 			}
 			break;
 
@@ -5501,7 +5501,7 @@ void FHeaderParser::CompileInterfaceDeclaration()
 	const FString ExpectedInterfaceName = InterfaceClassDef->GetNameWithPrefix(EEnforceInterfacePrefix::U);
 	if (DeclaredInterfaceName != ExpectedInterfaceName)
 	{
-		FUHTException::Throwf(*this, TEXT("Interface name '%s' is invalid, the first class should be identified as '%s'"), *DeclaredInterfaceName, *ExpectedInterfaceName );
+		Throwf(TEXT("Interface name '%s' is invalid, the first class should be identified as '%s'"), *DeclaredInterfaceName, *ExpectedInterfaceName );
 	}
 
 	// Try parsing metadata for the interface
@@ -5524,7 +5524,7 @@ void FHeaderParser::CompileRigVMMethodDeclaration(FUnrealStructDefinitionInfo& S
 {
 	if (!MatchSymbol(TEXT("(")))
 	{
-		FUHTException::Throwf(*this, TEXT("Bad RIGVM_METHOD definition"));
+		Throwf(TEXT("Bad RIGVM_METHOD definition"));
 	}
 
 	// find the next close brace
@@ -5562,7 +5562,7 @@ void FHeaderParser::CompileRigVMMethodDeclaration(FUnrealStructDefinitionInfo& S
 
 	if (!MatchSymbol(TEXT("(")))
 	{
-		FUHTException::Throwf(*this, TEXT("Bad RIGVM_METHOD definition"));
+		Throwf(TEXT("Bad RIGVM_METHOD definition"));
 	}
 
 	TArray<FString> ParamsContent;
@@ -5698,7 +5698,7 @@ void FHeaderParser::ParseRigVMMethodParameters(FUnrealStructDefinitionInfo& Stru
 
 		if (Parameter.bEditorOnly)
 		{
-			UE_LOG_ERROR_UHT(*this, TEXT("RigVM Struct '%s' - Member '%s' is editor only - WITH_EDITORONLY_DATA not allowed on structs with RIGVM_METHOD."), *StructDef.GetName(), *Parameter.Name, *MemberCPPType);
+			LogError(TEXT("RigVM Struct '%s' - Member '%s' is editor only - WITH_EDITORONLY_DATA not allowed on structs with RIGVM_METHOD."), *StructDef.GetName(), *Parameter.Name, *MemberCPPType);
 		}
 
 		if (!ExtendedCPPType.IsEmpty())
@@ -5706,7 +5706,7 @@ void FHeaderParser::ParseRigVMMethodParameters(FUnrealStructDefinitionInfo& Stru
 			// we only support arrays - no maps or similar data structures
 			if (MemberCPPType != TArrayText && MemberCPPType != TEnumAsByteText)
 			{
-				UE_LOG_ERROR_UHT(*this, TEXT("RigVM Struct '%s' - Member '%s' type '%s' not supported by RigVM."), *StructDef.GetName(), *Parameter.Name, *MemberCPPType);
+				LogError(TEXT("RigVM Struct '%s' - Member '%s' type '%s' not supported by RigVM."), *StructDef.GetName(), *Parameter.Name, *MemberCPPType);
 				continue;
 			}
 		}
@@ -5732,12 +5732,12 @@ void FHeaderParser::ParseRigVMMethodParameters(FUnrealStructDefinitionInfo& Stru
 
 	if (StructRigVMInfo.Members.Num() == 0)
 	{
-		UE_LOG_ERROR_UHT(*this, TEXT("RigVM Struct '%s' - has zero members - invalid RIGVM_METHOD."), *StructDef.GetName());
+		LogError(TEXT("RigVM Struct '%s' - has zero members - invalid RIGVM_METHOD."), *StructDef.GetName());
 	}
 
 	if (StructRigVMInfo.Members.Num() > 64)
 	{
-		UE_LOG_ERROR_UHT(*this, TEXT("RigVM Struct '%s' - has %d members (64 is the limit)."), *StructDef.GetName(), StructRigVMInfo.Members.Num());
+		LogError(TEXT("RigVM Struct '%s' - has %d members (64 is the limit)."), *StructDef.GetName(), StructRigVMInfo.Members.Num());
 	}
 }
 
@@ -5784,14 +5784,14 @@ void FHeaderParser::ParseParameterList(FUnrealFunctionDefinitionInfo& FunctionDe
 			{
 				if (!FunctionDef.HasAnyFunctionFlags(FUNC_NetRequest | FUNC_NetResponse))
 				{
-					UE_LOG_ERROR_UHT(*this, TEXT("Maps are not supported in an RPC."));
+					LogError(TEXT("Maps are not supported in an RPC."));
 				}
 			}
 			else if (Property.ArrayType == EArrayType::Set)
 			{
 				if (!FunctionDef.HasAnyFunctionFlags(FUNC_NetRequest | FUNC_NetResponse))
 				{
-					UE_LOG_ERROR_UHT(*this, TEXT("Sets are not supported in an RPC."));
+					LogError(TEXT("Sets are not supported in an RPC."));
 				}
 			}
 
@@ -5807,39 +5807,39 @@ void FHeaderParser::ParseParameterList(FUnrealFunctionDefinitionInfo& FunctionDe
 			{
 				if (Property.PropertyFlags & CPF_OutParm)
 				{
-					UE_LOG_ERROR_UHT(*this, TEXT("Replicated functions cannot contain out parameters"));
+					LogError(TEXT("Replicated functions cannot contain out parameters"));
 				}
 
 				if (Property.PropertyFlags & CPF_RepSkip)
 				{
-					UE_LOG_ERROR_UHT(*this, TEXT("Only service request functions cannot contain NoReplication parameters"));
+					LogError(TEXT("Only service request functions cannot contain NoReplication parameters"));
 				}
 
 				if (PropDef.GetPropertyBase().IsDelegateOrDelegateStaticArray())
 				{
-					UE_LOG_ERROR_UHT(*this, TEXT("Replicated functions cannot contain delegate parameters (this would be insecure)"));
+					LogError(TEXT("Replicated functions cannot contain delegate parameters (this would be insecure)"));
 				}
 
 				if (Property.Type == CPT_String && Property.RefQualifier != ERefQualifier::ConstRef && !PropDef.IsStaticArray())
 				{
-					UE_LOG_ERROR_UHT(*this, TEXT("Replicated FString parameters must be passed by const reference"));
+					LogError(TEXT("Replicated FString parameters must be passed by const reference"));
 				}
 
 				if (Property.ArrayType == EArrayType::Dynamic && Property.RefQualifier != ERefQualifier::ConstRef && !PropDef.IsStaticArray())
 				{
-					UE_LOG_ERROR_UHT(*this, TEXT("Replicated TArray parameters must be passed by const reference"));
+					LogError(TEXT("Replicated TArray parameters must be passed by const reference"));
 				}
 			}
 			else
 			{
 				if (!(Property.PropertyFlags & CPF_RepSkip) && (Property.PropertyFlags & CPF_OutParm))
 				{
-					UE_LOG_ERROR_UHT(*this, TEXT("Service request functions cannot contain out parameters, unless marked NotReplicated"));
+					LogError(TEXT("Service request functions cannot contain out parameters, unless marked NotReplicated"));
 				}
 
 				if (!(Property.PropertyFlags & CPF_RepSkip) && PropDef.GetPropertyBase().IsDelegateOrDelegateStaticArray())
 				{
-					UE_LOG_ERROR_UHT(*this, TEXT("Service request functions cannot contain delegate parameters, unless marked NotReplicated"));
+					LogError(TEXT("Service request functions cannot contain delegate parameters, unless marked NotReplicated"));
 				}
 			}
 		}
@@ -5852,7 +5852,7 @@ void FHeaderParser::ParseParameterList(FUnrealFunctionDefinitionInfo& FunctionDe
 					if (EnumDef->GetUnderlyingType() != EUnderlyingEnumType::uint8 &&
 						EnumDef->GetUnderlyingType() != EUnderlyingEnumType::Unspecified)
 					{
-						FUHTException::Throwf(*this, TEXT("Invalid enum param for Blueprints - currently only uint8 supported"));
+						Throwf(TEXT("Invalid enum param for Blueprints - currently only uint8 supported"));
 					}
 				}
 			}
@@ -5868,7 +5868,7 @@ void FHeaderParser::ParseParameterList(FUnrealFunctionDefinitionInfo& FunctionDe
 				{
 					if (FCString::Stricmp(*PropDef.GetNameCPP(), *InvalidName) == 0)
 					{
-						UE_LOG_ERROR_UHT(*this, TEXT("Paramater name '%s' in function is invalid, '%s' is a reserved name."), *InvalidName, *InvalidName);
+						LogError(TEXT("Paramater name '%s' in function is invalid, '%s' is a reserved name."), *InvalidName, *InvalidName);
 					}
 				}
 			}
@@ -5921,7 +5921,7 @@ void FHeaderParser::ParseParameterList(FUnrealFunctionDefinitionInfo& FunctionDe
 					const bool bDefaultValueParsed = FPropertyTraits::DefaultValueStringCppFormatToInnerFormat(PropDef, DefaultArgText, InnerDefaultValue);
 					if (!bDefaultValueParsed)
 					{
-						FUHTException::Throwf(*this, TEXT("C++ Default parameter not parsed: %s \"%s\" "), *PropDef.GetName(), *DefaultArgText);
+						Throwf(TEXT("C++ Default parameter not parsed: %s \"%s\" "), *PropDef.GetName(), *DefaultArgText);
 					}
 
 					MetaData->Add(KeyName, InnerDefaultValue);
@@ -5955,7 +5955,7 @@ FUnrealFunctionDefinitionInfo& FHeaderParser::CompileDelegateDeclaration(const T
 		GetToken(Token);
 		if (!IsValidDelegateDeclaration(Token))
 		{
-			FUHTException::Throwf(*this, TEXT("Unexpected token following UDELEGATE(): %s"), Token.Identifier);
+			Throwf(TEXT("Unexpected token following UDELEGATE(): %s"), Token.Identifier);
 		}
 
 		DelegateMacro = Token.Identifier;
@@ -5990,13 +5990,13 @@ FUnrealFunctionDefinitionInfo& FHeaderParser::CompileDelegateDeclaration(const T
 
 	if (DelegateMacro != ExpectedOriginalString)
 	{
-		FUHTException::Throwf(*this, TEXT("Unable to parse delegate declaration; expected '%s' but found '%s'."), *ExpectedOriginalString, *DelegateMacro);
+		Throwf(TEXT("Unable to parse delegate declaration; expected '%s' but found '%s'."), *ExpectedOriginalString, *DelegateMacro);
 	}
 
 	// Multi-cast delegate function signatures are not allowed to have a return value
 	if (bHasReturnValue && bIsMulticast)
 	{
-		UE_LOG_ERROR_UHT(*this, TEXT("Multi-cast delegates function signatures must not return a value"));
+		LogError(TEXT("Multi-cast delegates function signatures must not return a value"));
 	}
 
 	// Delegate signature
@@ -6027,7 +6027,7 @@ FUnrealFunctionDefinitionInfo& FHeaderParser::CompileDelegateDeclaration(const T
 	// Get the delegate name
 	if (!GetIdentifier(FuncInfo.Function))
 	{
-		FUHTException::Throwf(*this, TEXT("Missing name for %s"), CurrentScopeName );
+		Throwf(TEXT("Missing name for %s"), CurrentScopeName );
 	}
 
 	// If this is a delegate function then go ahead and mangle the name so we don't collide with
@@ -6040,7 +6040,7 @@ FUnrealFunctionDefinitionInfo& FHeaderParser::CompileDelegateDeclaration(const T
 
 		if (!Name.StartsWith(TEXT("F"), ESearchCase::CaseSensitive))
 		{
-			FUHTException::Throwf(*this, TEXT("Delegate type declarations must start with F"));
+			Throwf(TEXT("Delegate type declarations must start with F"));
 		}
 
 		Name.RightChopInline(1, false);
@@ -6052,7 +6052,7 @@ FUnrealFunctionDefinitionInfo& FHeaderParser::CompileDelegateDeclaration(const T
 		FCString::Strcpy( FuncInfo.Function.Identifier, *Name );
 	}
 
-	FUnrealFunctionDefinitionInfo& DelegateSignatureFunctionDef = bIsSparse ? CreateDelegateFunction<USparseDelegateFunction>(MoveTemp(FuncInfo), EFunctionType::SparseDelegate) : CreateDelegateFunction<UDelegateFunction>(MoveTemp(FuncInfo), EFunctionType::Delegate);
+	FUnrealFunctionDefinitionInfo& DelegateSignatureFunctionDef = CreateFunction(MoveTemp(FuncInfo), bIsSparse ? EFunctionType::SparseDelegate : EFunctionType::Delegate);
 	const FFuncInfo& FuncDefFuncInfo = DelegateSignatureFunctionDef.GetFunctionData();
 
 	// determine whether this function should be 'const'
@@ -6069,14 +6069,14 @@ FUnrealFunctionDefinitionInfo& FHeaderParser::CompileDelegateDeclaration(const T
 
 		if (!GetIdentifier(OwningClass))
 		{
-			FUHTException::Throwf(*this, TEXT("Missing OwningClass specifier."));
+			Throwf(TEXT("Missing OwningClass specifier."));
 		}
 		RequireSymbol(TEXT(','), TEXT("Delegate Declaration"));
 		
 		FToken DelegateName;
 		if (!GetIdentifier(DelegateName))
 		{
-			FUHTException::Throwf(*this, TEXT("Missing Delegate Name."));
+			Throwf(TEXT("Missing Delegate Name."));
 		}
 
 		DelegateSignatureFunctionDef.SetSparseOwningClassName(*GetClassNameWithoutPrefix(FString(OwningClass.Identifier)));
@@ -6096,7 +6096,7 @@ FUnrealFunctionDefinitionInfo& FHeaderParser::CompileDelegateDeclaration(const T
 		int32 ParamCount = UE_PTRDIFF_TO_INT32(FoundParamCount - UHTConfig.DelegateParameterCountStrings.GetData()) + 1;
 		if (DelegateSignatureFunctionDef.GetProperties().Num() != ParamCount)
 		{
-			FUHTException::Throwf(*this, TEXT("Expected %d parameters but found %d parameters"), ParamCount, DelegateSignatureFunctionDef.GetProperties().Num());
+			Throwf(TEXT("Expected %d parameters but found %d parameters"), ParamCount, DelegateSignatureFunctionDef.GetProperties().Num());
 		}
 	}
 	else
@@ -6135,7 +6135,7 @@ FUnrealFunctionDefinitionInfo& FHeaderParser::CompileDelegateDeclaration(const T
 		FUnrealFunctionDefinitionInfo* TestFuncDef = *FunctionIterator;
 		if (TestFuncDef != &DelegateSignatureFunctionDef && TestFuncDef->GetFName() == DelegateSignatureFunctionDef.GetFName())
 		{
-			FUHTException::Throwf(*this, TEXT("Can't override delegate signature function '%s'"), FuncDefFuncInfo.Function.Identifier);
+			Throwf(TEXT("Can't override delegate signature function '%s'"), FuncDefFuncInfo.Function.Identifier);
 		}
 	}
 
@@ -6159,7 +6159,7 @@ FUnrealFunctionDefinitionInfo& FHeaderParser::CompileFunctionDeclaration()
 	FUnrealClassDefinitionInfo& OuterClassDef = UHTCastChecked<FUnrealClassDefinitionInfo>(GetCurrentClassDef());
 	if (!OuterClassDef.HasAnyClassFlags(CLASS_Native))
 	{
-		FUHTException::Throwf(*this, TEXT("Should only be here for native classes!"));
+		Throwf(TEXT("Should only be here for native classes!"));
 	}
 
 	// Process all specifiers.
@@ -6190,7 +6190,7 @@ FUnrealFunctionDefinitionInfo& FHeaderParser::CompileFunctionDeclaration()
 	}
 	else
 	{
-		FUHTException::Throwf(*this, TEXT("Unknown access level"));
+		Throwf(TEXT("Unknown access level"));
 	}
 
 	// non-static functions in a const class must be const themselves
@@ -6227,7 +6227,7 @@ FUnrealFunctionDefinitionInfo& FHeaderParser::CompileFunctionDeclaration()
 	if ((FuncInfo.FunctionFlags & FUNC_BlueprintPure) && OuterClassDef.HasAnyClassFlags(CLASS_Interface))
 	{
 		// Until pure interface casts are supported, we don't allow pures in interfaces
-		UE_LOG_ERROR_UHT(*this, TEXT("BlueprintPure specifier is not allowed for interface functions"));
+		LogError(TEXT("BlueprintPure specifier is not allowed for interface functions"));
 	}
 
 	if (FuncInfo.FunctionFlags & FUNC_Net)
@@ -6265,7 +6265,7 @@ FUnrealFunctionDefinitionInfo& FHeaderParser::CompileFunctionDeclaration()
 			// To allow for quick iteration, don't enforce the requirement that game functions have to be categorized
 			if (bIsCurrentModulePartOfEngine)
 			{
-				UE_LOG_ERROR_UHT(*this, TEXT("An explicit Category specifier is required for Blueprint accessible functions in an Engine module."));
+				LogError(TEXT("An explicit Category specifier is required for Blueprint accessible functions in an Engine module."));
 			}
 		}
 	}
@@ -6280,7 +6280,7 @@ FUnrealFunctionDefinitionInfo& FHeaderParser::CompileFunctionDeclaration()
 			// Ensure that blueprint events are only allowed in implementable interfaces. Internal only functions allowed
 			if (!bCanImplementInBlueprints)
 			{
-				UE_LOG_ERROR_UHT(*this, TEXT("Interfaces that are not implementable in blueprints cannot have BlueprintImplementableEvent members."));
+				LogError(TEXT("Interfaces that are not implementable in blueprints cannot have BlueprintImplementableEvent members."));
 			}
 		}
 		
@@ -6291,7 +6291,7 @@ FUnrealFunctionDefinitionInfo& FHeaderParser::CompileFunctionDeclaration()
 			// Ensure that if this interface contains blueprint callable functions that are not blueprint defined, that it must be implemented natively
 			if (bCanImplementInBlueprints)
 			{
-				UE_LOG_ERROR_UHT(*this, TEXT("Blueprint implementable interfaces cannot contain BlueprintCallable functions that are not BlueprintImplementableEvents.  Use CannotImplementInterfaceInBlueprint on the interface if you wish to keep this function."));
+				LogError(TEXT("Blueprint implementable interfaces cannot contain BlueprintCallable functions that are not BlueprintImplementableEvents.  Use CannotImplementInterfaceInBlueprint on the interface if you wish to keep this function."));
 			}
 		}
 	}
@@ -6313,7 +6313,7 @@ FUnrealFunctionDefinitionInfo& FHeaderParser::CompileFunctionDeclaration()
 
 					if (OuterClassDef.HasAnyClassFlags(CLASS_RequiredAPI))
 					{
-						FUHTException::Throwf(*this, TEXT("'%s' must not be used on methods of a class that is marked '%s' itself."), *RequiredAPIMacroIfPresent, *RequiredAPIMacroIfPresent);
+						Throwf(TEXT("'%s' must not be used on methods of a class that is marked '%s' itself."), *RequiredAPIMacroIfPresent, *RequiredAPIMacroIfPresent);
 					}
 					FuncInfo.FunctionFlags |= FUNC_RequiredAPI;
 					FuncInfo.FunctionExportFlags |= FUNCEXPORT_RequiredAPI;
@@ -6332,7 +6332,7 @@ FUnrealFunctionDefinitionInfo& FHeaderParser::CompileFunctionDeclaration()
 	// Look for static again, in case there was an ENGINE_API token first
 	if (!APIMacroIfPresent.IsEmpty() && MatchIdentifier(TEXT("static"), ESearchCase::CaseSensitive))
 	{
-		FUHTException::Throwf(*this, TEXT("Unexpected API macro '%s'. Did you mean to put '%s' after the static keyword?"), *APIMacroIfPresent, *APIMacroIfPresent);
+		Throwf(TEXT("Unexpected API macro '%s'. Did you mean to put '%s' after the static keyword?"), *APIMacroIfPresent, *APIMacroIfPresent);
 	}
 
 	// Look for virtual again, in case there was an ENGINE_API token first
@@ -6352,18 +6352,18 @@ FUnrealFunctionDefinitionInfo& FHeaderParser::CompileFunctionDeclaration()
 		{
 			if (OuterClassDef.HasAnyClassFlags(CLASS_Interface))
 			{
-				FUHTException::Throwf(*this, TEXT("BlueprintImplementableEvents in Interfaces must not be declared 'virtual'"));
+				Throwf(TEXT("BlueprintImplementableEvents in Interfaces must not be declared 'virtual'"));
 			}
 
 			// if this is a BlueprintNativeEvent, make sure it's not "virtual"
 			else if (FuncInfo.FunctionFlags & FUNC_Native)
 			{
-				UE_LOG_ERROR_UHT(*this, TEXT("BlueprintNativeEvent functions must be non-virtual."));
+				LogError(TEXT("BlueprintNativeEvent functions must be non-virtual."));
 			}
 
 			else
 			{
-				UE_LOG_WARNING_UHT(*this, TEXT("BlueprintImplementableEvents should not be virtual. Use BlueprintNativeEvent instead."));
+				LogWarning(TEXT("BlueprintImplementableEvents should not be virtual. Use BlueprintNativeEvent instead."));
 			}
 		}
 	}
@@ -6372,7 +6372,7 @@ FUnrealFunctionDefinitionInfo& FHeaderParser::CompileFunctionDeclaration()
 		// if this is a function in an Interface, it must be marked 'virtual' unless it's an event
 		if (OuterClassDef.HasAnyClassFlags(CLASS_Interface) && !(FuncInfo.FunctionFlags & FUNC_BlueprintEvent))
 		{
-			FUHTException::Throwf(*this, TEXT("Interface functions that are not BlueprintImplementableEvents must be declared 'virtual'"));
+			Throwf(TEXT("Interface functions that are not BlueprintImplementableEvents must be declared 'virtual'"));
 		}
 	}
 
@@ -6385,7 +6385,7 @@ FUnrealFunctionDefinitionInfo& FHeaderParser::CompileFunctionDeclaration()
 
 		if (OuterClassDef.HasAnyClassFlags(CLASS_Interface))
 		{
-			UE_LOG_ERROR_UHT(*this, TEXT("Interface functions cannot be declared 'final'"));
+			LogError(TEXT("Interface functions cannot be declared 'final'"));
 		}
 	}
 
@@ -6404,12 +6404,12 @@ FUnrealFunctionDefinitionInfo& FHeaderParser::CompileFunctionDeclaration()
 	// Get function or operator name.
 	if (!GetIdentifier(FuncInfo.Function))
 	{
-		FUHTException::Throwf(*this, TEXT("Missing %s name"), TypeOfFunction);
+		Throwf(TEXT("Missing %s name"), TypeOfFunction);
 	}
 
 	if ( !MatchSymbol(TEXT('(')) )
 	{
-		FUHTException::Throwf(*this, TEXT("Bad %s definition"), TypeOfFunction);
+		Throwf(TEXT("Bad %s definition"), TypeOfFunction);
 	}
 
 	if (FuncInfo.FunctionFlags & FUNC_Net)
@@ -6417,14 +6417,14 @@ FUnrealFunctionDefinitionInfo& FHeaderParser::CompileFunctionDeclaration()
 		bool bIsNetService = !!(FuncInfo.FunctionFlags & (FUNC_NetRequest | FUNC_NetResponse));
 		if (bHasReturnValue && !bIsNetService)
 		{
-			FUHTException::Throwf(*this, TEXT("Replicated functions can't have return values"));
+			Throwf(TEXT("Replicated functions can't have return values"));
 		}
 
 		if (FuncInfo.RPCId > 0)
 		{
 			if (FString* ExistingFunc = UsedRPCIds.Find(FuncInfo.RPCId))
 			{
-				FUHTException::Throwf(*this, TEXT("Function %s already uses identifier %d"), **ExistingFunc, FuncInfo.RPCId);
+				Throwf(TEXT("Function %s already uses identifier %d"), **ExistingFunc, FuncInfo.RPCId);
 			}
 
 			UsedRPCIds.Add(FuncInfo.RPCId, FuncInfo.Function.Identifier);
@@ -6451,7 +6451,7 @@ FUnrealFunctionDefinitionInfo& FHeaderParser::CompileFunctionDeclaration()
 		}
 	}
 
-	FUnrealFunctionDefinitionInfo& FuncDef = CreateFunction(MoveTemp(FuncInfo));
+	FUnrealFunctionDefinitionInfo& FuncDef = CreateFunction(MoveTemp(FuncInfo), EFunctionType::Function);
 	const FFuncInfo& FuncDefFuncInfo = FuncDef.GetFunctionData();
 	FUnrealFunctionDefinitionInfo* SuperFuncDef = FuncDef.GetSuperFunction();
 
@@ -6490,7 +6490,7 @@ FUnrealFunctionDefinitionInfo& FHeaderParser::CompileFunctionDeclaration()
 		if (FUnrealFunctionDefinitionInfo* OverriddenFunctionDef = FindFunction(*SuperStructDef, FuncDefFuncInfo.Function.Identifier, true, nullptr))
 		{
 			// Native function overrides should be done in CPP text, not in a UFUNCTION() declaration (you can't change flags, and it'd otherwise be a burden to keep them identical)
-			UE_LOG_ERROR_UHT(*this, TEXT("%s: Override of UFUNCTION in parent class (%s) cannot have a UFUNCTION() declaration above it; it will use the same parameters as the original declaration."), FuncDefFuncInfo.Function.Identifier, *OverriddenFunctionDef->GetOuter()->GetName());
+			LogError(TEXT("%s: Override of UFUNCTION in parent class (%s) cannot have a UFUNCTION() declaration above it; it will use the same parameters as the original declaration."), FuncDefFuncInfo.Function.Identifier, *OverriddenFunctionDef->GetOuter()->GetName());
 		}
 	}
 
@@ -6500,7 +6500,7 @@ FUnrealFunctionDefinitionInfo& FHeaderParser::CompileFunctionDeclaration()
 		// we don't want to prevent compilation:
 		if (!bClassGeneratedFromBP)
 		{
-			UE_LOG_ERROR_UHT(*this, TEXT("BlueprintPure specifier is not allowed for functions with no return value and no output parameters."));
+			LogError(TEXT("BlueprintPure specifier is not allowed for functions with no return value and no output parameters."));
 		}
 	}
 
@@ -6511,7 +6511,7 @@ FUnrealFunctionDefinitionInfo& FHeaderParser::CompileFunctionDeclaration()
 		if (FuncDef.HasAnyFunctionFlags(FUNC_Native))
 		{
 			// @TODO: UCREMOVAL Reconsider?
-			//FUHTException::Throwf(*this, TEXT("'const' may only be used for native functions"));
+			//Throwf(TEXT("'const' may only be used for native functions"));
 		}
 
 		FuncDef.SetFunctionFlags(FUNC_Const);
@@ -6552,7 +6552,7 @@ FUnrealFunctionDefinitionInfo& FHeaderParser::CompileFunctionDeclaration()
 
 		if (!bGotZero)
 		{
-			FUHTException::Throwf(*this, TEXT("Expected 0 to indicate function is abstract"));
+			Throwf(TEXT("Expected 0 to indicate function is abstract"));
 		}
 	}
 
@@ -6564,11 +6564,11 @@ FUnrealFunctionDefinitionInfo& FHeaderParser::CompileFunctionDeclaration()
 		FuncDef.SetFunctionExportFlags(FUNCEXPORT_Final);
 		if (UHTCastChecked<FUnrealClassDefinitionInfo>(GetCurrentClassDef()).HasAnyClassFlags(CLASS_Interface))
 		{
-			FUHTException::Throwf(*this, TEXT("Interface functions cannot be declared 'final'"));
+			Throwf(TEXT("Interface functions cannot be declared 'final'"));
 		}
 		else if (FuncDef.HasAnyFunctionFlags(FUNC_BlueprintEvent))
 		{
-			FUHTException::Throwf(*this, TEXT("Blueprint events cannot be declared 'final'"));
+			Throwf(TEXT("Blueprint events cannot be declared 'final'"));
 		}
 	}
 
@@ -6577,7 +6577,7 @@ FUnrealFunctionDefinitionInfo& FHeaderParser::CompileFunctionDeclaration()
 	{
 		if ((FuncDef.GetFunctionFlags() & FUNC_NetFuncFlags) != (SuperFuncDef->GetFunctionFlags() & FUNC_NetFuncFlags))
 		{
-			FUHTException::Throwf(*this, TEXT("Overridden function '%s': Cannot specify different replication flags when overriding a function."), *FuncDef.GetName());
+			Throwf(TEXT("Overridden function '%s': Cannot specify different replication flags when overriding a function."), *FuncDef.GetName());
 		}
 	}
 
@@ -6585,7 +6585,7 @@ FUnrealFunctionDefinitionInfo& FHeaderParser::CompileFunctionDeclaration()
 	// this is required because the networking code only checks the class for RPCs when initializing network data, not any states within it
 	if (FuncDef.HasAnyFunctionFlags(FUNC_Net) && SuperFuncDef == nullptr && UHTCast<FUnrealClassDefinitionInfo>(FuncDef.GetOuter()) == nullptr)
 	{
-		FUHTException::Throwf(*this, TEXT("Function '%s': Base implementation of RPCs cannot be in a state. Add a stub outside state scope."), *FuncDef.GetName());
+		Throwf(TEXT("Function '%s': Base implementation of RPCs cannot be in a state. Add a stub outside state scope."), *FuncDef.GetName());
 	}
 
 	if (FuncDef.HasAnyFunctionFlags(FUNC_BlueprintCallable | FUNC_BlueprintEvent))
@@ -6594,14 +6594,14 @@ FUnrealFunctionDefinitionInfo& FHeaderParser::CompileFunctionDeclaration()
 		{
 			if (PropertyDef->IsStaticArray())
 			{
-				FUHTException::Throwf(*this, TEXT("Static array cannot be exposed to blueprint. Function: %s Parameter %s\n"), *FuncDef.GetName(), *PropertyDef->GetName());
+				Throwf(TEXT("Static array cannot be exposed to blueprint. Function: %s Parameter %s\n"), *FuncDef.GetName(), *PropertyDef->GetName());
 			}
 
 			if (!FPropertyTraits::IsSupportedByBlueprint(*PropertyDef, false))
 			{
 				FString ExtendedCPPType;
 				FString CPPType = PropertyDef->GetCPPType(&ExtendedCPPType);
-				UE_LOG_ERROR_UHT(*this, TEXT("Type '%s%s' is not supported by blueprint. %s.%s"), *CPPType, *ExtendedCPPType, *FuncDef.GetName(), *PropertyDef->GetName());
+				LogError(TEXT("Type '%s%s' is not supported by blueprint. %s.%s"), *CPPType, *ExtendedCPPType, *FuncDef.GetName(), *PropertyDef->GetName());
 			}
 		}
 	}
@@ -6613,7 +6613,7 @@ FUnrealFunctionDefinitionInfo& FHeaderParser::CompileFunctionDeclaration()
 	FToken Token;
 	if (!GetToken(Token))
 	{
-		FUHTException::Throwf(*this, TEXT("Unexpected end of file"));
+		Throwf(TEXT("Unexpected end of file"));
 	}
 
 	// Optionally consume a semicolon
@@ -6653,7 +6653,7 @@ void FHeaderParser::ParseFieldMetaData(TMap<FName, FString>& MetaData, const TCH
 		RequireSymbol( TEXT('('), ErrorMessageGetter );
 		if (!GetRawTokenRespectingQuotes(PropertyMetaData, TCHAR(')')))
 		{
-			FUHTException::Throwf(*this, TEXT("'%s': No metadata specified"), FieldName);
+			Throwf(TEXT("'%s': No metadata specified"), FieldName);
 		}
 		RequireSymbol( TEXT(')'), ErrorMessageGetter);
 	}
@@ -6766,7 +6766,7 @@ void FHeaderParser::ValidatePropertyIsDeprecatedIfNecessary(const FPropertyBase&
 	// check to see if we have a FClassProperty using a deprecated class
 	if (VarProperty.MetaClassDef != nullptr && VarProperty.MetaClassDef->HasAnyClassFlags(CLASS_Deprecated))
 	{
-		UE_LOG_ERROR_UHT(*this, TEXT("Property is using a deprecated class: %s.  Property should be marked deprecated as well."), *VarProperty.MetaClassDef->GetPathName());
+		LogError(TEXT("Property is using a deprecated class: %s.  Property should be marked deprecated as well."), *VarProperty.MetaClassDef->GetPathName());
 	}
 
 	// check to see if we have a FObjectProperty using a deprecated class.
@@ -6776,7 +6776,7 @@ void FHeaderParser::ValidatePropertyIsDeprecatedIfNecessary(const FPropertyBase&
 	{
 		if (VarProperty.ClassDef->HasAnyClassFlags(CLASS_Deprecated))	// and the object class being used has been deprecated
 		{
-			UE_LOG_ERROR_UHT(*this, TEXT("Property is using a deprecated class: %s.  Property should be marked deprecated as well."), *VarProperty.ClassDef->GetPathName());
+			LogError(TEXT("Property is using a deprecated class: %s.  Property should be marked deprecated as well."), *VarProperty.ClassDef->GetPathName());
 		}
 	}
 }
@@ -6805,7 +6805,7 @@ bool FHeaderParser::ValidateScriptStructOkForNet(const FString& OriginStructName
 			if (!PropertyDef->HasAnyPropertyFlags(CPF_RepSkip))
 			{
 				bIsStructValid = false;
-				UE_LOG_ERROR_UHT(*this, TEXT("Sets are not supported for Replication or RPCs.  Set %s in %s.  Origin %s"), *PropertyDef->GetName(), *PropertyDef->GetOuter()->GetName(), *OriginStructName);
+				LogError(TEXT("Sets are not supported for Replication or RPCs.  Set %s in %s.  Origin %s"), *PropertyDef->GetName(), *PropertyDef->GetOuter()->GetName(), *OriginStructName);
 			}
 		}
 		else if (PropertyDef->IsMap())
@@ -6813,7 +6813,7 @@ bool FHeaderParser::ValidateScriptStructOkForNet(const FString& OriginStructName
 			if (!PropertyDef->HasAnyPropertyFlags(CPF_RepSkip))
 			{
 				bIsStructValid = false;
-				UE_LOG_ERROR_UHT(*this, TEXT("Maps are not supported for Replication or RPCs.  Map %s in %s.  Origin %s"), *PropertyDef->GetName(), *PropertyDef->GetOuter()->GetName(), *OriginStructName);
+				LogError(TEXT("Maps are not supported for Replication or RPCs.  Map %s in %s.  Origin %s"), *PropertyDef->GetName(), *PropertyDef->GetOuter()->GetName(), *OriginStructName);
 			}
 		}
 		else if (PropertyDef->IsStructOrStructStaticArray())
@@ -6890,7 +6890,7 @@ void FHeaderParser::CompileVariableDeclaration(FUnrealStructDefinitionInfo& Stru
 		}
 		else
 		{
-			UE_LOG_ERROR_UHT(*this, TEXT("An explicit Category specifier is required for any property exposed to the editor or Blueprints in an Engine module."));
+			LogError(TEXT("An explicit Category specifier is required for any property exposed to the editor or Blueprints in an Engine module."));
 		}
 	}
 
@@ -6899,13 +6899,13 @@ void FHeaderParser::CompileVariableDeclaration(FUnrealStructDefinitionInfo& Stru
 	{
 		// Get the name of the type, removing the asterisk representing the pointer
 		FString TypeName = FString(TypeRange.Count, Input + TypeRange.StartIndex).TrimStartAndEnd().LeftChop(1).TrimEnd();
-		FUHTException::Throwf(*this, TEXT("UPROPERTY pointers cannot be interfaces - did you mean TScriptInterface<%s>?"), *TypeName);
+		Throwf(TEXT("UPROPERTY pointers cannot be interfaces - did you mean TScriptInterface<%s>?"), *TypeName);
 	}
 
 	// If the category was specified explicitly, it wins
 	if (Category && !(OriginalPropertyBase.PropertyFlags & (CPF_Edit|CPF_BlueprintVisible|CPF_BlueprintAssignable|CPF_BlueprintCallable)))
 	{
-		UE_LOG_WARNING_UHT(*this, TEXT("Property has a Category set but is not exposed to the editor or Blueprints with EditAnywhere, BlueprintReadWrite, VisibleAnywhere, BlueprintReadOnly, BlueprintAssignable, BlueprintCallable keywords.\r\n"));
+		LogWarning(TEXT("Property has a Category set but is not exposed to the editor or Blueprints with EditAnywhere, BlueprintReadWrite, VisibleAnywhere, BlueprintReadOnly, BlueprintAssignable, BlueprintCallable keywords.\r\n"));
 	}
 
 	// Make sure that editblueprint variables are editable
@@ -6913,26 +6913,26 @@ void FHeaderParser::CompileVariableDeclaration(FUnrealStructDefinitionInfo& Stru
 	{
 		if (OriginalPropertyBase.PropertyFlags & CPF_DisableEditOnInstance)
 		{
-			UE_LOG_ERROR_UHT(*this, TEXT("Property cannot have 'DisableEditOnInstance' without being editable"));
+			LogError(TEXT("Property cannot have 'DisableEditOnInstance' without being editable"));
 		}
 
 		if (OriginalPropertyBase.PropertyFlags & CPF_DisableEditOnTemplate)
 		{
-			UE_LOG_ERROR_UHT(*this, TEXT("Property cannot have 'DisableEditOnTemplate' without being editable"));
+			LogError(TEXT("Property cannot have 'DisableEditOnTemplate' without being editable"));
 		}
 	}
 
 	// Validate.
 	if (OriginalPropertyBase.PropertyFlags & CPF_ParmFlags)
 	{
-		FUHTException::Throwf(*this, TEXT("Illegal type modifiers in member variable declaration") );
+		Throwf(TEXT("Illegal type modifiers in member variable declaration") );
 	}
 
 	if (FString* ExposeOnSpawnValue = OriginalPropertyBase.MetaData.Find(NAME_ExposeOnSpawn))
 	{
 		if ((*ExposeOnSpawnValue == TEXT("true")) && !FExposeOnSpawnValidator::IsSupported(OriginalPropertyBase))
 		{
-			UE_LOG_ERROR_UHT(*this, TEXT("ExposeOnSpawn - Property cannot be exposed"));
+			LogError(TEXT("ExposeOnSpawn - Property cannot be exposed"));
 		}
 	}
 
@@ -6962,7 +6962,7 @@ void FHeaderParser::CompileVariableDeclaration(FUnrealStructDefinitionInfo& Stru
 				int32 BitfieldSize = 0;
 				if (!GetConstInt(/*out*/ BitfieldSize) || (BitfieldSize != 1))
 				{
-					FUHTException::Throwf(*this, TEXT("Bad or missing bitfield size for '%s', must be 1."), *NewPropDef.GetName());
+					Throwf(TEXT("Bad or missing bitfield size for '%s', must be 1."), *NewPropDef.GetName());
 				}
 			}
 		}
@@ -6974,7 +6974,7 @@ void FHeaderParser::CompileVariableDeclaration(FUnrealStructDefinitionInfo& Stru
 		{
 			if (NewProperties.Num())
 			{
-				FUHTException::Throwf(*this, TEXT("Comma delimited properties cannot be converted %s.%s\n"), *StructDef.GetName(), *NewPropDef.GetName());
+				Throwf(TEXT("Comma delimited properties cannot be converted %s.%s\n"), *StructDef.GetName(), *NewPropDef.GetName());
 			}
 		}
 
@@ -6995,19 +6995,19 @@ void FHeaderParser::CompileVariableDeclaration(FUnrealStructDefinitionInfo& Stru
 		{
 			if (ScriptStructDef != nullptr && !StructDef.GetBoolMetaDataHierarchical(FHeaderParserNames::NAME_BlueprintType))
 			{
-				UE_LOG_ERROR_UHT(*this, TEXT("Cannot expose property to blueprints in a struct that is not a BlueprintType. %s.%s"), *StructDef.GetName(), *NewPropDef.GetName());
+				LogError(TEXT("Cannot expose property to blueprints in a struct that is not a BlueprintType. %s.%s"), *StructDef.GetName(), *NewPropDef.GetName());
 			}
 
 			if (NewPropDef.IsStaticArray())
 			{
-				UE_LOG_ERROR_UHT(*this, TEXT("Static array cannot be exposed to blueprint %s.%s"), *StructDef.GetName(), *NewPropDef.GetName());
+				LogError(TEXT("Static array cannot be exposed to blueprint %s.%s"), *StructDef.GetName(), *NewPropDef.GetName());
 			}
 
 			if (!FPropertyTraits::IsSupportedByBlueprint(NewPropDef, true))
 			{
 				FString ExtendedCPPType;
 				FString CPPType = NewPropDef.GetCPPType(&ExtendedCPPType);
-				UE_LOG_ERROR_UHT(*this, TEXT("Type '%s%s' is not supported by blueprint. %s.%s"), *CPPType, *ExtendedCPPType, *StructDef.GetName(), *NewPropDef.GetName());
+				LogError(TEXT("Type '%s%s' is not supported by blueprint. %s.%s"), *CPPType, *ExtendedCPPType, *StructDef.GetName(), *NewPropDef.GetName());
 			}
 		}
 
@@ -7123,7 +7123,7 @@ bool FHeaderParser::CompileStatement(TArray<FUnrealFunctionDefinitionInfo*>& Del
 	}
 	else if (!CompileDeclaration(DelegatesToFixup, Token))
 	{
-		FUHTException::Throwf(*this, TEXT("'%s': Bad command or expression"), Token.Identifier );
+		Throwf(TEXT("'%s': Bad command or expression"), Token.Identifier );
 	}
 	return true;
 }
@@ -7165,11 +7165,11 @@ void FHeaderParser::SkipStatements( int32 NestCount, const TCHAR* ErrorTag  )
 
 	if( NestCount > 0 )
 	{
-		FUHTException::Throwf(*this, TEXT("Unexpected end of file at end of %s"), ErrorTag );
+		Throwf(TEXT("Unexpected end of file at end of %s"), ErrorTag );
 	}
 	else if ( NestCount < 0 )
 	{
-		FUHTException::Throwf(*this, TEXT("Extraneous closing brace found in %s"), ErrorTag);
+		Throwf(TEXT("Extraneous closing brace found in %s"), ErrorTag);
 	}
 }
 
@@ -7297,22 +7297,22 @@ void FHeaderParser::ParseHeader()
 		}
 
 		RPCsNeedingHookup.Empty();
-		FUHTException::Throwf(*this, TEXT("%s"), *ErrorMsg);
+		Throwf(TEXT("%s"), *ErrorMsg);
 	}
 
 	// Make sure the compilation ended with valid nesting.
 	if (bEncounteredNewStyleClass_UnmatchedBrackets)
 	{
-		FUHTException::Throwf(*this, TEXT("Missing } at end of class") );
+		Throwf(TEXT("Missing } at end of class") );
 	}
 
 	if (NestLevel == 1)
 	{
-		FUHTException::Throwf(*this, TEXT("Internal nest inconsistency") );
+		Throwf(TEXT("Internal nest inconsistency") );
 	}
 	else if (NestLevel > 2)
 	{
-		FUHTException::Throwf(*this, TEXT("Unexpected end of script in '%s' block"), NestTypeName(TopNest->NestType) );
+		Throwf(TEXT("Unexpected end of script in '%s' block"), NestTypeName(TopNest->NestType) );
 	}
 
 	for (TSharedRef<FUnrealTypeDefinitionInfo>& TypeDef : SourceFile.GetDefinedClasses())
@@ -7326,7 +7326,7 @@ void FHeaderParser::ParseHeader()
 	if (!bSpottedAutogeneratedHeaderInclude && !bEmptyFile && !bNoExportClassesOnly)
 	{
 		const FString& ExpectedHeaderName = SourceFile.GetGeneratedHeaderFilename();
-		FUHTException::Throwf(*this, TEXT("Expected an include at the top of the header: '#include \"%s\"'"), *ExpectedHeaderName);
+		Throwf(TEXT("Expected an include at the top of the header: '#include \"%s\"'"), *ExpectedHeaderName);
 	}
 
 	for (const TSharedRef<FUnrealTypeDefinitionInfo>& ClassDef : SourceFile.GetDefinedClasses())
@@ -7382,20 +7382,20 @@ FHeaderParser::FHeaderParser(FUnrealPackageDefinitionInfo& InPackageDef, FUnreal
 }
 
 // Throws if a specifier value wasn't provided
-void FHeaderParser::RequireSpecifierValue(const FUHTExceptionContext& Context, const FPropertySpecifier& Specifier, bool bRequireExactlyOne)
+void FHeaderParser::RequireSpecifierValue(const FUHTMessageProvider& Context, const FPropertySpecifier& Specifier, bool bRequireExactlyOne)
 {
 	if (Specifier.Values.Num() == 0)
 	{
-		FUHTException::Throwf(Context, TEXT("The specifier '%s' must be given a value"), *Specifier.Key);
+		Context.Throwf(TEXT("The specifier '%s' must be given a value"), *Specifier.Key);
 	}
 	else if ((Specifier.Values.Num() != 1) && bRequireExactlyOne)
 	{
-		FUHTException::Throwf(Context, TEXT("The specifier '%s' must be given exactly one value"), *Specifier.Key);
+		Context.Throwf(TEXT("The specifier '%s' must be given exactly one value"), *Specifier.Key);
 	}
 }
 
 // Throws if a specifier value wasn't provided
-FString FHeaderParser::RequireExactlyOneSpecifierValue(const FUHTExceptionContext& Context, const FPropertySpecifier& Specifier)
+FString FHeaderParser::RequireExactlyOneSpecifierValue(const FUHTMessageProvider& Context, const FPropertySpecifier& Specifier)
 {
 	RequireSpecifierValue(Context, Specifier, /*bRequireExactlyOne*/ true);
 	return Specifier.Values[0];
@@ -7670,9 +7670,7 @@ void FHeaderParser::SimplifiedClassParse(FUnrealSourceFile& SourceFile, const TC
 					bool bKeepOldDirective = ShouldKeepDirective(OldDirective);
 					if (bKeepNewDirective != bKeepOldDirective)
 					{
-						FUHTException::Throwf(
-							SourceFile,
-							CurrentLine,
+						FUHTMessage(SourceFile, CurrentLine).Throwf(
 							TEXT("Mixing %s with %s in an #elif preprocessor block is not supported"),
 							GetBlockDirectiveTypeString(OldDirective),
 							GetBlockDirectiveTypeString(NewDirective)
@@ -7705,7 +7703,7 @@ void FHeaderParser::SimplifiedClassParse(FUnrealSourceFile& SourceFile, const TC
 							break;
 
 						case EBlockDirectiveType::WithHotReload:
-							FUHTException::Throwf(SourceFile, CurrentLine, TEXT("Bad preprocessor directive in metadata declaration: %s; Only 'CPP', '1' and '0' can have #else directives"), *ClassName);
+							FUHTMessage(SourceFile, CurrentLine).Throwf(TEXT("Bad preprocessor directive in metadata declaration: %s; Only 'CPP', '1' and '0' can have #else directives"), *ClassName);
 
 						case EBlockDirectiveType::UnrecognizedBlock:
 						case EBlockDirectiveType::WithEditor:
@@ -7746,7 +7744,7 @@ void FHeaderParser::SimplifiedClassParse(FUnrealSourceFile& SourceFile, const TC
 							|| FindInitialStr(FoundSubstr, TrimmedStrLine, TEXT("UDELEGATE"))
 							|| FindInitialStr(FoundSubstr, TrimmedStrLine, TEXT("UFUNCTION")))
 						{
-							FUHTException::Throwf(SourceFile, CurrentLine, TEXT("%s must not be inside preprocessor blocks, except for WITH_EDITORONLY_DATA"), FoundSubstr);
+							FUHTMessage(SourceFile, CurrentLine).Throwf(TEXT("%s must not be inside preprocessor blocks, except for WITH_EDITORONLY_DATA"), FoundSubstr);
 						}
 
 						// Try and determine if this line contains something like a serialize function
@@ -7768,7 +7766,7 @@ void FHeaderParser::SimplifiedClassParse(FUnrealSourceFile& SourceFile, const TC
 									if (((TrimmedStrLine.Find(Str_FArchive, ESearchCase::CaseSensitive, ESearchDir::FromStart, Pos)) != -1) ||
 										((TrimmedStrLine.Find(Str_FStructuredArchive, ESearchCase::CaseSensitive, ESearchDir::FromStart, Pos)) != -1))
 									{
-										FUHTException::Throwf(SourceFile, CurrentLine, TEXT("'%s' must not be inside preprocessor blocks, except for WITH_EDITORONLY_DATA"), *TrimmedStrLine);
+										FUHTMessage(SourceFile, CurrentLine).Throwf(TEXT("'%s' must not be inside preprocessor blocks, except for WITH_EDITORONLY_DATA"), *TrimmedStrLine);
 									}
 								}
 							}
@@ -7824,7 +7822,7 @@ void FHeaderParser::SimplifiedClassParse(FUnrealSourceFile& SourceFile, const TC
 			{
 				if (bFoundGeneratedInclude)
 				{
-					FUHTException::Throwf(SourceFile, CurrentLine, TEXT("#include found after .generated.h file - the .generated.h file should always be the last #include in a header"));
+					FUHTMessage(SourceFile, CurrentLine).Throwf(TEXT("#include found after .generated.h file - the .generated.h file should always be the last #include in a header"));
 				}
 
 				bFoundGeneratedInclude = DependsOnHeaderName.Contains(TEXT(".generated.h"));
@@ -7958,11 +7956,11 @@ void FHeaderParser::SimplifiedClassParse(FUnrealSourceFile& SourceFile, const TC
 				{
 					if (UInterfaceMacroDecl[10] != TEXT('('))
 					{
-						FUHTException::Throwf(SourceFile, CurrentLine, TEXT("Missing open parenthesis after UINTERFACE"));
+						FUHTMessage(SourceFile, CurrentLine).Throwf(TEXT("Missing open parenthesis after UINTERFACE"));
 					}
 
 					TSharedRef<FUnrealTypeDefinitionInfo> ClassDecl = Parser.ParseClassDeclaration(StartOfLine + (UInterfaceMacroDecl - Str), CurrentLine, true, TEXT("UINTERFACE"));
-					bFoundExportedClasses |= !EnumHasAnyFlags(ClassDecl->AsClassChecked().GetParsedClassFlags(), CLASS_NoExport); //@TODO - Use proper API once we have internal class flags
+					bFoundExportedClasses |= !ClassDecl->AsClassChecked().HasAnyClassFlags(CLASS_NoExport);
 					SourceFile.AddDefinedClass(MoveTemp(ClassDecl));
 				}
 			}
@@ -7973,11 +7971,11 @@ void FHeaderParser::SimplifiedClassParse(FUnrealSourceFile& SourceFile, const TC
 				{
 					if (UClassMacroDecl[6] != TEXT('('))
 					{
-						FUHTException::Throwf(SourceFile, CurrentLine, TEXT("Missing open parenthesis after UCLASS"));
+						FUHTMessage(SourceFile, CurrentLine).Throwf(TEXT("Missing open parenthesis after UCLASS"));
 					}
 
 					TSharedRef<FUnrealTypeDefinitionInfo> ClassDecl = Parser.ParseClassDeclaration(StartOfLine + (UClassMacroDecl - Str), CurrentLine, false, TEXT("UCLASS"));
-					bFoundExportedClasses |= !EnumHasAnyFlags(ClassDecl->AsClassChecked().GetParsedClassFlags(), CLASS_NoExport); //@TODO - Use proper API once we have internal class flags
+					bFoundExportedClasses |= !ClassDecl->AsClassChecked().HasAnyClassFlags(CLASS_NoExport);
 					SourceFile.AddDefinedClass(MoveTemp(ClassDecl));
 				}
 			}
@@ -7988,7 +7986,7 @@ void FHeaderParser::SimplifiedClassParse(FUnrealSourceFile& SourceFile, const TC
 				{
 					if (UEnumMacroDecl[5] != TEXT('('))
 					{
-						FUHTException::Throwf(SourceFile, CurrentLine, TEXT("Missing open parenthesis after UENUM"));
+						FUHTMessage(SourceFile, CurrentLine).Throwf(TEXT("Missing open parenthesis after UENUM"));
 					}
 
 					TSharedRef<FUnrealTypeDefinitionInfo> EnumDecl = Parser.ParseEnumDeclaration(StartOfLine + (UEnumMacroDecl - Str), CurrentLine);
@@ -8002,7 +8000,7 @@ void FHeaderParser::SimplifiedClassParse(FUnrealSourceFile& SourceFile, const TC
 				{
 					if (UStructMacroDecl[7] != TEXT('('))
 					{
-						FUHTException::Throwf(SourceFile, CurrentLine, TEXT("Missing open parenthesis after USTRUCT"));
+						FUHTMessage(SourceFile, CurrentLine).Throwf(TEXT("Missing open parenthesis after USTRUCT"));
 					}
 
 					TSharedRef<FUnrealTypeDefinitionInfo> StructDecl = Parser.ParseStructDeclaration(StartOfLine + (UStructMacroDecl - Str), CurrentLine);
@@ -8014,7 +8012,7 @@ void FHeaderParser::SimplifiedClassParse(FUnrealSourceFile& SourceFile, const TC
 
 	if (bFoundExportedClasses && !bFoundGeneratedInclude)
 	{
-		FUHTException::Throwf(Parser, TEXT("No #include found for the .generated.h file - the .generated.h file should always be the last #include in a header"));
+		Parser.Throwf(TEXT("No #include found for the .generated.h file - the .generated.h file should always be the last #include in a header"));
 	}
 }
 
@@ -8058,7 +8056,7 @@ TSharedRef<FUnrealTypeDefinitionInfo> FHeaderPreParser::ParseClassDeclaration(co
 	ParseInheritance(TEXT("class"), [this, &ClassNameWithoutPrefixStr, &ClassDef = *ClassDef](const TCHAR* ClassName, bool bIsSuperClass)
 	{
 		FString ClassNameStr(ClassName);
-		SourceFile.AddClassIncludeIfNeeded(InputLine, ClassNameWithoutPrefixStr, ClassNameStr);
+		SourceFile.AddClassIncludeIfNeeded(*this, ClassNameWithoutPrefixStr, ClassNameStr);
 		if (bIsSuperClass)
 		{
 			ClassDef.GetSuperStructInfo().Name = MoveTemp(ClassNameStr);
@@ -8097,7 +8095,7 @@ TSharedRef<FUnrealTypeDefinitionInfo> FHeaderPreParser::ParseEnumDeclaration(con
 	bool bReadEnumName = false;
 	if (!GetIdentifier(EnumToken))
 	{
-		FUHTException::Throwf(*this, TEXT("Missing identifier after UENUM()"));
+		Throwf(TEXT("Missing identifier after UENUM()"));
 	}
 
 	if (EnumToken.Matches(TEXT("namespace"), ESearchCase::CaseSensitive))
@@ -8110,7 +8108,7 @@ TSharedRef<FUnrealTypeDefinitionInfo> FHeaderPreParser::ParseEnumDeclaration(con
 	{
 		if (!GetIdentifier(EnumToken))
 		{
-			FUHTException::Throwf(*this, TEXT("Missing identifier after enum"));
+			Throwf(TEXT("Missing identifier after enum"));
 		}
 
 		if (EnumToken.Matches(TEXT("class"), ESearchCase::CaseSensitive) || EnumToken.Matches(TEXT("struct"), ESearchCase::CaseSensitive))
@@ -8128,13 +8126,13 @@ TSharedRef<FUnrealTypeDefinitionInfo> FHeaderPreParser::ParseEnumDeclaration(con
 	}
 	else
 	{
-		FUHTException::Throwf(*this, TEXT("UENUM() should be followed by \'enum\' or \'namespace\' keywords."));
+		Throwf(TEXT("UENUM() should be followed by \'enum\' or \'namespace\' keywords."));
 	}
 
 	// Get enumeration name.
 	if (!bReadEnumName)
 	{
-		FUHTException::Throwf(*this, TEXT("Missing enumeration name"));
+		Throwf(TEXT("Missing enumeration name"));
 	}
 
 	TSharedRef<FUnrealEnumDefinitionInfo> EnumDef = MakeShareable(new FUnrealEnumDefinitionInfo(SourceFile, InLineNumber, EnumToken.Identifier, FName(EnumToken.Identifier, FNAME_Add)));
@@ -8184,7 +8182,7 @@ TSharedRef<FUnrealTypeDefinitionInfo> FHeaderPreParser::ParseStructDeclaration(c
 	ParseInheritance(TEXT("struct"), [this, &StructNameStripped, &StructDef = *StructDef](const TCHAR* StructName, bool bIsSuperClass)
 	{
 		FString StructNameStr(StructName);
-		SourceFile.AddScriptStructIncludeIfNeeded(InputLine, StructNameStripped, StructNameStr);
+		SourceFile.AddScriptStructIncludeIfNeeded(*this, StructNameStripped, StructNameStr);
 		if (bIsSuperClass)
 		{
 			StructDef.GetSuperStructInfo().Name = MoveTemp(StructNameStr);
@@ -8435,12 +8433,12 @@ void FHeaderParser::PostPopNestClass(FUnrealClassDefinitionInfo& CurrentClassDef
 
 					if (InterfaceFunctionDef->HasAnyFunctionFlags(FUNC_Event) && !ClassFunctionDef->HasAnyFunctionFlags(FUNC_Event))
 					{
-						FUHTException::Throwf(*this, TEXT("Implementation of function '%s::%s' must be declared as 'event' to match declaration in interface '%s'"), *ClassFunctionDef->GetOuter()->GetName(), *ClassFunctionDef->GetName(), *InterfaceDef->GetName());
+						Throwf(TEXT("Implementation of function '%s::%s' must be declared as 'event' to match declaration in interface '%s'"), *ClassFunctionDef->GetOuter()->GetName(), *ClassFunctionDef->GetName(), *InterfaceDef->GetName());
 					}
 
 					if (InterfaceFunctionDef->HasAnyFunctionFlags(FUNC_Delegate) && !ClassFunctionDef->HasAnyFunctionFlags(FUNC_Delegate))
 					{
-						FUHTException::Throwf(*this, TEXT("Implementation of function '%s::%s' must be declared as 'delegate' to match declaration in interface '%s'"), *ClassFunctionDef->GetOuter()->GetName(), *ClassFunctionDef->GetName(), *InterfaceDef->GetName());
+						Throwf(TEXT("Implementation of function '%s::%s' must be declared as 'delegate' to match declaration in interface '%s'"), *ClassFunctionDef->GetOuter()->GetName(), *ClassFunctionDef->GetName(), *InterfaceDef->GetName());
 					}
 
 					// Making sure all the parameters match up correctly
@@ -8448,7 +8446,7 @@ void FHeaderParser::PostPopNestClass(FUnrealClassDefinitionInfo& CurrentClassDef
 
 					if (ClassFunctionDef->GetProperties().Num() != InterfaceFunctionDef->GetProperties().Num())
 					{
-						FUHTException::Throwf(*this, TEXT("Implementation of function '%s' conflicts with interface '%s' - different number of parameters (%i/%i)"), *InterfaceFunctionDef->GetName(), *InterfaceDef->GetName(), ClassFunctionDef->GetProperties().Num(), InterfaceFunctionDef->GetProperties().Num());
+						Throwf(TEXT("Implementation of function '%s' conflicts with interface '%s' - different number of parameters (%i/%i)"), *InterfaceFunctionDef->GetName(), *InterfaceDef->GetName(), ClassFunctionDef->GetProperties().Num(), InterfaceFunctionDef->GetProperties().Num());
 					}
 
 					for (int32 Index = 0, End = ClassFunctionDef->GetProperties().Num(); Index != End; ++Index)
@@ -8459,11 +8457,11 @@ void FHeaderParser::PostPopNestClass(FUnrealClassDefinitionInfo& CurrentClassDef
 						{
 							if (InterfaceParamDef.HasAnyPropertyFlags(CPF_ReturnParm))
 							{
-								FUHTException::Throwf(*this, TEXT("Implementation of function '%s' conflicts only by return type with interface '%s'"), *InterfaceFunctionDef->GetName(), *InterfaceDef->GetName());
+								Throwf(TEXT("Implementation of function '%s' conflicts only by return type with interface '%s'"), *InterfaceFunctionDef->GetName(), *InterfaceDef->GetName());
 							}
 							else
 							{
-								FUHTException::Throwf(*this, TEXT("Implementation of function '%s' conflicts with interface '%s' - parameter %i '%s'"), *InterfaceFunctionDef->GetName(), *InterfaceDef->GetName(), End, *InterfaceParamDef.GetName());
+								Throwf(TEXT("Implementation of function '%s' conflicts with interface '%s' - parameter %i '%s'"), *InterfaceFunctionDef->GetName(), *InterfaceDef->GetName(), End, *InterfaceParamDef.GetName());
 							}
 						}
 					}
@@ -8481,7 +8479,7 @@ void FHeaderParser::PostPopNestClass(FUnrealClassDefinitionInfo& CurrentClassDef
 					&& !InterfaceFunctionDef->HasAnyFunctionFlags(FUNC_BlueprintEvent)
 					&& !InterfaceDef->HasMetaData(NAME_CannotImplementInterfaceInBlueprint))  // FBlueprintMetadata::MD_CannotImplementInterfaceInBlueprint
 				{
-					FUHTException::Throwf(*this, TEXT("Missing UFunction implementation of function '%s' from interface '%s'.  This function needs a UFUNCTION() declaration."), *InterfaceFunctionDef->GetName(), *InterfaceDef->GetName());
+					Throwf(TEXT("Missing UFunction implementation of function '%s' from interface '%s'.  This function needs a UFUNCTION() declaration."), *InterfaceFunctionDef->GetName(), *InterfaceDef->GetName());
 				}
 			}
 		}
@@ -8508,7 +8506,7 @@ void FHeaderParser::PostPopNestInterface(FUnrealClassDefinitionInfo& CurrentInte
 	}
 }
 
-FDocumentationPolicy FHeaderParser::GetDocumentationPolicyFromName(const FUHTExceptionContext& Context, const FString& PolicyName)
+FDocumentationPolicy FHeaderParser::GetDocumentationPolicyFromName(const FUHTMessageProvider& Context, const FString& PolicyName)
 {
 	FDocumentationPolicy DocumentationPolicy;
 	if (FCString::Strcmp(*PolicyName, TEXT("Strict")) == 0)
@@ -8521,7 +8519,7 @@ FDocumentationPolicy FHeaderParser::GetDocumentationPolicyFromName(const FUHTExc
 	}
 	else
 	{
-		FUHTException::Throwf(Context, TEXT("Documentation Policy '%s' not yet supported"), *PolicyName);
+		Context.Throwf(TEXT("Documentation Policy '%s' not yet supported"), *PolicyName);
 	}
 	return DocumentationPolicy;
 }
@@ -8558,7 +8556,7 @@ void FHeaderParser::CheckDocumentationPolicyForEnum(FUnrealEnumDefinitionInfo& E
 		const FString* EnumToolTip = MetaData.Find(NAME_ToolTip);
 		if (EnumToolTip == nullptr)
 		{
-			UE_LOG_ERROR_UHT(*this, TEXT("Enum '%s' does not provide a tooltip / comment (DocumentationPolicy)."), *EnumDef.GetName());
+			LogError(TEXT("Enum '%s' does not provide a tooltip / comment (DocumentationPolicy)."), *EnumDef.GetName());
 		}
 	}
 
@@ -8574,14 +8572,14 @@ void FHeaderParser::CheckDocumentationPolicyForEnum(FUnrealEnumDefinitionInfo& E
 		const FString* ToolTip = Entry.Find(NAME_ToolTip);
 		if (ToolTip == nullptr)
 		{
-			UE_LOG_ERROR_UHT(*this, TEXT("Enum entry '%s::%s' does not provide a tooltip / comment (DocumentationPolicy)."), *EnumDef.GetName(), *(*EntryName));
+			LogError(TEXT("Enum entry '%s::%s' does not provide a tooltip / comment (DocumentationPolicy)."), *EnumDef.GetName(), *(*EntryName));
 			continue;
 		}
 
 		const FString* ExistingEntry = ToolTipToEntry.Find(*ToolTip);
 		if (ExistingEntry != nullptr)
 		{
-			UE_LOG_ERROR_UHT(*this, TEXT("Enum entries '%s::%s' and '%s::%s' have identical tooltips / comments (DocumentationPolicy)."), *EnumDef.GetName(), *(*ExistingEntry), *EnumDef.GetName(), *(*EntryName));
+			LogError(TEXT("Enum entries '%s::%s' and '%s::%s' have identical tooltips / comments (DocumentationPolicy)."), *EnumDef.GetName(), *(*ExistingEntry), *EnumDef.GetName(), *(*EntryName));
 		}
 		ToolTipToEntry.Add(*ToolTip, *EntryName);
 	}
@@ -8602,7 +8600,7 @@ void FHeaderParser::CheckDocumentationPolicyForStruct(FUnrealStructDefinitionInf
 
 		if (ClassTooltip.IsEmpty() || ClassTooltip.Equals(StructDef.GetName()))
 		{
-			UE_LOG_ERROR_UHT(*this, TEXT("Struct '%s' does not provide a tooltip / comment (DocumentationPolicy)."), *StructDef.GetName());
+			LogError(TEXT("Struct '%s' does not provide a tooltip / comment (DocumentationPolicy)."), *StructDef.GetName());
 		}
 	}
 
@@ -8614,13 +8612,13 @@ void FHeaderParser::CheckDocumentationPolicyForStruct(FUnrealStructDefinitionInf
 			FString ToolTip = PropertyDef->GetToolTipText().ToString();
 			if (ToolTip.IsEmpty() || ToolTip.Equals(PropertyDef->GetDisplayNameText().ToString()))
 			{
-				UE_LOG_ERROR_UHT(*this, TEXT("Property '%s::%s' does not provide a tooltip / comment (DocumentationPolicy)."), *StructDef.GetName(), *PropertyDef->GetName());
+				LogError(TEXT("Property '%s::%s' does not provide a tooltip / comment (DocumentationPolicy)."), *StructDef.GetName(), *PropertyDef->GetName());
 				continue;
 			}
 			const FName* ExistingPropertyName = ToolTipToPropertyName.Find(ToolTip);
 			if (ExistingPropertyName != nullptr)
 			{
-				UE_LOG_ERROR_UHT(*this, TEXT("Property '%s::%s' and '%s::%s' are using identical tooltips (DocumentationPolicy)."), *StructDef.GetName(), *ExistingPropertyName->ToString(), *StructDef.GetName(), *PropertyDef->GetName());
+				LogError(TEXT("Property '%s::%s' and '%s::%s' are using identical tooltips (DocumentationPolicy)."), *StructDef.GetName(), *ExistingPropertyName->ToString(), *StructDef.GetName(), *PropertyDef->GetName());
 			}
 			ToolTipToPropertyName.Add(MoveTemp(ToolTip), PropertyDef->GetFName());
 		}
@@ -8637,7 +8635,7 @@ void FHeaderParser::CheckDocumentationPolicyForStruct(FUnrealStructDefinitionInf
 
 				if(!CheckUIMinMaxRangeFromMetaData(UIMin, UIMax))
 				{
-					UE_LOG_ERROR_UHT(*this, TEXT("Property '%s::%s' does not provide a valid UIMin / UIMax (DocumentationPolicy)."), *StructDef.GetName(), *PropertyDef->GetName());
+					LogError(TEXT("Property '%s::%s' does not provide a valid UIMin / UIMax (DocumentationPolicy)."), *StructDef.GetName(), *PropertyDef->GetName());
 				}
 			}
 		}
@@ -8654,13 +8652,13 @@ void FHeaderParser::CheckDocumentationPolicyForStruct(FUnrealStructDefinitionInf
 				FString ToolTip = FunctionDef->GetToolTipText().ToString();
 				if (ToolTip.IsEmpty())
 				{
-					UE_LOG_ERROR_UHT(*this, TEXT("Function '%s::%s' does not provide a tooltip / comment (DocumentationPolicy)."), *ClassDef->GetName(), *FunctionDef->GetName());
+					LogError(TEXT("Function '%s::%s' does not provide a tooltip / comment (DocumentationPolicy)."), *ClassDef->GetName(), *FunctionDef->GetName());
 					continue;
 				}
 				const FName* ExistingFuncName = ToolTipToFunc.Find(ToolTip);
 				if (ExistingFuncName != nullptr)
 				{
-					UE_LOG_ERROR_UHT(*this, TEXT("Functions '%s::%s' and '%s::%s' uses identical tooltips / comments (DocumentationPolicy)."), *ClassDef->GetName(), *(*ExistingFuncName).ToString(), *ClassDef->GetName(), *FunctionDef->GetName());
+					LogError(TEXT("Functions '%s::%s' and '%s::%s' uses identical tooltips / comments (DocumentationPolicy)."), *ClassDef->GetName(), *(*ExistingFuncName).ToString(), *ClassDef->GetName(), *FunctionDef->GetName());
 				}
 				ToolTipToFunc.Add(MoveTemp(ToolTip), FunctionDef->GetFName());
 			}
@@ -8684,7 +8682,7 @@ void FHeaderParser::CheckDocumentationPolicyForFunc(FUnrealClassDefinitionInfo&
 		const FString* FunctionTooltip = FunctionDef.FindMetaData(NAME_ToolTip);
 		if (FunctionTooltip == nullptr)
 		{
-			UE_LOG_ERROR_UHT(*this, TEXT("Function '%s::%s' does not provide a tooltip / comment (DocumentationPolicy)."), *ClassDef.GetName(), *FunctionDef.GetName());
+			LogError(TEXT("Function '%s::%s' does not provide a tooltip / comment (DocumentationPolicy)."), *ClassDef.GetName(), *FunctionDef.GetName());
 		}
 	}
 
@@ -8693,7 +8691,7 @@ void FHeaderParser::CheckDocumentationPolicyForFunc(FUnrealClassDefinitionInfo&
 		const FString* FunctionComment = FunctionDef.FindMetaData(NAME_Comment);
 		if (FunctionComment == nullptr)
 		{
-			UE_LOG_ERROR_UHT(*this, TEXT("Function '%s::%s' does not provide a comment (DocumentationPolicy)."), *ClassDef.GetName(), *FunctionDef.GetName());
+			LogError(TEXT("Function '%s::%s' does not provide a comment (DocumentationPolicy)."), *ClassDef.GetName(), *FunctionDef.GetName());
 			return;
 		}
 		
@@ -8723,7 +8721,7 @@ void FHeaderParser::CheckDocumentationPolicyForFunc(FUnrealClassDefinitionInfo&
 				const FString* ParamToolTip = ParamToolTips.Find(ParamName);
 				if (ParamToolTip == nullptr)
 				{
-					UE_LOG_ERROR_UHT(*this, TEXT("Function '%s::%s' doesn't provide a tooltip for parameter '%s' (DocumentationPolicy)."), *ClassDef.GetName(), *FunctionDef.GetName(), *ParamName.ToString());
+					LogError(TEXT("Function '%s::%s' doesn't provide a tooltip for parameter '%s' (DocumentationPolicy)."), *ClassDef.GetName(), *FunctionDef.GetName(), *ParamName.ToString());
 				}
 				ExistingFields.Add(ParamName);
 			}
@@ -8738,7 +8736,7 @@ void FHeaderParser::CheckDocumentationPolicyForFunc(FUnrealClassDefinitionInfo&
 				}
 				if (!ExistingFields.Contains(ParamName))
 				{
-					UE_LOG_ERROR_UHT(*this, TEXT("Function '%s::%s' provides a tooltip for an unknown parameter '%s' (DocumentationPolicy)."), *ClassDef.GetName(), *FunctionDef.GetName(), *Pair.Key.ToString());
+					LogError(TEXT("Function '%s::%s' provides a tooltip for an unknown parameter '%s' (DocumentationPolicy)."), *ClassDef.GetName(), *FunctionDef.GetName(), *Pair.Key.ToString());
 				}
 			}
 
@@ -8754,7 +8752,7 @@ void FHeaderParser::CheckDocumentationPolicyForFunc(FUnrealClassDefinitionInfo&
 				const FName* ExistingParam = ToolTipToParam.Find(Pair.Value);
 				if (ExistingParam != nullptr)
 				{
-					UE_LOG_ERROR_UHT(*this, TEXT("Function '%s::%s' uses identical tooltips for parameters '%s' and '%s' (DocumentationPolicy)."), *ClassDef.GetName(), *FunctionDef.GetName(), *ExistingParam->ToString(), *Pair.Key.ToString());
+					LogError(TEXT("Function '%s::%s' uses identical tooltips for parameters '%s' and '%s' (DocumentationPolicy)."), *ClassDef.GetName(), *FunctionDef.GetName(), *ExistingParam->ToString(), *Pair.Key.ToString());
 				}
 				ToolTipToParam.Add(MoveTemp(Pair.Value), Pair.Key);
 			}
@@ -8830,9 +8828,12 @@ void FHeaderParser::ConditionalLogPointerUsage(EPointerMemberBehavior PointerMem
 	}
 }
 
-template <class TFunctionType>
-FUnrealFunctionDefinitionInfo& CreateFunctionImpl(FUnrealSourceFile& SourceFile, int32 InputLine, FFuncInfo&& FuncInfo, EFunctionType InFunctionType, FUnrealObjectDefinitionInfo& OuterDef, FScope* CurrentScope)
+FUnrealFunctionDefinitionInfo& FHeaderParser::CreateFunction(FFuncInfo&& FuncInfo, EFunctionType InFunctionType) const
 {
+	FScope* CurrentScope = GetCurrentScope();
+	FUnrealObjectDefinitionInfo& OuterDef = IsInAClass() ? static_cast<FUnrealObjectDefinitionInfo&>(GetCurrentClassDef()) : SourceFile.GetPackageDef();
+
+
 	// Allocate local property frame, push nesting level and verify
 	// uniqueness at this scope level.
 	{
@@ -8842,7 +8843,7 @@ FUnrealFunctionDefinitionInfo& CreateFunctionImpl(FUnrealSourceFile& SourceFile,
 			FUnrealFieldDefinitionInfo* Type = *TypeIterator;
 			if (Type->GetFName() == FuncInfo.Function.Identifier)
 			{
-				FUHTException::Throwf(SourceFile, InputLine, TEXT("'%s' conflicts with '%s'"), FuncInfo.Function.Identifier, *Type->GetFullName());
+				Throwf(TEXT("'%s' conflicts with '%s'"), FuncInfo.Function.Identifier, *Type->GetFullName());
 			}
 		}
 	}
@@ -8865,15 +8866,3 @@ FUnrealFunctionDefinitionInfo& CreateFunctionImpl(FUnrealSourceFile& SourceFile,
 
 	return FuncDef;
 }
-
-FUnrealFunctionDefinitionInfo& FHeaderParser::CreateFunction(FFuncInfo&& FuncInfo) const
-{
-	return CreateFunctionImpl<UFunction>(SourceFile, InputLine, MoveTemp(FuncInfo), EFunctionType::Function, GetCurrentClassDef(), GetCurrentScope());
-}
-
-template<class T>
-FUnrealFunctionDefinitionInfo& FHeaderParser::CreateDelegateFunction(FFuncInfo&& FuncInfo, EFunctionType InFunctionType) const
-{
-	FUnrealObjectDefinitionInfo& OuterDef = IsInAClass() ? static_cast<FUnrealObjectDefinitionInfo&>(GetCurrentClassDef()) : SourceFile.GetPackageDef();
-	return CreateFunctionImpl<T>(SourceFile, InputLine, MoveTemp(FuncInfo), InFunctionType, OuterDef, GetCurrentScope());
-}
diff --git a/Engine/Source/Programs/UnrealHeaderTool/Private/HeaderParser.h b/Engine/Source/Programs/UnrealHeaderTool/Private/HeaderParser.h
--- a/Engine/Source/Programs/UnrealHeaderTool/Private/HeaderParser.h
+++ b/Engine/Source/Programs/UnrealHeaderTool/Private/HeaderParser.h
@@ -328,8 +328,8 @@ protected:
 
 public:
 	// Throws if a specifier value wasn't provided
-	static void RequireSpecifierValue(const FUHTExceptionContext& Context, const FPropertySpecifier& Specifier, bool bRequireExactlyOne = false);
-	static FString RequireExactlyOneSpecifierValue(const FUHTExceptionContext& Context, const FPropertySpecifier& Specifier);
+	static void RequireSpecifierValue(const FUHTMessageProvider& Context, const FPropertySpecifier& Specifier, bool bRequireExactlyOne = false);
+	static FString RequireExactlyOneSpecifierValue(const FUHTMessageProvider& Context, const FPropertySpecifier& Specifier);
 
 	/**
 	* Find a field in the specified context.  Starts with the specified scope, then iterates
@@ -406,13 +406,7 @@ protected:
 	/**
 	 * Create new function object based on given info structure.
 	 */
-	FUnrealFunctionDefinitionInfo& CreateFunction(FFuncInfo&& FuncInfo) const;
-
-	/**
-	 * Create new delegate function object based on given info structure.
-	 */
-	template<typename T>
-	FUnrealFunctionDefinitionInfo& CreateDelegateFunction(FFuncInfo&& FuncInfo, EFunctionType InFunctionType) const;
+	FUnrealFunctionDefinitionInfo& CreateFunction(FFuncInfo&& FuncInfo, EFunctionType InFunctionType) const;
 
 	FUnrealClassDefinitionInfo& CompileClassDeclaration();
 	FUnrealFunctionDefinitionInfo& CompileDelegateDeclaration(const TCHAR* DelegateIdentifier, EDelegateSpecifierAction::Type SpecifierAction = EDelegateSpecifierAction::DontParse);
@@ -601,7 +595,7 @@ private:
 	void VerifyRepNotifyCallback(FUnrealPropertyDefinitionInfo& PropertyDef, FUnrealFunctionDefinitionInfo* TargetFuncDef);
 
 	// Constructs the policy from a string
-	static FDocumentationPolicy GetDocumentationPolicyFromName(const FUHTExceptionContext& Context, const FString& PolicyName);
+	static FDocumentationPolicy GetDocumentationPolicyFromName(const FUHTMessageProvider& Context, const FString& PolicyName);
 
 	// Constructs the policy for documentation checks for a given struct
 	static FDocumentationPolicy GetDocumentationPolicyForStruct(FUnrealStructDefinitionInfo& StructDef);
diff --git a/Engine/Source/Programs/UnrealHeaderTool/Private/IScriptGeneratorPluginInterface.cpp b/Engine/Source/Programs/UnrealHeaderTool/Private/IScriptGeneratorPluginInterface.cpp
--- a/Engine/Source/Programs/UnrealHeaderTool/Private/IScriptGeneratorPluginInterface.cpp
+++ b/Engine/Source/Programs/UnrealHeaderTool/Private/IScriptGeneratorPluginInterface.cpp
@@ -26,7 +26,7 @@ EBuildModuleType::Type EBuildModuleType::Parse(const TCHAR* Value)
 	int32 TypeIndex = Algo::FindSortedStringCaseInsensitive(Value, AlphabetizedTypes);
 	if (TypeIndex < 0)
 	{
-		FUHTException::Throwf(FString(GManifestFilename), 1, TEXT("Unrecognized EBuildModuleType name: %s"), Value);
+		FUHTMessage(GManifestFilename).Throwf(TEXT("Unrecognized EBuildModuleType name: %s"), Value);
 	}
 
 	static EBuildModuleType::Type AlphabetizedValues[] = {
@@ -60,7 +60,7 @@ EPackageOverrideType::Type EPackageOverrideType::Parse(const TCHAR* Value)
 	int32 TypeIndex = Algo::FindSortedStringCaseInsensitive(Value, AlphabetizedTypes);
 	if (TypeIndex < 0)
 	{
-		FUHTException::Throwf(FString(GManifestFilename), 1, TEXT("Unrecognized EPackageOverrideType name: %s"), Value);
+		FUHTMessage(GManifestFilename).Throwf(TEXT("Unrecognized EPackageOverrideType name: %s"), Value);
 	}
 
 	static EPackageOverrideType::Type AlphabetizedValues[] = {
@@ -73,4 +73,4 @@ EPackageOverrideType::Type EPackageOverrideType::Parse(const TCHAR* Value)
 	};
 
 	return AlphabetizedValues[TypeIndex];
-}
+}
\ No newline at end of file
diff --git a/Engine/Source/Programs/UnrealHeaderTool/Private/Manifest.cpp b/Engine/Source/Programs/UnrealHeaderTool/Private/Manifest.cpp
--- a/Engine/Source/Programs/UnrealHeaderTool/Private/Manifest.cpp
+++ b/Engine/Source/Programs/UnrealHeaderTool/Private/Manifest.cpp
@@ -30,7 +30,7 @@ namespace
 	{
 		if (JsonValue->Type != TJsonFieldType<T>::Value)
 		{
-			FUHTException::Throwf(FString(GManifestFilename), 1, TEXT("'%s' is the wrong type"), Outer);
+			FUHTMessage(GManifestFilename).Throwf(TEXT("'%s' is the wrong type"), Outer);
 		}
 
 		JsonValue->AsArgumentType(OutVal);
@@ -43,12 +43,12 @@ namespace
 
 		if (!JsonValue)
 		{
-			FUHTException::Throwf(FString(GManifestFilename), 1, TEXT("Unable to find field '%s' in '%s'"), FieldName, Outer);
+			FUHTMessage(GManifestFilename).Throwf(TEXT("Unable to find field '%s' in '%s'"), FieldName, Outer);
 		}
 
 		if ((*JsonValue)->Type != TJsonFieldType<T>::Value)
 		{
-			FUHTException::Throwf(FString(GManifestFilename), 1, TEXT("Field '%s' in '%s' is the wrong type"), Outer);
+			FUHTMessage(GManifestFilename).Throwf(TEXT("Field '%s' in '%s' is the wrong type"), Outer);
 		}
 
 		(*JsonValue)->AsArgumentType(OutVal);
@@ -73,7 +73,7 @@ FManifest FManifest::LoadFromFile(const FString& Filename)
 
 	if (!FFileHelper::LoadFileToString(Json, *Filename))
 	{
-		FUHTException::Throwf(FString(Filename), 1, TEXT("Unable to load manifest: %s"), *Filename);
+		FUHTMessage(Filename).Throwf(TEXT("Unable to load manifest: %s"), *Filename);
 	}
 
 	TSharedPtr<FJsonObject> RootObject = TSharedPtr<FJsonObject>();
@@ -81,7 +81,7 @@ FManifest FManifest::LoadFromFile(const FString& Filename)
 
 	if (!FJsonSerializer::Deserialize(Reader, RootObject))
 	{
-		FUHTException::Throwf(FString(Filename), 1, TEXT("Manifest is malformed: %s"), *Filename);
+		FUHTMessage(Filename).Throwf(TEXT("Manifest is malformed: %s"), *Filename);
 	}
 
 	TArray<TSharedPtr<FJsonValue>> ModulesArray;
diff --git a/Engine/Source/Programs/UnrealHeaderTool/Private/ParserHelper.cpp b/Engine/Source/Programs/UnrealHeaderTool/Private/ParserHelper.cpp
--- a/Engine/Source/Programs/UnrealHeaderTool/Private/ParserHelper.cpp
+++ b/Engine/Source/Programs/UnrealHeaderTool/Private/ParserHelper.cpp
@@ -353,7 +353,7 @@ void FFuncInfo::SetFunctionNames(FUnrealFunctionDefinitionInfo& FunctionDef)
 			}
 			else if (CppImplName == FunctionName)
 			{
-				FUHTException::Throwf(FunctionDef, TEXT("Native implementation function must be different than original function name."));
+				FunctionDef.Throwf(TEXT("Native implementation function must be different than original function name."));
 			}
 
 			if (CppValidationImplName.IsEmpty() && FunctionDef.HasAllFunctionFlags(FUNC_NetValidate))
@@ -362,7 +362,7 @@ void FFuncInfo::SetFunctionNames(FUnrealFunctionDefinitionInfo& FunctionDef)
 			}
 			else if (CppValidationImplName == FunctionName)
 			{
-				FUHTException::Throwf(FunctionDef, TEXT("Validation function must be different than original function name."));
+				FunctionDef.Throwf(TEXT("Validation function must be different than original function name."));
 			}
 		}
 	}
diff --git a/Engine/Source/Programs/UnrealHeaderTool/Private/PropertyTypes.cpp b/Engine/Source/Programs/UnrealHeaderTool/Private/PropertyTypes.cpp
--- a/Engine/Source/Programs/UnrealHeaderTool/Private/PropertyTypes.cpp
+++ b/Engine/Source/Programs/UnrealHeaderTool/Private/PropertyTypes.cpp
@@ -15,7 +15,7 @@ extern FUnrealClassDefinitionInfo* GUObjectDef;
 extern FUnrealClassDefinitionInfo* GUClassDef;
 extern FUnrealClassDefinitionInfo* GUInterfaceDef;
 
-void AddEditInlineMetaData(TMap<FName, FString>& MetaData);
+extern void AddEditInlineMetaData(TMap<FName, FString>& MetaData);
 
 // Following is the relationship between the property types
 //
@@ -204,7 +204,7 @@ namespace
 		case CPT_Delegate:					return FuncDispatch<FPropertyTypeTraitsDelegate>()(PropDef, std::forward<Args>(args)...);
 		case CPT_MulticastDelegate:			return FuncDispatch<FPropertyTypeTraitsMulticastDelegate>()(PropDef, std::forward<Args>(args)...);
 		case CPT_FieldPath:					return FuncDispatch<FPropertyTypeTraitsFieldPath>()(PropDef, std::forward<Args>(args)...);
-		default:							FUHTException::Throwf(PropDef, TEXT("Unknown property type %i"), (uint8)PropDef.GetPropertyBase().Type);
+		default:							PropDef.Throwf(TEXT("Unknown property type %i"), (uint8)PropDef.GetPropertyBase().Type);
 		}
 	}
 
@@ -977,7 +977,7 @@ struct FPropertyTypeTraitsEnum : public FPropertyTypeTraitsBase
 		}
 		if (EnumDef->HasMetaData(TEXT("Hidden"), EnumEntryIndex))
 		{
-			FUHTException::Throwf(PropDef, TEXT("Hidden enum entries cannot be used as default values: %s \"%s\" "), *PropDef.GetName(), *CppForm);
+			PropDef.Throwf(TEXT("Hidden enum entries cannot be used as default values: %s \"%s\" "), *PropDef.GetName(), *CppForm);
 		}
 		return true;
 	}
@@ -1646,7 +1646,7 @@ struct FPropertyTypeTraitsText : public FPropertyTypeTraitsBase
 		// These should be replaced with INVTEXT as FText::FromString can produce inconsistent keys
 		if (FDefaultValueHelper::StringFromCppString(CppForm, TEXT("FText::FromString"), OutForm))
 		{
-			UE_LOG_WARNING_UHT(PropDef, TEXT("FText::FromString should be replaced with INVTEXT for default parameter values"));
+			PropDef.LogWarning(TEXT("FText::FromString should be replaced with INVTEXT for default parameter values"));
 			return true;
 		}
 
@@ -1672,7 +1672,7 @@ struct FPropertyTypeTraitsText : public FPropertyTypeTraitsBase
 			{
 				if (ParsedTextNamespace.GetValue().Equals(UHTDummyNamespace))
 				{
-					FUHTException::Throwf(PropDef, TEXT("LOCTEXT default parameter values are not supported; use NSLOCTEXT instead: %s \"%s\" "), *PropDef.GetName(), *CppForm);
+					PropDef.Throwf(TEXT("LOCTEXT default parameter values are not supported; use NSLOCTEXT instead: %s \"%s\" "), *PropDef.GetName(), *CppForm);
 				}
 			}
 		}
diff --git a/Engine/Source/Programs/UnrealHeaderTool/Private/UnrealSourceFile.cpp b/Engine/Source/Programs/UnrealHeaderTool/Private/UnrealSourceFile.cpp
--- a/Engine/Source/Programs/UnrealHeaderTool/Private/UnrealSourceFile.cpp
+++ b/Engine/Source/Programs/UnrealHeaderTool/Private/UnrealSourceFile.cpp
@@ -148,7 +148,7 @@ FString FUnrealSourceFile::GetFileDefineName() const
 	return FString::Printf(TEXT("%s_%s_generated_h"), *API, *GetStrippedFilename());
 }
 
-void FUnrealSourceFile::AddClassIncludeIfNeeded(int32 InputLine, const FString& ClassNameWithoutPrefix, const FString& DependencyClassName)
+void FUnrealSourceFile::AddClassIncludeIfNeeded(FUHTMessageProvider& Context, const FString& ClassNameWithoutPrefix, const FString& DependencyClassName)
 {
 	if (!Algo::FindBy(GetDefinedClasses(), DependencyClassName, [](const TSharedRef<FUnrealTypeDefinitionInfo>& Info) { return Info->GetNameCPP(); }))
 	{
@@ -156,7 +156,7 @@ void FUnrealSourceFile::AddClassIncludeIfNeeded(int32 InputLine, const FString&
 
 		if (ClassNameWithoutPrefix == DependencyClassNameWithoutPrefix)
 		{
-			FUHTException::Throwf(*this, InputLine, TEXT("A class cannot inherit itself or a type with the same name but a different prefix"));
+			Context.Throwf(TEXT("A class cannot inherit itself or a type with the same name but a different prefix"));
 		}
 
 		FString StrippedDependencyName = DependencyClassName.Mid(1);
@@ -170,7 +170,7 @@ void FUnrealSourceFile::AddClassIncludeIfNeeded(int32 InputLine, const FString&
 	}
 }
 
-void FUnrealSourceFile::AddScriptStructIncludeIfNeeded(int32 InputLine, const FString& StructNameWithoutPrefix, const FString& DependencyStructName)
+void FUnrealSourceFile::AddScriptStructIncludeIfNeeded(FUHTMessageProvider& Context, const FString& StructNameWithoutPrefix, const FString& DependencyStructName)
 {
 	if (!Algo::FindBy(GetDefinedStructs(), DependencyStructName, [](const TSharedRef<FUnrealTypeDefinitionInfo>& Info) { return Info->GetNameCPP(); }))
 	{
@@ -178,7 +178,7 @@ void FUnrealSourceFile::AddScriptStructIncludeIfNeeded(int32 InputLine, const FS
 
 		if (StructNameWithoutPrefix == DependencyStructNameWithoutPrefix)
 		{
-			FUHTException::Throwf(*this, InputLine, TEXT("A struct cannot inherit itself or a type with the same name but a different prefix"));
+			Context.Throwf(TEXT("A struct cannot inherit itself or a type with the same name but a different prefix"));
 		}
 
 		FString StrippedDependencyName = DependencyStructName.Mid(1);
diff --git a/Engine/Source/Programs/UnrealHeaderTool/Private/UnrealSourceFile.h b/Engine/Source/Programs/UnrealHeaderTool/Private/UnrealSourceFile.h
--- a/Engine/Source/Programs/UnrealHeaderTool/Private/UnrealSourceFile.h
+++ b/Engine/Source/Programs/UnrealHeaderTool/Private/UnrealSourceFile.h
@@ -322,12 +322,12 @@ public:
 	/**
 	 * Add an include for a class if required
 	 */
-	void AddClassIncludeIfNeeded(int32 InputLine, const FString& ClassNameWithoutPrefix, const FString& DependencyClassName);
+	void AddClassIncludeIfNeeded(FUHTMessageProvider& Context, const FString& ClassNameWithoutPrefix, const FString& DependencyClassName);
 
 	/**
 	 * Add an include for a script struct if required
 	 */
-	void AddScriptStructIncludeIfNeeded(int32 InputLine, const FString& StructNameWithoutPrefix, const FString& DependencyStructName);
+	void AddScriptStructIncludeIfNeeded(FUHTMessageProvider& Context, const FString& StructNameWithoutPrefix, const FString& DependencyStructName);
 
 	/**
 	 * Add an include for a type definition if required.  TypeDef can be null.
@@ -553,3 +553,9 @@ private:
 	// Collection of all singletons found during code generation
 	TArray<FUnrealFieldDefinitionInfo*> Singletons;
 };
+
+template <>
+inline FString GetMessageFilename<FUnrealSourceFile>(const FUnrealSourceFile& Source)
+{
+	return Source.GetFilename();
+}
diff --git a/Engine/Source/Programs/UnrealHeaderTool/Private/UnrealTypeDefinitionInfo.cpp b/Engine/Source/Programs/UnrealHeaderTool/Private/UnrealTypeDefinitionInfo.cpp
--- a/Engine/Source/Programs/UnrealHeaderTool/Private/UnrealTypeDefinitionInfo.cpp
+++ b/Engine/Source/Programs/UnrealHeaderTool/Private/UnrealTypeDefinitionInfo.cpp
@@ -161,7 +161,7 @@ void FUHTMetaData::RemapMetaData(FUnrealTypeDefinitionInfo& TypeDef, TMap<FName,
 
 		if (NewKey != NAME_None)
 		{
-			UE_LOG_WARNING_UHT(TypeDef, TEXT("Remapping old metadata key '%s' to new key '%s', please update the declaration."), *CurrentKey.ToString(), *NewKey.ToString());
+			TypeDef.LogWarning(TEXT("Remapping old metadata key '%s' to new key '%s', please update the declaration."), *CurrentKey.ToString(), *NewKey.ToString());
 			CurrentKey = NewKey;
 		}
 	}
@@ -314,7 +314,7 @@ TSharedRef<FScope> FUnrealTypeDefinitionInfo::GetScope()
 {
 	if (!HasSource())
 	{
-		FUHTException::Throwf(*this, TEXT("Attempt to fetch the scope for type \"%s\" when it doesn't implement the method or there is no source file associated with the type."), *GetNameCPP());
+		Throwf(TEXT("Attempt to fetch the scope for type \"%s\" when it doesn't implement the method or there is no source file associated with the type."), *GetNameCPP());
 	}
 	return GetUnrealSourceFile().GetScope();
 }
@@ -328,7 +328,7 @@ uint32 FUnrealTypeDefinitionInfo::GetHash(bool bIncludeNoExport) const
 {
 	if (Hash == 0)
 	{
-		FUHTException::Throwf(*this, TEXT("Attempt to fetch the generated hash for type \"%s\" before it has been generated.  Include dependencies, topological sort, or job graph is in error."), *GetNameCPP());
+		Throwf(TEXT("Attempt to fetch the generated hash for type \"%s\" before it has been generated.  Include dependencies, topological sort, or job graph is in error."), *GetNameCPP());
 	}
 	return Hash;
 }
@@ -406,7 +406,7 @@ void FUnrealTypeDefinitionInfo::ValidateMetaDataFormat(const FName InKey, ECheck
 	{
 		if (!InValue.IsNumeric())
 		{
-			FUHTException::Throwf(*this, TEXT("Metadata value for '%s' is non-numeric : '%s'"), *InKey.ToString(), *InValue);
+			Throwf(TEXT("Metadata value for '%s' is non-numeric : '%s'"), *InKey.ToString(), *InValue);
 		}
 	}
 	break;
@@ -425,7 +425,7 @@ void FUnrealTypeDefinitionInfo::ValidateMetaDataFormat(const FName InKey, ECheck
 
 				if (ClassDef != nullptr && ClassDef->GetName() == TEXT("BlueprintFunctionLibrary"))
 				{
-					FUHTException::Throwf(*this, TEXT("%s doesn't make sense on static method '%s' in a blueprint function library"), *InKey.ToString(), *FuncDef->GetName());
+					Throwf(TEXT("%s doesn't make sense on static method '%s' in a blueprint function library"), *InKey.ToString(), *FuncDef->GetName());
 				}
 			}
 		}
@@ -469,7 +469,7 @@ void FUnrealTypeDefinitionInfo::ValidateMetaDataFormat(const FName InKey, ECheck
 
 			if (!bGoodParams)
 			{
-				UE_LOG_ERROR_UHT(*FuncDef, TEXT("Commutative associative binary operators must have exactly 2 parameters of the same type and a return value."));
+				FuncDef->LogError(TEXT("Commutative associative binary operators must have exactly 2 parameters of the same type and a return value."));
 			}
 		}
 	}
@@ -500,7 +500,7 @@ void FUnrealTypeDefinitionInfo::ValidateMetaDataFormat(const FName InKey, ECheck
 					FUnrealPropertyDefinitionInfo* FoundFieldDef = FHeaderParser::FindProperty(*FuncDef, *Entry, false);
 					if (!FoundFieldDef)
 					{
-						UE_LOG_ERROR_UHT(*FuncDef, TEXT("Function does not have a parameter named '%s'"), *Entry);
+						FuncDef->LogError(TEXT("Function does not have a parameter named '%s'"), *Entry);
 					}
 					else
 					{
@@ -514,7 +514,7 @@ void FUnrealTypeDefinitionInfo::ValidateMetaDataFormat(const FName InKey, ECheck
 							}
 							else
 							{
-								UE_LOG_ERROR_UHT(*FuncDef, TEXT("Function already specified an ExpandEnumAsExec input (%s), but '%s' is also an input parameter. Only one is permitted."), *FirstInputDef->GetName(), *Entry);
+								FuncDef->LogError(TEXT("Function already specified an ExpandEnumAsExec input (%s), but '%s' is also an input parameter. Only one is permitted."), *FirstInputDef->GetName(), *Entry);
 							}
 						}
 					}
@@ -530,7 +530,7 @@ void FUnrealTypeDefinitionInfo::ValidateMetaDataFormat(const FName InKey, ECheck
 		const FString ExperimentalValue(TEXT("Experimental"));
 		if ((InValue != EarlyAccessValue) && (InValue != ExperimentalValue))
 		{
-			FUHTException::Throwf(*this, TEXT("'%s' metadata was '%s' but it must be %s or %s"), *InKey.ToString(), *InValue, *ExperimentalValue, *EarlyAccessValue);
+			Throwf(TEXT("'%s' metadata was '%s' but it must be %s or %s"), *InKey.ToString(), *InValue, *ExperimentalValue, *EarlyAccessValue);
 		}
 	}
 	break;
@@ -542,13 +542,13 @@ void FUnrealTypeDefinitionInfo::ValidateMetaDataFormat(const FName InKey, ECheck
 		{
 			if (!PropDef->IsNumericOrNumericStaticArray() && !PropDef->IsStructOrStructStaticArray())
 			{
-				FUHTException::Throwf(*this, TEXT("'Units' meta data can only be applied to numeric and struct properties"));
+				Throwf(TEXT("'Units' meta data can only be applied to numeric and struct properties"));
 			}
 		}
 
 		if (!FUnitConversion::UnitFromString(*InValue))
 		{
-			FUHTException::Throwf(*this, TEXT("Unrecognized units (%s) specified for property '%s'"), *InValue, *GetFullName());
+			Throwf(TEXT("Unrecognized units (%s) specified for property '%s'"), *InValue, *GetFullName());
 		}
 	}
 	break;
@@ -558,7 +558,7 @@ void FUnrealTypeDefinitionInfo::ValidateMetaDataFormat(const FName InKey, ECheck
 		const TCHAR* StrictValue = TEXT("Strict");
 		if (InValue != StrictValue)
 		{
-			FUHTException::Throwf(*this, TEXT("'%s' metadata was '%s' but it must be %s"), *InKey.ToString(), *InValue, *StrictValue);
+			Throwf(TEXT("'%s' metadata was '%s' but it must be %s"), *InKey.ToString(), *InValue, *StrictValue);
 		}
 	}
 	break;
@@ -1746,7 +1746,7 @@ FUnrealClassDefinitionInfo* FUnrealClassDefinitionInfo::FindScriptClassOrThrow(c
 		return ResultDef;
 	}
 
-	FUHTException::Throwf(Parser, MoveTemp(ErrorMsg));
+	Parser.Throwf(MoveTemp(ErrorMsg));
 
 	// Unreachable, but compiler will warn otherwise because FUHTException::Throwf isn't declared noreturn
 	return nullptr;
@@ -1869,7 +1869,7 @@ void FUnrealClassDefinitionInfo::PostParseFinalizeInternal()
 	{
 		FString UName = GetNameCPP();
 		FString IName = TEXT("I") + UName.RightChop(1);
-		FUHTException::Throwf(*this, TEXT("UInterface '%s' parsed without a corresponding '%s'"), *UName, *IName);
+		Throwf(TEXT("UInterface '%s' parsed without a corresponding '%s'"), *UName, *IName);
 	}
 
 	FHeaderParser::CheckSparseClassData(*this);
@@ -2029,7 +2029,7 @@ void FUnrealClassDefinitionInfo::ParseClassProperties(TArray<FPropertySpecifier>
 
 		case EClassMetadataSpecifier::DependsOn:
 
-			FUHTException::Throwf(*this, TEXT("The dependsOn specifier is deprecated. Please use #include \"ClassHeaderFilename.h\" instead."));
+			Throwf(TEXT("The dependsOn specifier is deprecated. Please use #include \"ClassHeaderFilename.h\" instead."));
 			break;
 
 		case EClassMetadataSpecifier::MinimalAPI:
@@ -2231,7 +2231,7 @@ void FUnrealClassDefinitionInfo::ParseClassProperties(TArray<FPropertySpecifier>
 			break;
 
 		default:
-			FUHTException::Throwf(*this, TEXT("Unknown class specifier '%s'"), *PropSpecifier.Key);
+			Throwf(TEXT("Unknown class specifier '%s'"), *PropSpecifier.Key);
 		}
 	}
 	SetClassFlags(ParsedClassFlags);
@@ -2324,7 +2324,7 @@ void FUnrealClassDefinitionInfo::MergeAndValidateClassFlags(const FString& Decla
 	{
 		if (!HasAnyClassFlags(CLASS_NotPlaceable))
 		{
-			FUHTException::Throwf(*this, TEXT("The 'placeable' specifier is only allowed on classes which have a base class that's marked as not placeable. Classes are assumed to be placeable by default."));
+			Throwf(TEXT("The 'placeable' specifier is only allowed on classes which have a base class that's marked as not placeable. Classes are assumed to be placeable by default."));
 		}
 		ClearClassFlags(CLASS_NotPlaceable);
 		bWantsToBePlaceable = false; // Reset this flag after it's been merged
@@ -2342,21 +2342,21 @@ void FUnrealClassDefinitionInfo::MergeAndValidateClassFlags(const FString& Decla
 		// don't allow actor classes to be declared editinlinenew
 		if (IsActorClass(*this))
 		{
-			FUHTException::Throwf(*this, TEXT("Invalid class attribute: Creating actor instances via the property window is not allowed"));
+			Throwf(TEXT("Invalid class attribute: Creating actor instances via the property window is not allowed"));
 		}
 	}
 
 	// Make sure both RequiredAPI and MinimalAPI aren't specified
 	if (HasAllClassFlags(CLASS_MinimalAPI | CLASS_RequiredAPI))
 	{
-		FUHTException::Throwf(*this, TEXT("MinimalAPI cannot be specified when the class is fully exported using a MODULENAME_API macro"));
+		Throwf(TEXT("MinimalAPI cannot be specified when the class is fully exported using a MODULENAME_API macro"));
 	}
 
 	// All classes must start with a valid Unreal prefix
 	const FString ExpectedClassName = GetNameWithPrefix();
 	if (DeclaredClassName != ExpectedClassName)
 	{
-		FUHTException::Throwf(*this, TEXT("Class name '%s' is invalid, should be identified as '%s'"), *DeclaredClassName, *ExpectedClassName);
+		Throwf(TEXT("Class name '%s' is invalid, should be identified as '%s'"), *DeclaredClassName, *ExpectedClassName);
 	}
 
 	// This check only works if we already have an object.  This has to be moved to the engine type create code
@@ -2370,7 +2370,7 @@ void FUnrealClassDefinitionInfo::MergeAndValidateClassFlags(const FString& Decla
 		{
 			if (!HasAnyClassFlags(CLASS_Intrinsic) && EnumHasAnyFlags(InitialEngineClassFlags, CLASS_Native))
 			{
-				FUHTException::Throwf(*this, TEXT("'noexport': Must include CLASS_NoExport in native class declaration"));
+				Throwf(TEXT("'noexport': Must include CLASS_NoExport in native class declaration"));
 			}
 		}
 
@@ -2378,12 +2378,12 @@ void FUnrealClassDefinitionInfo::MergeAndValidateClassFlags(const FString& Decla
 		{
 			if (HasAnyClassFlags(CLASS_NoExport))
 			{
-				FUHTException::Throwf(*this, TEXT("'abstract': NoExport class missing abstract keyword from class declaration (must change C++ version first)"));
+				Throwf(TEXT("'abstract': NoExport class missing abstract keyword from class declaration (must change C++ version first)"));
 				SetClassFlags(CLASS_Abstract);
 			}
 			else if (IsNative())
 			{
-				FUHTException::Throwf(*this, TEXT("'abstract': missing abstract keyword from class declaration - class will no longer be exported as abstract"));
+				Throwf(TEXT("'abstract': missing abstract keyword from class declaration - class will no longer be exported as abstract"));
 			}
 		}
 	}
@@ -2400,12 +2400,12 @@ void FUnrealClassDefinitionInfo::SetAndValidateConfigName()
 			FUnrealClassDefinitionInfo* SuperClassDef = GetSuperClass();
 			if (!SuperClassDef)
 			{
-				FUHTException::Throwf(*this, TEXT("Cannot inherit config filename: %s has no super class"), *GetName());
+				Throwf(TEXT("Cannot inherit config filename: %s has no super class"), *GetName());
 			}
 
 			if (SuperClassDef->GetClassConfigName() == NAME_None)
 			{
-				FUHTException::Throwf(*this, TEXT("Cannot inherit config filename: parent class %s is not marked config."), *SuperClassDef->GetPathName());
+				Throwf(TEXT("Cannot inherit config filename: parent class %s is not marked config."), *SuperClassDef->GetPathName());
 			}
 		}
 		else
@@ -2429,11 +2429,11 @@ void FUnrealClassDefinitionInfo::SetAndValidateWithinClass()
 		FUnrealClassDefinitionInfo* RequiredWithinClassDef = FUnrealClassDefinitionInfo::FindClass(*ClassWithinStr);
 		if (!RequiredWithinClassDef)
 		{
-			FUHTException::Throwf(*this, TEXT("Within class '%s' not found."), *ClassWithinStr);
+			Throwf(TEXT("Within class '%s' not found."), *ClassWithinStr);
 		}
 		if (RequiredWithinClassDef->IsChildOf(*GUInterfaceDef))
 		{
-			FUHTException::Throwf(*this, TEXT("Classes cannot be 'within' interfaces"));
+			Throwf(TEXT("Classes cannot be 'within' interfaces"));
 		}
 		else if (ClassWithin == NULL || ClassWithin == GUObjectDef || RequiredWithinClassDef->IsChildOf(*ClassWithin))
 		{
@@ -2441,7 +2441,7 @@ void FUnrealClassDefinitionInfo::SetAndValidateWithinClass()
 		}
 		else if (ClassWithin != RequiredWithinClassDef)
 		{
-			FUHTException::Throwf(*this, TEXT("%s must be within %s, not %s"), *GetPathName(), *ClassWithin->GetPathName(), *RequiredWithinClassDef->GetPathName());
+			Throwf(TEXT("%s must be within %s, not %s"), *GetPathName(), *ClassWithin->GetPathName(), *RequiredWithinClassDef->GetPathName());
 		}
 	}
 	else
@@ -2454,7 +2454,7 @@ void FUnrealClassDefinitionInfo::SetAndValidateWithinClass()
 
 	if (!ClassWithin->IsChildOf(*ExpectedWithinDef))
 	{
-		FUHTException::Throwf(*this, TEXT("Parent class declared within '%s'.  Cannot override within class with '%s' since it isn't a child"), *ExpectedWithinDef->GetName(), *ClassWithin->GetName());
+		Throwf(TEXT("Parent class declared within '%s'.  Cannot override within class with '%s' since it isn't a child"), *ExpectedWithinDef->GetName(), *ClassWithin->GetName());
 	}
 }
 
diff --git a/Engine/Source/Programs/UnrealHeaderTool/Private/UnrealTypeDefinitionInfo.h b/Engine/Source/Programs/UnrealHeaderTool/Private/UnrealTypeDefinitionInfo.h
--- a/Engine/Source/Programs/UnrealHeaderTool/Private/UnrealTypeDefinitionInfo.h
+++ b/Engine/Source/Programs/UnrealHeaderTool/Private/UnrealTypeDefinitionInfo.h
@@ -10,7 +10,7 @@
 #include "BaseParser.h"
 #include "ParserHelper.h"
 
-#define UHT_ENABLE_ENGINE_TYPE_CHECKS 1
+#define UHT_ENABLE_ENGINE_TYPE_CHECKS 0
 
 // Forward declarations.
 class FHeaderParser;
@@ -232,7 +232,7 @@ private:
  */
 class FUnrealTypeDefinitionInfo 
 	: public TSharedFromThis<FUnrealTypeDefinitionInfo>
-	, public FUHTExceptionContext
+	, public FUHTMessageProvider
 {
 public:
 	enum class ECreateEngineTypesPhase : uint8
@@ -520,7 +520,7 @@ public:
 	 */
 	virtual void AddMetaData(TMap<FName, FString>&& InMetaData)
 	{
-		FUHTException::Throwf(*this, TEXT("Meta data can not be set for a definition of this type."));
+		Throwf(TEXT("Meta data can not be set for a definition of this type."));
 	}
 
 	/**
@@ -2258,7 +2258,7 @@ public:
 	{
 		if (DefinitionRange.End <= DefinitionRange.Start)
 		{
-			FUHTException::Throwf(*this, TEXT("The class definition range is invalid. Most probably caused by previous parsing error."));
+			Throwf(TEXT("The class definition range is invalid. Most probably caused by previous parsing error."));
 		}
 	}
 
