Changed the metadata base class API to return a reference to the local meta data since all users now supply metadata.
Moved the parsing of the enumeration form and underlying type to the pre-parser so they can be accessed during the parse phase.
Fixed "warning as errors" return code.
Removed the SimplePropertyTypeDispatch since it was only used the the version that supported containers.

#rnx

diff --git a/Engine/Source/Programs/UnrealHeaderTool/Private/BaseParser.cpp b/Engine/Source/Programs/UnrealHeaderTool/Private/BaseParser.cpp
--- a/Engine/Source/Programs/UnrealHeaderTool/Private/BaseParser.cpp
+++ b/Engine/Source/Programs/UnrealHeaderTool/Private/BaseParser.cpp
@@ -1218,4 +1218,61 @@ void FBaseParser::InsertMetaDataPair(TMap<FName, FString>& MetaData, FName KeyNa
 void FBaseParser::InitMetadataKeywords()
 {
 	InitMetadataKeyswordsInternal();
-}
\ No newline at end of file
+}
+
+EUnderlyingEnumType FBaseParser::ParseUnderlyingEnumType()
+{
+	EUnderlyingEnumType UnderlyingType = EUnderlyingEnumType::uint8;
+	if (MatchSymbol(TEXT(':')))
+	{
+		FToken BaseToken;
+		if (!GetIdentifier(BaseToken))
+		{
+			Throwf(TEXT("Missing enum base"));
+		}
+
+		if (!FCString::Strcmp(BaseToken.Identifier, TEXT("uint8")))
+		{
+			UnderlyingType = EUnderlyingEnumType::uint8;
+		}
+		else if (!FCString::Strcmp(BaseToken.Identifier, TEXT("uint16")))
+		{
+			UnderlyingType = EUnderlyingEnumType::uint16;
+		}
+		else if (!FCString::Strcmp(BaseToken.Identifier, TEXT("uint32")))
+		{
+			UnderlyingType = EUnderlyingEnumType::uint32;
+		}
+		else if (!FCString::Strcmp(BaseToken.Identifier, TEXT("uint64")))
+		{
+			UnderlyingType = EUnderlyingEnumType::uint64;
+		}
+		else if (!FCString::Strcmp(BaseToken.Identifier, TEXT("int8")))
+		{
+			UnderlyingType = EUnderlyingEnumType::int8;
+		}
+		else if (!FCString::Strcmp(BaseToken.Identifier, TEXT("int16")))
+		{
+			UnderlyingType = EUnderlyingEnumType::int16;
+		}
+		else if (!FCString::Strcmp(BaseToken.Identifier, TEXT("int32")))
+		{
+			UnderlyingType = EUnderlyingEnumType::int32;
+		}
+		else if (!FCString::Strcmp(BaseToken.Identifier, TEXT("int64")))
+		{
+			UnderlyingType = EUnderlyingEnumType::int64;
+		}
+		else
+		{
+			Throwf(TEXT("Unsupported enum class base type: %s"), BaseToken.Identifier);
+		}
+	}
+	else
+	{
+		UnderlyingType = EUnderlyingEnumType::Unspecified;
+	}
+
+	return UnderlyingType;
+}
+
diff --git a/Engine/Source/Programs/UnrealHeaderTool/Private/BaseParser.h b/Engine/Source/Programs/UnrealHeaderTool/Private/BaseParser.h
--- a/Engine/Source/Programs/UnrealHeaderTool/Private/BaseParser.h
+++ b/Engine/Source/Programs/UnrealHeaderTool/Private/BaseParser.h
@@ -17,6 +17,19 @@ enum class EPointerMemberBehavior
 	AllowAndLog,
 };
 
+enum class EUnderlyingEnumType
+{
+	Unspecified,
+	uint8,
+	uint16,
+	uint32,
+	uint64,
+	int8,
+	int16,
+	int32,
+	int64
+};
+
 using FMetaData = TMap<FName, FString>;
 
 /////////////////////////////////////////////////////
@@ -255,6 +268,11 @@ public:
 	template <typename Lambda>
 	void ParseInheritance(const TCHAR* What, Lambda&& InLambda);
 
+	/**
+	 * Parse the underlying enum type
+	 */
+	EUnderlyingEnumType ParseUnderlyingEnumType();
+
 	//////////////
 
 	// Initialize the metadata keywords prior to parsing
diff --git a/Engine/Source/Programs/UnrealHeaderTool/Private/Exceptions.cpp b/Engine/Source/Programs/UnrealHeaderTool/Private/Exceptions.cpp
--- a/Engine/Source/Programs/UnrealHeaderTool/Private/Exceptions.cpp
+++ b/Engine/Source/Programs/UnrealHeaderTool/Private/Exceptions.cpp
@@ -189,7 +189,7 @@ ECompilationResult::Type FResults::GetOverallResults()
 	using namespace UE::UnrealHeaderTool::Exceptions::Private;
 
 	// For some legacy reason, we don't actually return the result
-	if (OverallResults != ECompilationResult::Succeeded || NumFailures > 0)
+	if (OverallResults != ECompilationResult::Succeeded || NumFailures > 0 || (NumWarnings > 0 && GWarn->TreatWarningsAsErrors))
 	{
 		return ECompilationResult::OtherCompilationError;
 	}
diff --git a/Engine/Source/Programs/UnrealHeaderTool/Private/HeaderParser.cpp b/Engine/Source/Programs/UnrealHeaderTool/Private/HeaderParser.cpp
--- a/Engine/Source/Programs/UnrealHeaderTool/Private/HeaderParser.cpp
+++ b/Engine/Source/Programs/UnrealHeaderTool/Private/HeaderParser.cpp
@@ -1003,6 +1003,11 @@ FUnrealEnumDefinitionInfo& FHeaderParser::CompileEnum()
 
 	Scope->AddType(EnumDef);
 
+	if (CppForm != EnumDef.GetCppForm())
+	{
+		Throwf(TEXT("ICE: Mismatch of enum CPP form between pre-parser and parser"));
+	}
+
 	for (const FPropertySpecifier& Specifier : SpecifiersFound)
 	{
 		switch ((EEnumSpecifier)Algo::FindSortedStringCaseInsensitive(*Specifier.Key, GEnumSpecifierStrings))
@@ -1028,57 +1033,11 @@ FUnrealEnumDefinitionInfo& FHeaderParser::CompileEnum()
 	EUnderlyingEnumType UnderlyingType = EUnderlyingEnumType::uint8;
 	if (CppForm == UEnum::ECppForm::EnumClass)
 	{
-		if (MatchSymbol(TEXT(':')))
-		{
-			FToken BaseToken;
-			if (!GetIdentifier(BaseToken))
-			{
-				Throwf(TEXT("Missing enum base") );
-			}
-
-			if (!FCString::Strcmp(BaseToken.Identifier, TEXT("uint8")))
-			{
-				UnderlyingType = EUnderlyingEnumType::uint8;
-			}
-			else if (!FCString::Strcmp(BaseToken.Identifier, TEXT("uint16")))
-			{
-				UnderlyingType = EUnderlyingEnumType::uint16;
-			}
-			else if (!FCString::Strcmp(BaseToken.Identifier, TEXT("uint32")))
-			{
-				UnderlyingType = EUnderlyingEnumType::uint32;
-			}
-			else if (!FCString::Strcmp(BaseToken.Identifier, TEXT("uint64")))
-			{
-				UnderlyingType = EUnderlyingEnumType::uint64;
-			}
-			else if (!FCString::Strcmp(BaseToken.Identifier, TEXT("int8")))
-			{
-				UnderlyingType = EUnderlyingEnumType::int8;
-			}
-			else if (!FCString::Strcmp(BaseToken.Identifier, TEXT("int16")))
-			{
-				UnderlyingType = EUnderlyingEnumType::int16;
-			}
-			else if (!FCString::Strcmp(BaseToken.Identifier, TEXT("int32")))
-			{
-				UnderlyingType = EUnderlyingEnumType::int32;
-			}
-			else if (!FCString::Strcmp(BaseToken.Identifier, TEXT("int64")))
-			{
-				UnderlyingType = EUnderlyingEnumType::int64;
-			}
-			else
-			{
-				Throwf(TEXT("Unsupported enum class base type: %s"), BaseToken.Identifier);
-			}
-		}
-		else
+		UnderlyingType = ParseUnderlyingEnumType();
+		if (UnderlyingType != EnumDef.GetUnderlyingType())
 		{
-			UnderlyingType = EUnderlyingEnumType::Unspecified;
+			Throwf(TEXT("ICE: Mismatch of underlying enum type between pre-parser and parser"));
 		}
-
-		EnumDef.SetUnderlyingType(UnderlyingType);
 	}
 	else
 	{
@@ -8093,6 +8052,7 @@ TSharedRef<FUnrealTypeDefinitionInfo> FHeaderPreParser::ParseEnumDeclaration(con
 
 	// Check enum type. This can be global 'enum', 'namespace' or 'enum class' enums.
 	bool bReadEnumName = false;
+	UEnum::ECppForm CppForm = UEnum::ECppForm::Regular;
 	if (!GetIdentifier(EnumToken))
 	{
 		Throwf(TEXT("Missing identifier after UENUM()"));
@@ -8100,6 +8060,8 @@ TSharedRef<FUnrealTypeDefinitionInfo> FHeaderPreParser::ParseEnumDeclaration(con
 
 	if (EnumToken.Matches(TEXT("namespace"), ESearchCase::CaseSensitive))
 	{
+		CppForm = UEnum::ECppForm::Namespaced;
+
 		SkipDeprecatedMacroIfNecessary(*this);
 
 		bReadEnumName = GetIdentifier(EnumToken);
@@ -8113,11 +8075,14 @@ TSharedRef<FUnrealTypeDefinitionInfo> FHeaderPreParser::ParseEnumDeclaration(con
 
 		if (EnumToken.Matches(TEXT("class"), ESearchCase::CaseSensitive) || EnumToken.Matches(TEXT("struct"), ESearchCase::CaseSensitive))
 		{
+			CppForm = UEnum::ECppForm::EnumClass;
 		}
 		else
 		{
 			// Put whatever token we found back so that we can correctly skip below
 			UngetToken(EnumToken);
+
+			CppForm = UEnum::ECppForm::Regular;
 		}
 
 		SkipAlignasAndDeprecatedMacroIfNecessary(*this);
@@ -8135,7 +8100,14 @@ TSharedRef<FUnrealTypeDefinitionInfo> FHeaderPreParser::ParseEnumDeclaration(con
 		Throwf(TEXT("Missing enumeration name"));
 	}
 
-	TSharedRef<FUnrealEnumDefinitionInfo> EnumDef = MakeShareable(new FUnrealEnumDefinitionInfo(SourceFile, InLineNumber, EnumToken.Identifier, FName(EnumToken.Identifier, FNAME_Add)));
+	// Read base for enum class
+	EUnderlyingEnumType UnderlyingType = EUnderlyingEnumType::uint8;
+	if (CppForm == UEnum::ECppForm::EnumClass)
+	{
+		UnderlyingType = ParseUnderlyingEnumType();
+	}
+
+	TSharedRef<FUnrealEnumDefinitionInfo> EnumDef = MakeShareable(new FUnrealEnumDefinitionInfo(SourceFile, InLineNumber, EnumToken.Identifier, FName(EnumToken.Identifier, FNAME_Add), CppForm, UnderlyingType));
 	return EnumDef;
 }
 
diff --git a/Engine/Source/Programs/UnrealHeaderTool/Private/PropertyTypes.cpp b/Engine/Source/Programs/UnrealHeaderTool/Private/PropertyTypes.cpp
--- a/Engine/Source/Programs/UnrealHeaderTool/Private/PropertyTypes.cpp
+++ b/Engine/Source/Programs/UnrealHeaderTool/Private/PropertyTypes.cpp
@@ -166,13 +166,39 @@ namespace
 	//----------------------------------------------------------------------------------------------------------------------------------------------------------
 
 	/**
-	 * Given just the property type, invoke the dispatch functor.
-	 *
-	 * NOTE: This method does not support container or enum types.
+	 * Based on the type and container settings, dispatch to the correct traits
 	 */
-	template <template<typename PropertyTraits> typename FuncDispatch, typename RetValue, typename PropDefType, typename ... Args>
-	RetValue SimplePropertyTypeDispatch(PropDefType& PropDef, Args&& ... args)
+	template <template <typename PropertyTraits> typename FuncDispatch, bool bHandleContainers, typename RetValue, typename PropDefType, typename ... Args>
+	RetValue PropertyTypeDispatch(PropDefType& PropDef, Args&& ... args)
 	{
+		const FPropertyBase& VarProperty = PropDef.GetPropertyBase();
+		if constexpr (bHandleContainers)
+		{
+			switch (VarProperty.ArrayType)
+			{
+			case EArrayType::Static:
+				return FuncDispatch<FPropertyTypeTraitsStaticArray>()(PropDef, std::forward<Args>(args)...);
+			case EArrayType::Dynamic:
+				return FuncDispatch<FPropertyTypeTraitsDynamicArray>()(PropDef, std::forward<Args>(args)...);
+			case EArrayType::Set:
+				return FuncDispatch<FPropertyTypeTraitsSet>()(PropDef, std::forward<Args>(args)...);
+			}
+
+			if (VarProperty.MapKeyProp.IsValid())
+			{
+				return FuncDispatch<FPropertyTypeTraitsMap>()(PropDef, std::forward<Args>(args)...);
+			}
+		}
+
+		// Check if it's an enum class property
+		// NOTE: VarProperty.EnumDef is a union and might not be an enum
+		if (VarProperty.IsEnum())
+		{
+			return FuncDispatch<FPropertyTypeTraitsEnum>()(PropDef, std::forward<Args>(args)...);
+		}
+
+
+		// We are just a simple engine type
 		switch (PropDef.GetPropertyBase().Type)
 		{
 		case CPT_Byte:						return FuncDispatch<FPropertyTypeTraitsByte>()(PropDef, std::forward<Args>(args)...);
@@ -208,41 +234,6 @@ namespace
 		}
 	}
 
-	/**
-	 * Given just the property type, invoke the dispatch functor.
-	 */
-	template <template <typename PropertyTraits> typename FuncDispatch, bool bHandleContainers, typename RetValue, typename PropDefType, typename ... Args>
-	RetValue ComplexPropertyTypeDispatch(PropDefType& PropDef, Args&& ... args)
-	{
-		const FPropertyBase& VarProperty = PropDef.GetPropertyBase();
-		if constexpr (bHandleContainers)
-		{
-			switch (VarProperty.ArrayType)
-			{
-			case EArrayType::Static:
-				return FuncDispatch<FPropertyTypeTraitsStaticArray>()(PropDef, std::forward<Args>(args)...);
-			case EArrayType::Dynamic:
-				return FuncDispatch<FPropertyTypeTraitsDynamicArray>()(PropDef, std::forward<Args>(args)...);
-			case EArrayType::Set:
-				return FuncDispatch<FPropertyTypeTraitsSet>()(PropDef, std::forward<Args>(args)...);
-			}
-
-			if (VarProperty.MapKeyProp.IsValid())
-			{
-				return FuncDispatch<FPropertyTypeTraitsMap>()(PropDef, std::forward<Args>(args)...);
-			}
-		}
-
-		// Check if it's an enum class property
-		// NOTE: VarProperty.EnumDef is a union and might not be an enum
-		if (VarProperty.IsEnum())
-		{
-			return FuncDispatch<FPropertyTypeTraitsEnum>()(PropDef, std::forward<Args>(args)...);
-		}
-
-		return SimplePropertyTypeDispatch<FuncDispatch, RetValue>(PropDef, std::forward<Args>(args)...);
-	}
-
 	//----------------------------------------------------------------------------------------------------------------------------------------------------------
 	//----------------------------------------------------------------------------------------------------------------------------------------------------------
 	// Dispatch Functors
@@ -321,18 +312,18 @@ namespace
 	template <bool bHandleContainers>
 	void CreatePropertyHelper(FUnrealPropertyDefinitionInfo& PropDef)
 	{
-		return ComplexPropertyTypeDispatch<CreatePropertyDispatch, bHandleContainers, void>(PropDef);
+		return PropertyTypeDispatch<CreatePropertyDispatch, bHandleContainers, void>(PropDef);
 	}
 
 	template <bool bHandleContainers>
 	FProperty* CreateEngineTypeHelper(FUnrealPropertyDefinitionInfo& PropDef, FFieldVariant Scope, const FName& Name, EObjectFlags ObjectFlags)
 	{
-		return ComplexPropertyTypeDispatch<CreateEngineTypeDispatch, bHandleContainers, FProperty*>(PropDef, Scope, std::ref(Name), ObjectFlags);
+		return PropertyTypeDispatch<CreateEngineTypeDispatch, bHandleContainers, FProperty*>(PropDef, Scope, std::ref(Name), ObjectFlags);
 	}
 
 	bool IsSupportedByBlueprintSansContainers(const FUnrealPropertyDefinitionInfo& PropDef, bool bMemberVariable)
 	{
-		return ComplexPropertyTypeDispatch<IsSupportedByBlueprintDispatch, false, bool>(PropDef, bMemberVariable);
+		return PropertyTypeDispatch<IsSupportedByBlueprintDispatch, false, bool>(PropDef, bMemberVariable);
 	}
 }
 
@@ -2280,17 +2271,17 @@ struct FPropertyTypeTraitsStaticArray : public FPropertyTypeTraitsBase
 
 	static FString GetEngineClassName(const FUnrealPropertyDefinitionInfo& PropDef)
 	{
-		return ComplexPropertyTypeDispatch<GetEngineClassNameDispatch, false, FString>(PropDef);
+		return PropertyTypeDispatch<GetEngineClassNameDispatch, false, FString>(PropDef);
 	}
 
 	static FString GetCPPType(const FUnrealPropertyDefinitionInfo& PropDef, FString* ExtendedTypeText, uint32 CPPExportFlags)
 	{
-		return ComplexPropertyTypeDispatch<GetCPPTypeDispatch, false, FString>(PropDef, ExtendedTypeText, CPPExportFlags);
+		return PropertyTypeDispatch<GetCPPTypeDispatch, false, FString>(PropDef, ExtendedTypeText, CPPExportFlags);
 	}
 
 	static FString GetCPPTypeForwardDeclaration(const FUnrealPropertyDefinitionInfo& PropDef)
 	{
-		return ComplexPropertyTypeDispatch<GetCPPTypeForwardDeclarationDispatch, false, FString>(PropDef);
+		return PropertyTypeDispatch<GetCPPTypeForwardDeclarationDispatch, false, FString>(PropDef);
 	}
 };
 
@@ -2540,7 +2531,7 @@ bool FPropertyTraits::DefaultValueStringCppFormatToInnerFormat(const FUnrealProp
 		return false;
 	}
 
-	return ComplexPropertyTypeDispatch<DefaultValueStringCppFormatToInnerFormatDispatch, false, bool>(PropDef, std::ref(CppForm), std::ref(OutForm));
+	return PropertyTypeDispatch<DefaultValueStringCppFormatToInnerFormatDispatch, false, bool>(PropDef, std::ref(CppForm), std::ref(OutForm));
 }
 
 TSharedRef<FUnrealPropertyDefinitionInfo> FPropertyTraits::CreateProperty(const FPropertyBase& VarProperty, FUnrealTypeDefinitionInfo& Outer, const FName& Name, EVariableCategory VariableCategory, EAccessSpecifier AccessSpecifier, const TCHAR* Dimensions, FUnrealSourceFile& SourceFile, int LineNumber, int ParsePosition)
@@ -2592,12 +2583,9 @@ FProperty* FPropertyTraits::CreateEngineType(TSharedRef<FUnrealPropertyDefinitio
 	}
 
 	// Add the meta data to the property
-	if (TMap<FName, FString>* MetaDataMap = PropDef.GetMetaDataMap())
+	for (TPair<FName, FString>& KVP : PropDef.GetMetaDataMap())
 	{
-		for (TPair<FName, FString>& KVP : *MetaDataMap)
-		{
-			Property->SetMetaData(KVP.Key, FString(KVP.Value));
-		}
+		Property->SetMetaData(KVP.Key, FString(KVP.Value));
 	}
 
 	// If we are parent to a struct
@@ -2622,22 +2610,22 @@ FProperty* FPropertyTraits::CreateEngineType(TSharedRef<FUnrealPropertyDefinitio
 
 bool FPropertyTraits::IsSupportedByBlueprint(const FUnrealPropertyDefinitionInfo& PropDef, bool bMemberVariable)
 {
-	return ComplexPropertyTypeDispatch<IsSupportedByBlueprintDispatch, true, bool>(PropDef, bMemberVariable);
+	return PropertyTypeDispatch<IsSupportedByBlueprintDispatch, true, bool>(PropDef, bMemberVariable);
 }
 
 FString FPropertyTraits::GetEngineClassName(const FUnrealPropertyDefinitionInfo& PropDef)
 {
-	return ComplexPropertyTypeDispatch<GetEngineClassNameDispatch, true, FString>(PropDef);
+	return PropertyTypeDispatch<GetEngineClassNameDispatch, true, FString>(PropDef);
 }
 
 FString FPropertyTraits::GetCPPType(const FUnrealPropertyDefinitionInfo& PropDef, FString* ExtendedTypeText/* = nullptr */, uint32 CPPExportFlags/* = 0 */)
 {
-	return ComplexPropertyTypeDispatch<GetCPPTypeDispatch, true, FString>(PropDef, ExtendedTypeText, CPPExportFlags);
+	return PropertyTypeDispatch<GetCPPTypeDispatch, true, FString>(PropDef, ExtendedTypeText, CPPExportFlags);
 }
 
 FString FPropertyTraits::GetCPPTypeForwardDeclaration(const FUnrealPropertyDefinitionInfo& PropDef)
 {
-	return ComplexPropertyTypeDispatch<GetCPPTypeForwardDeclarationDispatch, true, FString>(PropDef);
+	return PropertyTypeDispatch<GetCPPTypeForwardDeclarationDispatch, true, FString>(PropDef);
 }
 
 bool FPropertyTraits::SameType(const FUnrealPropertyDefinitionInfo& Lhs, const FUnrealPropertyDefinitionInfo& Rhs)
diff --git a/Engine/Source/Programs/UnrealHeaderTool/Private/UnrealHeaderToolMain.cpp b/Engine/Source/Programs/UnrealHeaderTool/Private/UnrealHeaderToolMain.cpp
--- a/Engine/Source/Programs/UnrealHeaderTool/Private/UnrealHeaderToolMain.cpp
+++ b/Engine/Source/Programs/UnrealHeaderTool/Private/UnrealHeaderToolMain.cpp
@@ -19,9 +19,6 @@ IMPLEMENT_APPLICATION(UnrealHeaderTool, "UnrealHeaderTool");
 
 DEFINE_LOG_CATEGORY(LogCompile);
 
-bool GUHTWarningLogged = false;
-bool GUHTErrorLogged = false;
-
 /**
  * Application entry point
  *
@@ -110,12 +107,6 @@ INT32_MAIN_INT32_ARGC_TCHAR_ARGV()
 
 	extern ECompilationResult::Type UnrealHeaderTool_Main(const FString& ModuleInfoFilename);
 	ECompilationResult::Type Result = UnrealHeaderTool_Main(ModuleInfoFilename);
-
-	if (Result == ECompilationResult::Succeeded && (GUHTErrorLogged || (GUHTWarningLogged && GWarn->TreatWarningsAsErrors)))
-	{
-		Result = ECompilationResult::OtherCompilationError;
-	}
-
 	return Result;
 }
 
diff --git a/Engine/Source/Programs/UnrealHeaderTool/Private/UnrealTypeDefinitionInfo.cpp b/Engine/Source/Programs/UnrealHeaderTool/Private/UnrealTypeDefinitionInfo.cpp
--- a/Engine/Source/Programs/UnrealHeaderTool/Private/UnrealTypeDefinitionInfo.cpp
+++ b/Engine/Source/Programs/UnrealHeaderTool/Private/UnrealTypeDefinitionInfo.cpp
@@ -198,10 +198,7 @@ const FString* FUHTMetaData::FindMetaDataHelper(const FName& Key) const
 
 	if (Key != NAME_None)
 	{
-		if (TMap<FName, FString>* MetaDataMap = GetMetaDataMap())
-		{
-			Result = MetaDataMap->Find(Key);
-		}
+		Result = GetMetaDataMap().Find(Key);
 	}
 #if UHT_ENABLE_ENGINE_TYPE_CHECKS
 	CheckFindMetaData(Key, Result);
@@ -211,14 +208,7 @@ const FString* FUHTMetaData::FindMetaDataHelper(const FName& Key) const
 
 void FUHTMetaData::SetMetaDataHelper(const FName& Key, const TCHAR* InValue) 
 {
-	if (TMap<FName, FString>* MetaDataMap = GetMetaDataMap())
-	{
-		MetaDataMap->Add(Key, InValue);
-	}
-	else
-	{
-		checkf(false, TEXT("If GetMetaDataMap can return nullptr, this SetMetaDataHelper must be overridden"));
-	}
+	GetMetaDataMap().Add(Key, InValue);
 }
 
 FName FUHTMetaData::GetMetaDataKey(const TCHAR* Key, int32 NameIndex, EFindName FindName) const
@@ -962,39 +952,39 @@ FString FUnrealObjectDefinitionInfo::GetFullGroupName(bool bStartWithOuter) cons
 	return Obj ? Obj->GetPathName(&GetPackageDef()) : TEXT("");
 }
 
-void FUnrealObjectDefinitionInfo::SetMetaDataHelper(const FName& Key, const TCHAR* InValue)
+void FUnrealObjectDefinitionInfo::AddMetaData(TMap<FName, FString>&& InMetaData)
 {
-	// If we don't have an existing meta data map, then we must be using the meta data off
-	// the object and no values have been set yet.  In that case just set the whole block
-	TMap<FName, FString>* MetaDataMap = GetMetaDataMap();
-	if (MetaDataMap == nullptr)
+	// only add if we have some!
+	if (InMetaData.Num())
 	{
-		UMetaData* MetaData = GetUObjectMetaData();
-		MetaData->SetValue(GetObject(), Key, InValue); // This is valid
-	}
 
-	// Otherwise we are either using the objects meta data and it already exists, or
-	// we are using a UHT private meta data block.  If we have both, then first
-	// set the object followed by the private copy.
-	else
-	{
-		TMap<FName, FString>* UObjectMetaDataMap = GetUObjectMetaDataMap();
-		if (MetaDataMap != UObjectMetaDataMap)
+		if (TMap<FName, FString>* UObjectMetaDataMap = GetUObjectMetaDataMap())
 		{
-			if (UObjectMetaDataMap != nullptr)
-			{
-				UObjectMetaDataMap->Add(Key, InValue);
-			}
-			else if (UMetaData* MetaData = GetUObjectMetaData())
-			{
-				MetaData->SetValue(GetObject(), Key, InValue); // This is valid
-			}
+			UObjectMetaDataMap->Append(InMetaData);
+		}
+		else if (UMetaData* UObjectMetaData = GetUObjectMetaData())
+		{
+			UObjectMetaData->SetObjectValues(GetObject(), InMetaData); // This is valid
 		}
 
-		MetaDataMap->Add(Key, InValue);
+		GetMetaDataMap().Append(MoveTemp(InMetaData));
 	}
 }
 
+void FUnrealObjectDefinitionInfo::SetMetaDataHelper(const FName& Key, const TCHAR* InValue)
+{
+	if (TMap<FName, FString>* UObjectMetaDataMap = GetUObjectMetaDataMap())
+	{
+		UObjectMetaDataMap->Add(Key, InValue);
+	}
+	else if (UMetaData* UObjectMetaData = GetUObjectMetaData())
+	{
+		UObjectMetaData->SetValue(GetObject(), Key, InValue); // This is valid
+	}
+
+	GetMetaDataMap().Add(Key, InValue);
+}
+
 FUnrealPackageDefinitionInfo::FUnrealPackageDefinitionInfo(const FManifestModule& InModule, UPackage* InPackage)
 	: FUnrealObjectDefinitionInfo(InPackage)
 	, Module(InModule)
@@ -1124,44 +1114,6 @@ void FUnrealFieldDefinitionInfo::AddCrossModuleReference(TSet<FString>* UniqueCr
 	}
 }
 
-void FUnrealFieldDefinitionInfo::AddMetaData(TMap<FName, FString>&& InMetaData)
-{
-	// only add if we have some!
-	if (InMetaData.Num())
-	{
-
-		// If we don't have an existing meta data map, then we must be using the meta data off
-		// the object and no values have been set yet.  In that case just set the whole block
-		TMap<FName, FString>* MetaDataMap = GetMetaDataMap();
-		if (MetaDataMap == nullptr)
-		{
-			UMetaData* MetaData = GetUObjectMetaData();
-			MetaData->SetObjectValues(GetObject(), MoveTemp(InMetaData)); // This is valid
-		}
-
-		// Otherwise we are either using the objects meta data and it already exists, or
-		// we are using a UHT private meta data block.  If we have both, then first
-		// set the object followed by the private copy.
-		else
-		{
-			TMap<FName, FString>* UObjectMetaDataMap = GetUObjectMetaDataMap();
-			if (MetaDataMap != UObjectMetaDataMap)
-			{
-				if (UObjectMetaDataMap != nullptr)
-				{
-					UObjectMetaDataMap->Append(InMetaData);
-				}
-				else if (UMetaData* MetaData = GetUObjectMetaData())
-				{
-					MetaData->SetObjectValues(GetObject(), InMetaData); // This is valid
-				}
-			}
-
-			MetaDataMap->Append(MoveTemp(InMetaData));
-		}
-	}
-}
-
 bool FUnrealFieldDefinitionInfo::IsDynamic() const
 {
 	return HasMetaData(NAME_ReplaceConverted);
@@ -1191,8 +1143,10 @@ FUnrealClassDefinitionInfo* FUnrealFieldDefinitionInfo::GetOwnerClass() const
 	return nullptr;
 }
 
-FUnrealEnumDefinitionInfo::FUnrealEnumDefinitionInfo(FUnrealSourceFile& InSourceFile, int32 InLineNumber, FString&& InNameCPP, FName InName)
+FUnrealEnumDefinitionInfo::FUnrealEnumDefinitionInfo(FUnrealSourceFile& InSourceFile, int32 InLineNumber, FString&& InNameCPP, FName InName, UEnum::ECppForm InCppForm, EUnderlyingEnumType InUnderlyingType)
 	: FUnrealFieldDefinitionInfo(InSourceFile, InLineNumber, MoveTemp(InNameCPP), InName, InSourceFile.GetPackageDef())
+	, UnderlyingType(InUnderlyingType)
+	, CppForm(InCppForm)
 { }
 
 FString FUnrealEnumDefinitionInfo::GenerateEnumPrefix() const
@@ -1450,7 +1404,7 @@ void FUnrealEnumDefinitionInfo::CreateUObjectEngineTypesInternal(ECreateEngineTy
 		UEnum* Enum = new(EC_InternalUseOnlyConstructor, Package, FName(EnumName), RF_Public) UEnum(FObjectInitializer());
 		Enum->SetEnums(Names, CppForm, EnumFlags, false);
 		Enum->CppType = CppType;
-		Enum->GetPackage()->GetMetaData()->SetObjectValues(Enum, MetaData);
+		Enum->GetPackage()->GetMetaData()->SetObjectValues(Enum, GetMetaDataMap());
 		SetObject(Enum);
 		GTypeDefinitionInfoMap.AddObjectLookup(Enum, SharedThis(this));
 		break;
@@ -1650,7 +1604,7 @@ void FUnrealScriptStructDefinitionInfo::CreateUObjectEngineTypesInternal(ECreate
 
 		UScriptStruct* ScriptStruct = new(EC_InternalUseOnlyConstructor, Package, FName(StructNameStripped), RF_Public) UScriptStruct(FObjectInitializer());
 		ScriptStruct->StructFlags = StructFlags;
-		ScriptStruct->GetPackage()->GetMetaData()->SetObjectValues(ScriptStruct, MetaData);
+		ScriptStruct->GetPackage()->GetMetaData()->SetObjectValues(ScriptStruct, GetMetaDataMap());
 		if (FUnrealStructDefinitionInfo* SuperStructDef = GetSuperStruct())
 		{
 			ScriptStruct->SetSuperStruct(SuperStructDef->GetStruct());
@@ -1829,7 +1783,7 @@ void FUnrealClassDefinitionInfo::CreateUObjectEngineTypesInternal(ECreateEngineT
 			Class->ClassConfigName = ClassConfigName;
 			Class->ClassWithin = ClassWithin->GetClassSafe();
 			Class->SetInternalFlags(GetInternalFlags());
-			Class->GetPackage()->GetMetaData()->SetObjectValues(Class, MetaData);
+			Class->GetPackage()->GetMetaData()->SetObjectValues(Class, GetMetaDataMap());
 
 			// Setup the base class
 			if (FUnrealClassDefinitionInfo* SuperClassDef = GetSuperClass())
@@ -2574,7 +2528,7 @@ void FUnrealFunctionDefinitionInfo::CreateUObjectEngineTypesInternal(ECreateEngi
 		Function->ReturnValueOffset = MAX_uint16;
 		Function->FirstPropertyToInit = nullptr;
 		Function->FunctionFlags |= FunctionData.FunctionFlags;
-		Function->GetPackage()->GetMetaData()->SetObjectValues(Function, MetaData);
+		Function->GetPackage()->GetMetaData()->SetObjectValues(Function, GetMetaDataMap());
 
 		SetObject(Function);
 		GTypeDefinitionInfoMap.Add(Function, SharedThis(this));
diff --git a/Engine/Source/Programs/UnrealHeaderTool/Private/UnrealTypeDefinitionInfo.h b/Engine/Source/Programs/UnrealHeaderTool/Private/UnrealTypeDefinitionInfo.h
--- a/Engine/Source/Programs/UnrealHeaderTool/Private/UnrealTypeDefinitionInfo.h
+++ b/Engine/Source/Programs/UnrealHeaderTool/Private/UnrealTypeDefinitionInfo.h
@@ -46,19 +46,6 @@ enum class EEnforceInterfacePrefix
 	U
 };
 
-enum class EUnderlyingEnumType
-{
-	Unspecified,
-	uint8,
-	uint16,
-	uint32,
-	uint64,
-	int8,
-	int16,
-	int32,
-	int64
-};
-
 enum class ESerializerArchiveType
 {
 	None = 0,
@@ -134,13 +121,11 @@ public:
 	 */
 	const FString* FindMetaData(const TCHAR* Key, int32 NameIndex = INDEX_NONE) const
 	{
-		TMap<FName, FString>* MetaDataMap = GetMetaDataMap();
-		return MetaDataMap ? MetaDataMap->Find(GetMetaDataKey(Key, NameIndex, FNAME_Find)) : nullptr;
+		return GetMetaDataMap().Find(GetMetaDataKey(Key, NameIndex, FNAME_Find));
 	}
 	const FString* FindMetaData(const FName& Key, int32 NameIndex = INDEX_NONE) const
 	{
-		TMap<FName, FString>* MetaDataMap = GetMetaDataMap();
-		return MetaDataMap ? MetaDataMap->Find(GetMetaDataKey(Key, NameIndex, FNAME_Find)) : nullptr;
+		return GetMetaDataMap().Find(GetMetaDataKey(Key, NameIndex, FNAME_Find));
 	}
 
 	/**
@@ -201,7 +186,7 @@ public:
 	void GetStringArrayMetaData(const FName& Key, TArray<FString>& Out) const;
 	TArray<FString> GetStringArrayMetaData(const FName& Key) const;
 
-	virtual TMap<FName, FString>* GetMetaDataMap() const = 0;
+	virtual TMap<FName, FString>& GetMetaDataMap() const = 0;
 
 protected:
 	virtual FString GetMetaDataIndexName(int32 NameIndex) const
@@ -711,7 +696,7 @@ public:
 	/**
 	 * Return the engine class name
 	 */
-	FString GetEngineClassName() const; //@TODO make this const &
+	FString GetEngineClassName() const;
 
 	/**
 	 * Return the package associated with this object
@@ -847,13 +832,7 @@ public:
 	 */
 	virtual TMap<FName, FString> GenerateMetadataMap() const override
 	{
-		TMap<FName, FString> MetaDataMap;
-		const TMap<FName, FString>* FieldMetaDataMap = GetMetaDataMap();
-		if (FieldMetaDataMap)
-		{
-			MetaDataMap = *FieldMetaDataMap;
-		}
-		return MetaDataMap;
+		return GetMetaDataMap();
 	}
 
 	/**
@@ -1218,13 +1197,12 @@ public:
 	virtual void AddMetaData(TMap<FName, FString>&& InMetaData) override
 	{
 		check(Property == nullptr);
-		GetMetaDataMap()->Append(MoveTemp(InMetaData));
+		GetMetaDataMap().Append(MoveTemp(InMetaData));
 	}
 
-	virtual TMap<FName, FString>* GetMetaDataMap() const override
+	virtual TMap<FName, FString>& GetMetaDataMap() const override
 	{
-		//return const_cast<TMap<FName, FString>*>(GetProperty()->GetMetaDataMap());
-		return const_cast<TMap<FName, FString>*>(&PropertyBase.MetaData);
+		return const_cast<TMap<FName, FString>&>(PropertyBase.MetaData);
 	}
 
 	FUnrealObjectDefinitionInfo* GetOwnerObject() const;
@@ -1435,27 +1413,28 @@ public:
 	 */
 	virtual TMap<FName, FString> GenerateMetadataMap() const override
 	{
-		TMap<FName, FString>* PackageMap = GetMetaDataMap();
 		TMap<FName, FString> Map;
-		if (PackageMap)
+		for (const TPair<FName, FString>& MetaKeyValue : GetMetaDataMap())
 		{
-			for (const TPair<FName, FString>& MetaKeyValue : *PackageMap)
+			FString Key = MetaKeyValue.Key.ToString();
+			if (!Key.StartsWith(TEXT("/Script")))
 			{
-				FString Key = MetaKeyValue.Key.ToString();
-				if (!Key.StartsWith(TEXT("/Script")))
-				{
-					Map.Add(MetaKeyValue.Key, MetaKeyValue.Value);
-				}
+				Map.Add(MetaKeyValue.Key, MetaKeyValue.Value);
 			}
 		}
 		return Map;
 	}
 
-	virtual TMap<FName, FString>* GetMetaDataMap() const override
+	virtual TMap<FName, FString>& GetMetaDataMap() const override
 	{
-		return GetUObjectMetaDataMap();
+		return const_cast<TMap<FName, FString>&>(MetaData);
 	}
 
+	/**
+	 * Add meta data for the given definition
+	 */
+	virtual void AddMetaData(TMap<FName, FString>&& InMetaData) override;
+
 protected:
 	virtual void SetMetaDataHelper(const FName& Key, const TCHAR* InValue) override;
 
@@ -1498,6 +1477,7 @@ protected:
 	}
 
 private:
+	TMap<FName, FString> MetaData;
 	UObject* Object = nullptr;
 	FName Name = NAME_None;
 	EInternalObjectFlags InternalObjectFlags = EInternalObjectFlags::None;
@@ -1739,11 +1719,6 @@ public:
 		return TypePackageName;
 	}
 
-	/**
-	 * Add meta data for the given definition
-	 */
-	virtual void AddMetaData(TMap<FName, FString>&& InMetaData) override;
-
 	/**
 	 * Helper function that checks if the field is a dynamic type
 	 */
@@ -1778,7 +1753,7 @@ private:
 class FUnrealEnumDefinitionInfo : public FUnrealFieldDefinitionInfo
 {
 public:
-	FUnrealEnumDefinitionInfo(FUnrealSourceFile& InSourceFile, int32 InLineNumber, FString&& InNameCPP, FName InName);
+	FUnrealEnumDefinitionInfo(FUnrealSourceFile& InSourceFile, int32 InLineNumber, FString&& InNameCPP, FName InName, UEnum::ECppForm InCppForm, EUnderlyingEnumType InUnderlyingType);
 
 	virtual FUnrealEnumDefinitionInfo* AsEnum() override
 	{
@@ -1917,7 +1892,7 @@ public:
 	bool SetEnums(TArray<TPair<FName, int64>>& InNames, UEnum::ECppForm InCppForm, EEnumFlags InFlags = EEnumFlags::None)
 	{
 		Names = InNames;
-		CppForm = InCppForm;
+		check(CppForm == InCppForm);
 		EnumFlags = InFlags;
 #if UHT_ENABLE_ENGINE_TYPE_CHECKS
 		if (UEnum* Enum = GetEnumSafe())
@@ -1937,14 +1912,6 @@ public:
 		return UnderlyingType;
 	}
 
-	/** 
-	 * Set the underlying enum type
-	 */
-	void SetUnderlyingType(EUnderlyingEnumType InUnderlyingType)
-	{
-		UnderlyingType = InUnderlyingType;
-	}
-
 	/**
 	 * Return true if the enumeration is editor only
 	 */
@@ -1967,11 +1934,6 @@ protected:
 	 */
 	virtual void CreateUObjectEngineTypesInternal(ECreateEngineTypesPhase Phase) override;
 
-	virtual TMap<FName, FString>* GetMetaDataMap() const override
-	{
-		return const_cast<TMap<FName, FString>*>(&MetaData);
-	}
-
 	virtual FString GetMetaDataIndexName(int32 NameIndex) const
 	{
 		FString EnumName = GetNameStringByIndex(NameIndex);
@@ -1990,7 +1952,6 @@ protected:
 
 private:
 	TArray<TPair<FName, int64>> Names;
-	TMap<FName, FString> MetaData;
 	FString CppType;
 	EUnderlyingEnumType UnderlyingType = EUnderlyingEnumType::Unspecified;
 	EEnumFlags EnumFlags = EEnumFlags::None;
@@ -2457,11 +2418,6 @@ public:
 		MacroDeclaredLineNumber = InMacroDeclaredLineNumber;
 	}
 
-	virtual TMap<FName, FString>* GetMetaDataMap() const override
-	{
-		return const_cast<TMap<FName, FString>*>(&MetaData);
-	}
-
 protected:
 	/**
 	 * Create UObject engine types
@@ -2474,7 +2430,6 @@ protected:
 	virtual void PostParseFinalizeInternal() override;
 
 private:
-	TMap<FName, FString> MetaData;
 	int32 MacroDeclaredLineNumber = INDEX_NONE;
 	EStructFlags StructFlags = STRUCT_NoFlags;
 };
@@ -2671,11 +2626,6 @@ public:
 		return FunctionType == EFunctionType::Delegate || FunctionType == EFunctionType::SparseDelegate;
 	}
 
-	virtual TMap<FName, FString>* GetMetaDataMap() const override
-	{
-		return const_cast<TMap<FName, FString>*>(&MetaData);
-	}
-
 protected:
 	/**
 	 * Create UObject engine types
@@ -2695,7 +2645,6 @@ private:
 	/** return value for this function */
 	TSharedPtr<FUnrealPropertyDefinitionInfo> ReturnProperty;
 
-	TMap<FName, FString> MetaData;
 	FName SparseOwningClassName = NAME_None;
 	FName SparseDelegateName = NAME_None;
 	EFunctionType FunctionType = EFunctionType::Function;
@@ -3124,11 +3073,6 @@ public:
 
 	FString GetNameWithPrefix(EEnforceInterfacePrefix EnforceInterfacePrefix = EEnforceInterfacePrefix::None) const;
 
-	virtual TMap<FName, FString>* GetMetaDataMap() const override
-	{
-		return const_cast<TMap<FName, FString>*>(&MetaData);
-	}
-
 protected:
 	/**
 	 * Create UObject engine types
@@ -3148,7 +3092,6 @@ private:
 	/** Sets and validates 'ConfigName' property */
 	void SetAndValidateConfigName();
 
-	TMap<FName, FString> MetaData; // Actual live meta data
 	TMap<FName, FString> ParsedMetaData; // Meta data from the preparse phase
 	TArray<FString> ShowCategories;
 	TArray<FString> ShowFunctions;
