Fixing issue in UHT where the CLASS_Config and CLASS_HasInstancedReference flags might not get property propagated when the parent and child are in the same H file.

#rnx

diff --git a/Engine/Source/Programs/UnrealHeaderTool/Private/HeaderParser.cpp b/Engine/Source/Programs/UnrealHeaderTool/Private/HeaderParser.cpp
--- a/Engine/Source/Programs/UnrealHeaderTool/Private/HeaderParser.cpp
+++ b/Engine/Source/Programs/UnrealHeaderTool/Private/HeaderParser.cpp
@@ -5167,6 +5167,13 @@ FUnrealClassDefinitionInfo& FHeaderParser::ParseClassNameDeclaration(FString& De
  */
 void PostParsingClassSetup(FUnrealClassDefinitionInfo& ClassDef)
 {
+	// Since these two flags are computed in this method, we have to re-propagate these flags from the super
+	// just in case they were defined in this source file.
+	if (FUnrealClassDefinitionInfo* SuperClassDef = ClassDef.GetSuperClass())
+	{
+		ClassDef.SetClassFlags(SuperClassDef->GetClassFlags() & (CLASS_Config | CLASS_HasInstancedReference));
+	}
+
 	// Set all optimization ClassFlags based on property types
 	auto HasAllOptimizationClassFlags = [&ClassDef]()
 	{
