Modify propagation of some property flags to happen after all parsing is complete.  This removes a hidden restriction where parse order beyond explicit includes affected the output.
Enable concurrent parsing of headers.

Tested by running UHT on a sample large scale game 1000 times and verifying the output matched with a prior version of UHT.

#rnx

diff --git a/Engine/Source/Programs/UnrealHeaderTool/Private/CodeGenerator.cpp b/Engine/Source/Programs/UnrealHeaderTool/Private/CodeGenerator.cpp
--- a/Engine/Source/Programs/UnrealHeaderTool/Private/CodeGenerator.cpp
+++ b/Engine/Source/Programs/UnrealHeaderTool/Private/CodeGenerator.cpp
@@ -6589,8 +6589,7 @@ void ParseSourceFiles(TArray<FUnrealSourceFile*>& OrderedSourceFiles)
 	// Disable loading of objects outside of this package (or more exactly, objects which aren't UFields, CDO, or templates)
 	TGuardValue<bool> AutoRestoreVerifyObjectRefsFlag(GVerifyObjectReferencesOnly, true);
 
-// Concurrent parsing disabled
-#if 0 
+#if UHT_ENABLE_CONCURRENT_PARSING
 	/**
 	 * For every FUnrealSourceFile being processed, an instance of this class represents the data associated with generating the new output.
 	 */
@@ -6682,11 +6681,26 @@ void ParseSourceFiles(TArray<FUnrealSourceFile*>& OrderedSourceFiles)
 
 void PrepareTypesForExport(TArray<FUnrealPackageDefinitionInfo*>& PackageDefs)
 {
+	// Until we fix the issue with types not having sources, we need to do this
+	TArray<FUnrealTypeDefinitionInfo*> NoSourceTypeDefs;
+	GTypeDefinitionInfoMap.ForAllTypesByName([&NoSourceTypeDefs](FUnrealTypeDefinitionInfo& TypeDef)
+	{
+		if (!TypeDef.HasSource())
+		{
+			NoSourceTypeDefs.Add(&TypeDef);
+		}
+	});
+
+	for (FUnrealTypeDefinitionInfo* TypeDef : NoSourceTypeDefs)
+	{
+		TypeDef->PostParseFinalize(EPostParseFinalizePhase::PreCreateEngineTypes);
+	}
+
 	for (FUnrealPackageDefinitionInfo* PackageDef : PackageDefs)
 	{
 		FResults::Try([PackageDef]()
 		{
-			PackageDef->CreateUObjectEngineTypes(FUnrealTypeDefinitionInfo::ECreateEngineTypesPhase::Phase1);
+			PackageDef->PostParseFinalize(EPostParseFinalizePhase::PreCreateEngineTypes);
 		});
 	}
 
@@ -6694,7 +6708,7 @@ void PrepareTypesForExport(TArray<FUnrealPackageDefinitionInfo*>& PackageDefs)
 	{
 		FResults::Try([PackageDef]()
 		{
-			PackageDef->CreateUObjectEngineTypes(FUnrealTypeDefinitionInfo::ECreateEngineTypesPhase::Phase2);
+			PackageDef->CreateUObjectEngineTypes(ECreateEngineTypesPhase::Phase1);
 		});
 	}
 
@@ -6702,23 +6716,31 @@ void PrepareTypesForExport(TArray<FUnrealPackageDefinitionInfo*>& PackageDefs)
 	{
 		FResults::Try([PackageDef]()
 		{
-			PackageDef->PostParseFinalize();
+			PackageDef->CreateUObjectEngineTypes(ECreateEngineTypesPhase::Phase2);
 		});
 	}
-	FResults::WaitForErrorTasks();
 
-	// Until we fix the issue with types not having sources, we need to do this
-	GTypeDefinitionInfoMap.ForAllTypesByName([](FUnrealTypeDefinitionInfo& TypeDef)
+	for (FUnrealTypeDefinitionInfo* TypeDef : NoSourceTypeDefs)
 	{
-		if (!TypeDef.HasSource())
+		TypeDef->PostParseFinalize(EPostParseFinalizePhase::PostCreateEngineTypes);
+	}
+
+	for (FUnrealPackageDefinitionInfo* PackageDef : PackageDefs)
+	{
+		FResults::Try([PackageDef]()
 		{
-			TypeDef.PostParseFinalize();
-		}
+			PackageDef->PostParseFinalize(EPostParseFinalizePhase::PostCreateEngineTypes);
+		});
+	}
+	FResults::WaitForErrorTasks();
+
+	for (FUnrealTypeDefinitionInfo* TypeDef : NoSourceTypeDefs)
+	{
 		if (FUnrealObjectDefinitionInfo* ObjectDef = UHTCast<FUnrealObjectDefinitionInfo>(TypeDef))
 		{
 			check(ObjectDef->GetObject());
 		}
-	});
+	}
 }
 
 void Export(TArray<FUnrealPackageDefinitionInfo*>& PackageDefs, TArray<FUnrealSourceFile*>& OrderedSourceFiles)
diff --git a/Engine/Source/Programs/UnrealHeaderTool/Private/HeaderParser.cpp b/Engine/Source/Programs/UnrealHeaderTool/Private/HeaderParser.cpp
--- a/Engine/Source/Programs/UnrealHeaderTool/Private/HeaderParser.cpp
+++ b/Engine/Source/Programs/UnrealHeaderTool/Private/HeaderParser.cpp
@@ -2321,31 +2321,7 @@ void FHeaderParser::FixupDelegateProperties(FUnrealStructDefinitionInfo& StructD
 		{
 			if (SourceDelegateFunctionDef->HasAnyFunctionFlags( FUNC_MulticastDelegate ))
 			{
-				DelegatePropertyBase.FunctionDef = SourceDelegateFunctionDef;
-
-				if (PropertyDef->HasAnyPropertyFlags(CPF_BlueprintAssignable | CPF_BlueprintCallable))
-				{
-					for (TSharedRef<FUnrealPropertyDefinitionInfo> FuncParamDef : SourceDelegateFunctionDef->GetProperties())
-					{
-						if (!FPropertyTraits::IsSupportedByBlueprint(*FuncParamDef, false))
-						{
-							FString ExtendedCPPType;
-							FString CPPType = FuncParamDef->GetCPPType(&ExtendedCPPType);
-							LogError(TEXT("Type '%s%s' is not supported by blueprint. %s.%s"), *CPPType, *ExtendedCPPType, *SourceDelegateFunctionDef->GetName(), *FuncParamDef->GetName());
-						}
-
-						if(FuncParamDef->HasAllPropertyFlags(CPF_OutParm) && !FuncParamDef->HasAllPropertyFlags(CPF_ConstParm)  )
-						{
-							const bool bClassGeneratedFromBP = StructDef.IsDynamic();
-							const bool bAllowedArrayRefFromBP = bClassGeneratedFromBP && FuncParamDef->IsDynamicArray();
-							if (!bAllowedArrayRefFromBP)
-							{
-								LogError(TEXT("BlueprintAssignable delegates do not support non-const references at the moment. Function: %s Parameter: '%s'"), *SourceDelegateFunctionDef->GetName(), *FuncParamDef->GetName());
-							}
-						}
-					}
-				}
-
+				PropertyDef->SetDelegateFunctionSignature(*SourceDelegateFunctionDef);
 			}
 			else
 			{
@@ -2785,8 +2761,8 @@ void FHeaderParser::GetVarType(
 	FScope*                         Scope,
 	FPropertyBase&                  VarProperty,
 	EPropertyFlags                  Disallow,
-	const FToken*                   OuterPropertyType,
-	const EPropertyFlags*			OuterPropertyFlags,
+	EUHTPropertyType				OuterPropertyType,
+	EPropertyFlags					OuterPropertyFlags,
 	EPropertyDeclarationStyle::Type PropertyDeclarationStyle,
 	EVariableCategory               VariableCategory,
 	FIndexRange*                    ParsedVarIndexRange,
@@ -2825,7 +2801,7 @@ void FHeaderParser::GetVarType(
 	const bool bIsParamList = (VariableCategory != EVariableCategory::Member) && MatchIdentifier(TEXT("UPARAM"), ESearchCase::CaseSensitive);
 
 	// No specifiers are allowed inside a TArray
-	if ((OuterPropertyType == NULL) || !OuterPropertyType->Matches(TEXT("TArray"), ESearchCase::CaseSensitive))
+	if (OuterPropertyType != EUHTPropertyType::DynamicArray)
 	{
 		// New-style UPROPERTY() syntax 
 		if (PropertyDeclarationStyle == EPropertyDeclarationStyle::UPROPERTY || bIsParamList)
@@ -3511,9 +3487,7 @@ void FHeaderParser::GetVarType(
 	{
 		RequireSymbol( TEXT('<'), TEXT("'tarray'") );
 
-		EPropertyFlags VarTypePropertyFlags = Flags;
-
-		GetVarType(Scope, VarProperty, Disallow, &VarType, &VarTypePropertyFlags, EPropertyDeclarationStyle::None, VariableCategory);
+		GetVarType(Scope, VarProperty, Disallow, EUHTPropertyType::DynamicArray, Flags, EPropertyDeclarationStyle::None, VariableCategory);
 		if (VarProperty.IsContainer())
 		{
 			Throwf(TEXT("Nested containers are not supported.") );
@@ -3525,7 +3499,6 @@ void FHeaderParser::GetVarType(
 			bNativeConstTemplateArg = true;
 		}
 
-		VarTypePropertyFlags = VarProperty.PropertyFlags & (CPF_ContainsInstancedReference | CPF_InstancedReference); // propagate these to the array, we will fix them later
 		VarProperty.ArrayType = EArrayType::Dynamic;
 
 		FToken CloseTemplateToken;
@@ -3600,10 +3573,8 @@ void FHeaderParser::GetVarType(
 	{
 		RequireSymbol( TEXT('<'), TEXT("'tmap'") );
 
-		EPropertyFlags VarTypePropertyFlags = Flags;
-
 		FPropertyBase MapKeyType(CPT_None);
-		GetVarType(Scope, MapKeyType, Disallow, &VarType, &VarTypePropertyFlags, EPropertyDeclarationStyle::None, VariableCategory);
+		GetVarType(Scope, MapKeyType, Disallow, EUHTPropertyType::Map, Flags, EPropertyDeclarationStyle::None, VariableCategory);
 		if (MapKeyType.IsContainer())
 		{
 			Throwf(TEXT("Nested containers are not supported.") );
@@ -3631,17 +3602,15 @@ void FHeaderParser::GetVarType(
 			Throwf(TEXT("Missing value type while parsing TMap."));
 		}
 
-		GetVarType(Scope, VarProperty, Disallow, &VarType, &VarTypePropertyFlags, EPropertyDeclarationStyle::None, VariableCategory);
+		GetVarType(Scope, VarProperty, Disallow, EUHTPropertyType::Map, Flags, EPropertyDeclarationStyle::None, VariableCategory);
 		if (VarProperty.IsContainer())
 		{
 			Throwf(TEXT("Nested containers are not supported.") );
 		}
 		// TODO: Prevent sparse delegate types from being used in a container
-
-		EPropertyFlags InnerFlags = (MapKeyType.PropertyFlags | VarProperty.PropertyFlags) & (CPF_ContainsInstancedReference | CPF_InstancedReference); // propagate these to the map value, we will fix them later
-		VarTypePropertyFlags = InnerFlags;
+		MapKeyType.PropertyFlags = (MapKeyType.PropertyFlags & CPF_UObjectWrapper); // Make sure the 'UObjectWrapper' flag is maintained so that 'TMap<TSubclassOf<...>, ...>' works
 		VarProperty.MapKeyProp = MakeShared<FPropertyBase>(MoveTemp(MapKeyType));
-		VarProperty.MapKeyProp->PropertyFlags = InnerFlags | (VarProperty.MapKeyProp->PropertyFlags & CPF_UObjectWrapper); // Make sure the 'UObjectWrapper' flag is maintained so that 'TMap<TSubclassOf<...>, ...>' works
+		VarProperty.MapKeyProp->DisallowFlags = ~(CPF_ContainsInstancedReference | CPF_InstancedReference | CPF_UObjectWrapper);
 
 		FToken CloseTemplateToken;
 		if (!GetToken(CloseTemplateToken, /*bNoConsts=*/ true, ESymbolParseOption::CloseTemplateBracket))
@@ -3685,9 +3654,7 @@ void FHeaderParser::GetVarType(
 	{
 		RequireSymbol( TEXT('<'), TEXT("'tset'") );
 
-		EPropertyFlags VarTypePropertyFlags = Flags;
-
-		GetVarType(Scope, VarProperty, Disallow, &VarType, &VarTypePropertyFlags, EPropertyDeclarationStyle::None, VariableCategory);
+		GetVarType(Scope, VarProperty, Disallow, EUHTPropertyType::Set, Flags, EPropertyDeclarationStyle::None, VariableCategory);
 		if (VarProperty.IsContainer())
 		{
 			Throwf(TEXT("Nested containers are not supported.") );
@@ -3709,7 +3676,6 @@ void FHeaderParser::GetVarType(
 			LogError(TEXT("Replicated sets are not supported."));
 		}
 
-		VarTypePropertyFlags = VarProperty.PropertyFlags & (CPF_ContainsInstancedReference | CPF_InstancedReference); // propagate these to the set, we will fix them later
 		VarProperty.ArrayType = EArrayType::Set;
 
 		FToken CloseTemplateToken;
@@ -3879,11 +3845,6 @@ void FHeaderParser::GetVarType(
 			VarProperty = FPropertyBase(InFunctionDef.HasAnyFunctionFlags(FUNC_MulticastDelegate) ? CPT_MulticastDelegate : CPT_Delegate);
 			VarProperty.DelegateName = *InIdentifierStripped;
 			VarProperty.FunctionDef = &InFunctionDef;
-
-			if (!(Disallow & CPF_InstancedReference))
-			{
-				Flags |= CPF_InstancedReference;
-			}
 		};
 
 		if (!StructDef && MatchSymbol(TEXT("::")))
@@ -3933,10 +3894,7 @@ void FHeaderParser::GetVarType(
 			bHandledType = true;
 
 			VarProperty = FPropertyBase(*StructDef);
-			if (StructDef->HasAnyStructFlags(STRUCT_HasInstancedReference) && !(Disallow & CPF_ContainsInstancedReference))
-			{
-				Flags |= CPF_ContainsInstancedReference;
-			}
+
 			// Struct keyword in front of a struct is legal, we 'consume' it
 			bUnconsumedStructKeyword = false;
 		}
@@ -4107,12 +4065,6 @@ void FHeaderParser::GetVarType(
 						}
 					}
 
-					// Inherit instancing flags
-					if (TempClassDef->HierarchyHasAnyClassFlags(CLASS_DefaultToInstanced))
-					{
-						Flags |= ((CPF_InstancedReference | CPF_ExportObject) & (~Disallow));
-					}
-
 					// Eat the star that indicates this is a pointer to the UObject
 					if (!(Flags & CPF_UObjectWrapper))
 					{
@@ -4258,11 +4210,11 @@ void FHeaderParser::GetVarType(
 		}
 	}
 
-	VarProperty.PropertyExportFlags = ExportFlags;
-
 	// Set FPropertyBase info.
 	VarProperty.PropertyFlags        |= Flags | ImpliedFlags;
 	VarProperty.ImpliedPropertyFlags |= ImpliedFlags;
+	VarProperty.PropertyExportFlags = ExportFlags;
+	VarProperty.DisallowFlags = Disallow;
 
 	// Set the RepNotify name, if the variable needs it
 	if( VarProperty.PropertyFlags & CPF_RepNotify )
@@ -4316,7 +4268,7 @@ void FHeaderParser::GetVarType(
 	if (VariableCategory != EVariableCategory::Member)
 	{
 		// These conditions are checked externally for struct/member variables where the flag can be inferred later on from the variable name itself
-		ValidatePropertyIsDeprecatedIfNecessary(VarProperty, OuterPropertyFlags);
+		ValidatePropertyIsDeprecatedIfNecessary(VarProperty, OuterPropertyType, OuterPropertyFlags);
 	}
 
 	// Check for invalid transients
@@ -5167,34 +5119,22 @@ FUnrealClassDefinitionInfo& FHeaderParser::ParseClassNameDeclaration(FString& De
  */
 void PostParsingClassSetup(FUnrealClassDefinitionInfo& ClassDef)
 {
-	// Since these two flags are computed in this method, we have to re-propagate these flags from the super
+	//@TODO: Move to post parse finalization?
+
+	// Since this flag is computed in this method, we have to re-propagate the flag from the super
 	// just in case they were defined in this source file.
 	if (FUnrealClassDefinitionInfo* SuperClassDef = ClassDef.GetSuperClass())
 	{
-		ClassDef.SetClassFlags(SuperClassDef->GetClassFlags() & (CLASS_Config | CLASS_HasInstancedReference));
+		ClassDef.SetClassFlags(SuperClassDef->GetClassFlags() & CLASS_Config);
 	}
 
-	// Set all optimization ClassFlags based on property types
-	auto HasAllOptimizationClassFlags = [&ClassDef]()
+	// Set the class config flag if any properties have config
+	for (TSharedRef<FUnrealPropertyDefinitionInfo> PropertyDef : ClassDef.GetProperties())
 	{
-		return (ClassDef.HasAllClassFlags(CLASS_Config | CLASS_HasInstancedReference));
-	};
-
-	if (!HasAllOptimizationClassFlags())
-	{
-		for (TSharedRef<FUnrealPropertyDefinitionInfo> PropertyDef : ClassDef.GetProperties())
+		if (PropertyDef->HasAnyPropertyFlags(CPF_Config))
 		{
-			if (PropertyDef->HasAnyPropertyFlags(CPF_Config))
-			{
-				ClassDef.SetClassFlags(CLASS_Config);
-				if (HasAllOptimizationClassFlags()) break;
-			}
-
-			if (PropertyDef->ContainsInstancedObjectProperty())
-			{
-				ClassDef.SetClassFlags(CLASS_HasInstancedReference);
-				if (HasAllOptimizationClassFlags()) break;
-			}
+			ClassDef.SetClassFlags(CLASS_Config);
+			break;
 		}
 	}
 
@@ -5728,7 +5668,7 @@ void FHeaderParser::ParseParameterList(FUnrealFunctionDefinitionInfo& FunctionDe
 		// Get parameter type.
 		FPropertyBase Property(CPT_None);
 		EVariableCategory VariableCategory = FunctionDef.HasAnyFunctionFlags(FUNC_Net) ? EVariableCategory::ReplicatedParameter : EVariableCategory::RegularParameter;
-		GetVarType(GetCurrentScope(), Property, ~(CPF_ParmFlags | CPF_AutoWeak | CPF_RepSkip | CPF_UObjectWrapper | CPF_NativeAccessSpecifiers), nullptr, nullptr, EPropertyDeclarationStyle::None, VariableCategory);
+		GetVarType(GetCurrentScope(), Property, ~(CPF_ParmFlags | CPF_AutoWeak | CPF_RepSkip | CPF_UObjectWrapper | CPF_NativeAccessSpecifiers), EUHTPropertyType::None, CPF_None, EPropertyDeclarationStyle::None, VariableCategory);
 		Property.PropertyFlags |= CPF_Parm;
 
 		if (bExpectCommaBeforeName)
@@ -5981,7 +5921,7 @@ FUnrealFunctionDefinitionInfo& FHeaderParser::CompileDelegateDeclaration(const T
 
 	if (bHasReturnValue)
 	{
-		GetVarType(GetCurrentScope(), ReturnType, CPF_None, nullptr, nullptr, EPropertyDeclarationStyle::None, EVariableCategory::Return);
+		GetVarType(GetCurrentScope(), ReturnType, CPF_None, EUHTPropertyType::None, CPF_None, EPropertyDeclarationStyle::None, EVariableCategory::Return);
 		RequireSymbol(TEXT(','), CurrentScopeName);
 	}
 
@@ -6359,7 +6299,7 @@ FUnrealFunctionDefinitionInfo& FHeaderParser::CompileFunctionDeclaration()
 	FPropertyBase ReturnType( CPT_None );
 
 	// C++ style functions always have a return value type, even if it's void
-	GetVarType(GetCurrentScope(), ReturnType, CPF_None, nullptr, nullptr, EPropertyDeclarationStyle::None, EVariableCategory::Return);
+	GetVarType(GetCurrentScope(), ReturnType, CPF_None, EUHTPropertyType::None, CPF_None, EPropertyDeclarationStyle::None, EVariableCategory::Return);
 	bool bHasReturnValue = ReturnType.Type != CPT_None;
 
 	// Skip whitespaces to get InputPos exactly on beginning of function name.
@@ -6554,24 +6494,6 @@ FUnrealFunctionDefinitionInfo& FHeaderParser::CompileFunctionDeclaration()
 		Throwf(TEXT("Function '%s': Base implementation of RPCs cannot be in a state. Add a stub outside state scope."), *FuncDef.GetName());
 	}
 
-	if (FuncDef.HasAnyFunctionFlags(FUNC_BlueprintCallable | FUNC_BlueprintEvent))
-	{
-		for (TSharedRef<FUnrealPropertyDefinitionInfo> PropertyDef : FuncDef.GetProperties())
-		{
-			if (PropertyDef->IsStaticArray())
-			{
-				Throwf(TEXT("Static array cannot be exposed to blueprint. Function: %s Parameter %s\n"), *FuncDef.GetName(), *PropertyDef->GetName());
-			}
-
-			if (!FPropertyTraits::IsSupportedByBlueprint(*PropertyDef, false))
-			{
-				FString ExtendedCPPType;
-				FString CPPType = PropertyDef->GetCPPType(&ExtendedCPPType);
-				LogError(TEXT("Type '%s%s' is not supported by blueprint. %s.%s"), *CPPType, *ExtendedCPPType, *FuncDef.GetName(), *PropertyDef->GetName());
-			}
-		}
-	}
-
 	// Just declaring a function, so end the nesting.
 	PostPopFunctionDeclaration(FuncDef);
 
@@ -6715,10 +6637,10 @@ bool FHeaderParser::IsBitfieldProperty(ELayoutMacroType LayoutMacroType)
 	return bIsBitfield;
 }
 
-void FHeaderParser::ValidatePropertyIsDeprecatedIfNecessary(const FPropertyBase& VarProperty, const EPropertyFlags* OuterPropertyFlags)
+void FHeaderParser::ValidatePropertyIsDeprecatedIfNecessary(const FPropertyBase& VarProperty, EUHTPropertyType OuterPropertyType, EPropertyFlags OuterPropertyFlags)
 {
-	// If the property is in an array that has been deprecated, then we don't need to do any further tests 
-	if (OuterPropertyFlags != nullptr && (*OuterPropertyFlags & CPF_Deprecated) != 0)
+	// If the property is in a container that has been deprecated, then we don't need to do any further tests 
+	if (OuterPropertyType != EUHTPropertyType::None && (OuterPropertyFlags & CPF_Deprecated) != 0)
 	{
 		return;
 	}
@@ -6842,7 +6764,7 @@ void FHeaderParser::CompileVariableDeclaration(FUnrealStructDefinitionInfo& Stru
 	FPropertyBase OriginalPropertyBase(CPT_None);
 	FIndexRange TypeRange;
 	ELayoutMacroType LayoutMacroType = ELayoutMacroType::None;
-	GetVarType(&*StructDef.GetScope(), OriginalPropertyBase, DisallowFlags, /*OuterPropertyType=*/ nullptr, nullptr, EPropertyDeclarationStyle::UPROPERTY, EVariableCategory::Member, &TypeRange, &LayoutMacroType);
+	GetVarType(&*StructDef.GetScope(), OriginalPropertyBase, DisallowFlags, EUHTPropertyType::None, CPF_None, EPropertyDeclarationStyle::UPROPERTY, EVariableCategory::Member, &TypeRange, &LayoutMacroType);
 	OriginalPropertyBase.PropertyFlags |= EdFlags;
 
 	FString* Category = OriginalPropertyBase.MetaData.Find(NAME_Category);
@@ -6934,7 +6856,7 @@ void FHeaderParser::CompileVariableDeclaration(FUnrealStructDefinitionInfo& Stru
 		}
 
 		// Deprecation validation
-		ValidatePropertyIsDeprecatedIfNecessary(PropertyBase, NULL);
+		ValidatePropertyIsDeprecatedIfNecessary(PropertyBase, EUHTPropertyType::None, CPF_None);
 
 		if (TopNest->NestType != ENestType::FunctionDeclaration)
 		{
@@ -6949,34 +6871,6 @@ void FHeaderParser::CompileVariableDeclaration(FUnrealStructDefinitionInfo& Stru
 		// we'll need any metadata tags we parsed later on when we call ConvertEOLCommentToTooltip() so the tags aren't clobbered
 		OriginalPropertyBase.MetaData = PropertyBase.MetaData;
 
-		if (ScriptStructDef != nullptr)
-		{
-			if (NewPropDef.ContainsInstancedObjectProperty())
-			{
-				ScriptStructDef->SetStructFlags(STRUCT_HasInstancedReference);
-			}
-		}
-
-		if (NewPropDef.HasAnyPropertyFlags(CPF_BlueprintVisible))
-		{
-			if (ScriptStructDef != nullptr && !StructDef.GetBoolMetaDataHierarchical(FHeaderParserNames::NAME_BlueprintType))
-			{
-				LogError(TEXT("Cannot expose property to blueprints in a struct that is not a BlueprintType. %s.%s"), *StructDef.GetName(), *NewPropDef.GetName());
-			}
-
-			if (NewPropDef.IsStaticArray())
-			{
-				LogError(TEXT("Static array cannot be exposed to blueprint %s.%s"), *StructDef.GetName(), *NewPropDef.GetName());
-			}
-
-			if (!FPropertyTraits::IsSupportedByBlueprint(NewPropDef, true))
-			{
-				FString ExtendedCPPType;
-				FString CPPType = NewPropDef.GetCPPType(&ExtendedCPPType);
-				LogError(TEXT("Type '%s%s' is not supported by blueprint. %s.%s"), *CPPType, *ExtendedCPPType, *StructDef.GetName(), *NewPropDef.GetName());
-			}
-		}
-
 		if (LayoutMacroType != ELayoutMacroType::None || !MatchSymbol(TEXT(',')))
 		{
 			break;
diff --git a/Engine/Source/Programs/UnrealHeaderTool/Private/HeaderParser.h b/Engine/Source/Programs/UnrealHeaderTool/Private/HeaderParser.h
--- a/Engine/Source/Programs/UnrealHeaderTool/Private/HeaderParser.h
+++ b/Engine/Source/Programs/UnrealHeaderTool/Private/HeaderParser.h
@@ -437,7 +437,8 @@ protected:
 	 * @param   Scope                     struct to create the property in
 	 * @param   VarProperty               will be filled in with type and property flag data for the property declaration that was parsed
 	 * @param   Disallow                  contains a mask of variable modifiers that are disallowed in this context
-	 * @param   OuterPropertyType         only specified when compiling the inner properties for arrays or maps.  corresponds to the FToken for the outer property declaration.
+	 * @param   OuterPropertyType         only specified when compiling the inner properties for arrays or maps.
+	 * @param	OuterPropertyFlags		  flags associated with the outer property being compiled.
 	 * @param   PropertyDeclarationStyle  if the variable is defined with a UPROPERTY
 	 * @param   VariableCategory          what kind of variable is being parsed
 	 * @param   ParsedVarIndexRange       The source text [Start, End) index range for the parsed type.
@@ -446,8 +447,8 @@ protected:
 		FScope*							Scope,
 		FPropertyBase&                  VarProperty,
 		EPropertyFlags                  Disallow,
-		const FToken*                   OuterPropertyType,
-		const EPropertyFlags*			OuterPropertyFlags,
+		EUHTPropertyType				OuterPropertyType,
+		EPropertyFlags					OuterPropertyFlags,
 		EPropertyDeclarationStyle::Type PropertyDeclarationStyle,
 		EVariableCategory               VariableCategory,
 		FIndexRange*                    ParsedVarIndexRange = nullptr,
@@ -547,7 +548,7 @@ protected:
 	void InitScriptLocation( FScriptLocation& Retry );
 	void ReturnToLocation( const FScriptLocation& Retry, bool Binary=1, bool Text=1 );
 
-	void ValidatePropertyIsDeprecatedIfNecessary(const FPropertyBase& VarProperty, const EPropertyFlags* OuterPropertyFlags);
+	void ValidatePropertyIsDeprecatedIfNecessary(const FPropertyBase& VarProperty, EUHTPropertyType OuterPropertyType, EPropertyFlags OuterPropertyFlags);
 
 	// Cache of ScriptStructs that have been validated for Net Replication and RPC
 	TSet<FUnrealScriptStructDefinitionInfo*> ScriptStructsValidForNet;
diff --git a/Engine/Source/Programs/UnrealHeaderTool/Private/ParserHelper.h b/Engine/Source/Programs/UnrealHeaderTool/Private/ParserHelper.h
--- a/Engine/Source/Programs/UnrealHeaderTool/Private/ParserHelper.h
+++ b/Engine/Source/Programs/UnrealHeaderTool/Private/ParserHelper.h
@@ -173,6 +173,7 @@ public:
 	EAllocatorType      AllocatorType = EAllocatorType::Default;
 	EPropertyFlags      PropertyFlags = CPF_None;
 	EPropertyFlags      ImpliedPropertyFlags = CPF_None;
+	EPropertyFlags		DisallowFlags = ~CPF_None; 
 	ERefQualifier::Type RefQualifier = ERefQualifier::None; // This is needed because of legacy stuff - FString mangles the flags for reasons that have become lost in time but we need this info for testing for invalid replicated function signatures.
 
 	TSharedPtr<FPropertyBase> MapKeyProp;
diff --git a/Engine/Source/Programs/UnrealHeaderTool/Private/PropertyTypes.cpp b/Engine/Source/Programs/UnrealHeaderTool/Private/PropertyTypes.cpp
--- a/Engine/Source/Programs/UnrealHeaderTool/Private/PropertyTypes.cpp
+++ b/Engine/Source/Programs/UnrealHeaderTool/Private/PropertyTypes.cpp
@@ -94,6 +94,29 @@ namespace
 {
 	static const FName NAME_ArraySizeEnum(TEXT("ArraySizeEnum"));
 
+	struct FPropertyFlagsChanges
+	{
+		EPropertyFlags Added;
+		EPropertyFlags Removed;
+	};
+
+	struct FPropertyFlagsChangeDetector
+	{
+		FPropertyFlagsChangeDetector(FUnrealPropertyDefinitionInfo& InPropDef)
+			: PropDef(InPropDef)
+			, OldFlags(InPropDef.GetPropertyFlags())
+		{}
+
+		FPropertyFlagsChanges Changes()
+		{
+			EPropertyFlags NewFlags = PropDef.GetPropertyFlags();
+			return FPropertyFlagsChanges{ NewFlags & ~OldFlags, OldFlags & ~NewFlags };
+		}
+
+		FUnrealPropertyDefinitionInfo& PropDef;
+		EPropertyFlags OldFlags;
+	};
+
 	void PropagateFlagsFromInnerAndHandlePersistentInstanceMetadata(EPropertyFlags& DestFlags, const TMap<FName, FString>& InMetaData, FUnrealPropertyDefinitionInfo& InnerDef)
 	{
 		// Copy some of the property flags to the container property.
@@ -303,6 +326,15 @@ namespace
 		}
 	};
 
+	template<typename TraitsType>
+	struct PostParseFinalizeDispatch
+	{
+		FPropertyFlagsChanges operator()(FUnrealPropertyDefinitionInfo& PropDef, EPropertyFlags& AddedFlags)
+		{
+			return TraitsType::PostParseFinalize(PropDef, AddedFlags);
+		}
+	};
+
 	//----------------------------------------------------------------------------------------------------------------------------------------------------------
 	//----------------------------------------------------------------------------------------------------------------------------------------------------------
 	// Helper Methods
@@ -325,6 +357,12 @@ namespace
 	{
 		return PropertyTypeDispatch<IsSupportedByBlueprintDispatch, false, bool>(PropDef, bMemberVariable);
 	}
+
+	template <bool bHandleContainers>
+	FPropertyFlagsChanges PostParseFinalizeHelper(FUnrealPropertyDefinitionInfo& PropDef, EPropertyFlags& AddedFlags)
+	{
+		return PropertyTypeDispatch<PostParseFinalizeDispatch, bHandleContainers, FPropertyFlagsChanges>(PropDef, AddedFlags);
+	}
 }
 
 /**
@@ -393,6 +431,18 @@ struct FPropertyTypeTraitsBase
 	//static FString GetCPPType(const FUnrealPropertyDefinitionInfo& PropDef, FString* ExtendedTypeText, uint32 CPPExportFlags);
 
 	//static FString GetCPPTypeForwardDeclaration(const FUnrealPropertyDefinitionInfo& PropDef);
+
+	/**
+	 * Handle any final initialization after parsing has completed
+	 *
+	 * @param PropDef The property in question
+	 * @param AddedFlags Flags added to the property that need to be propagated to the map key type
+	 * @return Collection of properties added and removed
+	 */
+	static FPropertyFlagsChanges PostParseFinalize(FUnrealPropertyDefinitionInfo& PropDef, EPropertyFlags& AddedFlags)
+	{
+		return FPropertyFlagsChanges{ CPF_None, CPF_None };
+	}
 };
 
 //----------------------------------------------------------------------------------------------------------------------------------------------------------
@@ -424,9 +474,14 @@ struct FPropertyTypeTraitsByte : public FPropertyTypeTraitsNumericBase
 
 	static void CreateProperty(FUnrealPropertyDefinitionInfo& PropDef)
 	{
+		const FPropertyBase& VarProperty = PropDef.GetPropertyBase();
 #if UHT_ENABLE_VALUE_PROPERTY_TAG
-		PropDef.GetUnrealSourceFile().AddTypeDefIncludeIfNeeded(PropDef.GetPropertyBase().EnumDef);
+		PropDef.GetUnrealSourceFile().AddTypeDefIncludeIfNeeded(VarProperty.EnumDef);
 #endif
+		if (VarProperty.EnumDef)
+		{
+			VarProperty.EnumDef->AddReferencingProperty(PropDef);
+		}
 	}
 
 	static FProperty* CreateEngineType(FUnrealPropertyDefinitionInfo& PropDef, FFieldVariant Scope, const FName& Name, EObjectFlags ObjectFlags)
@@ -983,6 +1038,8 @@ struct FPropertyTypeTraitsEnum : public FPropertyTypeTraitsBase
 			return FPropertyTypeTraitsByte::CreateProperty(PropDef);
 		}
 
+		VarProperty.EnumDef->AddReferencingProperty(PropDef);
+
 #if UHT_ENABLE_VALUE_PROPERTY_TAG
 		PropDef.GetUnrealSourceFile().AddTypeDefIncludeIfNeeded(VarProperty.EnumDef);
 #endif
@@ -1114,19 +1171,41 @@ struct FPropertyTypeTraitsObjectBase : public FPropertyTypeTraitsBase
 		}
 		return bIsNull; // always return as null is the only the processing we can do for object defaults
 	}
-};
 
-struct FPropertyTypeTraitsObjectReference : public FPropertyTypeTraitsObjectBase
-{
 	static void CreateProperty(FUnrealPropertyDefinitionInfo& PropDef)
 	{
 		FPropertyBase& VarProperty = PropDef.GetPropertyBase();
 		check(VarProperty.ClassDef);
+		VarProperty.ClassDef->AddReferencingProperty(PropDef);
 
 #if UHT_ENABLE_PTR_PROPERTY_TAG
 		PropDef.GetUnrealSourceFile().AddTypeDefIncludeIfNeeded(VarProperty.ClassDef);
 #endif
+	}
+
+	static FPropertyFlagsChanges PostParseFinalize(FUnrealPropertyDefinitionInfo& PropDef, EPropertyFlags& AddedFlags)
+	{
+		FPropertyFlagsChangeDetector Detector(PropDef);
+		FPropertyBase& VarProperty = PropDef.GetPropertyBase();
 
+		// Inherit instancing flags
+		if (VarProperty.ClassDef->HierarchyHasAnyClassFlags(CLASS_DefaultToInstanced))
+		{
+			VarProperty.PropertyFlags |= (CPF_InstancedReference | CPF_ExportObject) & (~VarProperty.DisallowFlags);
+		}
+		AddedFlags |= Detector.Changes().Added;
+		return Detector.Changes();
+	}
+};
+
+struct FPropertyTypeTraitsObjectReference : public FPropertyTypeTraitsObjectBase
+{
+	static FPropertyFlagsChanges PostParseFinalize(FUnrealPropertyDefinitionInfo& PropDef, EPropertyFlags& AddedFlags)
+	{
+		FPropertyFlagsChangeDetector Detector(PropDef);
+		FPropertyTypeTraitsObjectBase::PostParseFinalize(PropDef, AddedFlags);
+
+		FPropertyBase& VarProperty = PropDef.GetPropertyBase();
 		if (VarProperty.ClassDef->IsChildOf(*GUClassDef))
 		{
 		}
@@ -1138,6 +1217,7 @@ struct FPropertyTypeTraitsObjectReference : public FPropertyTypeTraitsObjectBase
 				AddEditInlineMetaData(VarProperty.MetaData);
 			}
 		}
+		return Detector.Changes();
 	}
 
 	static FProperty* CreateEngineType(FUnrealPropertyDefinitionInfo& PropDef, FFieldVariant Scope, const FName& Name, EObjectFlags ObjectFlags)
@@ -1216,15 +1296,6 @@ struct FPropertyTypeTraitsObjectReference : public FPropertyTypeTraitsObjectBase
 
 struct FPropertyTypeTraitsWeakObjectReference : public FPropertyTypeTraitsObjectBase
 {
-	static void CreateProperty(FUnrealPropertyDefinitionInfo& PropDef)
-	{
-		const FPropertyBase& VarProperty = PropDef.GetPropertyBase();
-		check(VarProperty.ClassDef);
-#if UHT_ENABLE_PTR_PROPERTY_TAG
-		PropDef.GetUnrealSourceFile().AddTypeDefIncludeIfNeeded(VarProperty.ClassDef);
-#endif
-	}
-
 	static FProperty* CreateEngineType(FUnrealPropertyDefinitionInfo& PropDef, FFieldVariant Scope, const FName& Name, EObjectFlags ObjectFlags)
 	{
 		const FPropertyBase& VarProperty = PropDef.GetPropertyBase();
@@ -1266,15 +1337,6 @@ struct FPropertyTypeTraitsWeakObjectReference : public FPropertyTypeTraitsObject
 
 struct FPropertyTypeTraitsLazyObjectReference : public FPropertyTypeTraitsObjectBase
 {
-	static void CreateProperty(FUnrealPropertyDefinitionInfo& PropDef)
-	{
-		const FPropertyBase& VarProperty = PropDef.GetPropertyBase();
-		check(VarProperty.ClassDef);
-#if UHT_ENABLE_PTR_PROPERTY_TAG
-		PropDef.GetUnrealSourceFile().AddTypeDefIncludeIfNeeded(VarProperty.ClassDef);
-#endif
-	}
-
 	static FProperty* CreateEngineType(FUnrealPropertyDefinitionInfo& PropDef, FFieldVariant Scope, const FName& Name, EObjectFlags ObjectFlags)
 	{
 		const FPropertyBase& VarProperty = PropDef.GetPropertyBase();
@@ -1306,15 +1368,12 @@ struct FPropertyTypeTraitsLazyObjectReference : public FPropertyTypeTraitsObject
 
 struct FPropertyTypeTraitsObjectPtrReference : public FPropertyTypeTraitsObjectBase
 {
-	static void CreateProperty(FUnrealPropertyDefinitionInfo& PropDef)
+	static FPropertyFlagsChanges PostParseFinalize(FUnrealPropertyDefinitionInfo& PropDef, EPropertyFlags& AddedFlags)
 	{
-		FPropertyBase& VarProperty = PropDef.GetPropertyBase();
-		check(VarProperty.ClassDef);
-
-#if UHT_ENABLE_PTR_PROPERTY_TAG
-		PropDef.GetUnrealSourceFile().AddTypeDefIncludeIfNeeded(VarProperty.ClassDef);
-#endif
+		FPropertyFlagsChangeDetector Detector(PropDef);
+		FPropertyTypeTraitsObjectBase::PostParseFinalize(PropDef, AddedFlags);
 
+		FPropertyBase& VarProperty = PropDef.GetPropertyBase();
 		if (VarProperty.ClassDef->IsChildOf(*GUClassDef))
 		{
 		}
@@ -1326,6 +1385,7 @@ struct FPropertyTypeTraitsObjectPtrReference : public FPropertyTypeTraitsObjectB
 				AddEditInlineMetaData(VarProperty.MetaData);
 			}
 		}
+		return Detector.Changes();
 	}
 
 	static FProperty* CreateEngineType(FUnrealPropertyDefinitionInfo& PropDef, FFieldVariant Scope, const FName& Name, EObjectFlags ObjectFlags)
@@ -1397,15 +1457,6 @@ struct FPropertyTypeTraitsObjectPtrReference : public FPropertyTypeTraitsObjectB
 
 struct FPropertyTypeTraitsSoftObjectReference : public FPropertyTypeTraitsObjectBase
 {
-	static void CreateProperty(FUnrealPropertyDefinitionInfo& PropDef)
-	{
-		FPropertyBase& VarProperty = PropDef.GetPropertyBase();
-		check(VarProperty.ClassDef);
-#if UHT_ENABLE_PTR_PROPERTY_TAG
-		PropDef.GetUnrealSourceFile().AddTypeDefIncludeIfNeeded(VarProperty.ClassDef);
-#endif
-	}
-
 	static FProperty* CreateEngineType(FUnrealPropertyDefinitionInfo& PropDef, FFieldVariant Scope, const FName& Name, EObjectFlags ObjectFlags)
 	{
 		FPropertyBase& VarProperty = PropDef.GetPropertyBase();
@@ -1479,13 +1530,12 @@ struct FPropertyTypeTraitsInterface : public FPropertyTypeTraitsBase
 {
 	static void CreateProperty(FUnrealPropertyDefinitionInfo& PropDef)
 	{
-		FPropertyBase& VarProperty = PropDef.GetPropertyBase();
-		check(VarProperty.ClassDef);
-		check(VarProperty.ClassDef->HasAnyClassFlags(CLASS_Interface));
+		FPropertyTypeTraitsObjectBase::CreateProperty(PropDef);
+	}
 
-#if UHT_ENABLE_PTR_PROPERTY_TAG
-		PropDef.GetUnrealSourceFile().AddTypeDefIncludeIfNeeded(VarProperty.ClassDef);
-#endif
+	static FPropertyFlagsChanges PostParseFinalize(FUnrealPropertyDefinitionInfo& PropDef, EPropertyFlags& AddedFlags)
+	{
+		return FPropertyTypeTraitsObjectBase::PostParseFinalize(PropDef, AddedFlags);
 	}
 
 	static FProperty* CreateEngineType(FUnrealPropertyDefinitionInfo& PropDef, FFieldVariant Scope, const FName& Name, EObjectFlags ObjectFlags)
@@ -1914,15 +1964,25 @@ struct FPropertyTypeTraitsStruct : public FPropertyTypeTraitsBase
 	{
 		FPropertyBase& VarProperty = PropDef.GetPropertyBase();
 		FUnrealScriptStructDefinitionInfo* ScriptStructDef = VarProperty.ScriptStructDef;
+		ScriptStructDef->AddReferencingProperty(PropDef);
 
 #if UHT_ENABLE_VALUE_PROPERTY_TAG
-		PropDef.GetUnrealSourceFile().AddTypeDefIncludeIfNeeded(VarProperty.ScriptStructDef);
+		PropDef.GetUnrealSourceFile().AddTypeDefIncludeIfNeeded(ScriptStructDef);
 #endif
+	}
 
-		if (ScriptStructDef->HasAnyStructFlags(STRUCT_HasInstancedReference))
+	static FPropertyFlagsChanges PostParseFinalize(FUnrealPropertyDefinitionInfo& PropDef, EPropertyFlags& AddedFlags)
+	{
+		FPropertyFlagsChangeDetector Detector(PropDef);
+		FPropertyBase& VarProperty = PropDef.GetPropertyBase();
+
+		// We always add regardless of disallow
+		if (VarProperty.ScriptStructDef->HasAnyStructFlags(STRUCT_HasInstancedReference))
 		{
 			VarProperty.PropertyFlags |= CPF_ContainsInstancedReference;
 		}
+		AddedFlags |= Detector.Changes().Added;
+		return Detector.Changes();
 	}
 
 	static FProperty* CreateEngineType(FUnrealPropertyDefinitionInfo& PropDef, FFieldVariant Scope, const FName& Name, EObjectFlags ObjectFlags)
@@ -1946,7 +2006,7 @@ struct FPropertyTypeTraitsStruct : public FPropertyTypeTraitsBase
 
 	static FString GetCPPType(const FUnrealPropertyDefinitionInfo& PropDef, FString* ExtendedTypeText, uint32 CPPExportFlags)
 	{
-		return FString::Printf(TEXT("F%s"), *PropDef.GetPropertyBase().ClassDef->GetName());
+		return FString::Printf(TEXT("F%s"), *PropDef.GetPropertyBase().ScriptStructDef->GetName());
 	}
 
 	static FString GetCPPTypeForwardDeclaration(const FUnrealPropertyDefinitionInfo& PropDef)
@@ -1958,7 +2018,7 @@ struct FPropertyTypeTraitsStruct : public FPropertyTypeTraitsBase
 			return FString();
 		}
 			
-		return FString::Printf(TEXT("struct F%s;"), *PropDef.GetPropertyBase().ClassDef->GetName());
+		return FString::Printf(TEXT("struct F%s;"), *PropDef.GetPropertyBase().ScriptStructDef->GetName());
 	}
 };
 
@@ -1968,17 +2028,22 @@ struct FPropertyTypeTraitsDelegate : public FPropertyTypeTraitsBase
 	{
 	}
 
+	static FPropertyFlagsChanges PostParseFinalize(FUnrealPropertyDefinitionInfo& PropDef, EPropertyFlags& AddedFlags)
+	{
+		FPropertyFlagsChangeDetector Detector(PropDef);
+		FPropertyBase& VarProperty = PropDef.GetPropertyBase();
+		VarProperty.PropertyFlags |= CPF_InstancedReference & (~VarProperty.DisallowFlags);
+		AddedFlags |= Detector.Changes().Added;
+		return Detector.Changes();
+	}
+
 	static FProperty* CreateEngineType(FUnrealPropertyDefinitionInfo& PropDef, FFieldVariant Scope, const FName& Name, EObjectFlags ObjectFlags)
 	{
 		const FPropertyBase& VarProperty = PropDef.GetPropertyBase();
+		check(VarProperty.FunctionDef);
 
 		FDelegateProperty* Result = new FDelegateProperty(Scope, Name, ObjectFlags);
-
-		// Once we fully transition to post parse engine type creation, this should always be set.
-		if (VarProperty.FunctionDef != nullptr)
-		{
-			Result->SignatureFunction = VarProperty.FunctionDef->GetFunction();
-		}
+		Result->SignatureFunction = VarProperty.FunctionDef->GetFunction();
 		return Result;
 	}
 
@@ -2045,6 +2110,15 @@ struct FPropertyTypeTraitsMulticastDelegate : public FPropertyTypeTraitsBase
 	{
 	}
 
+	static FPropertyFlagsChanges PostParseFinalize(FUnrealPropertyDefinitionInfo& PropDef, EPropertyFlags& AddedFlags)
+	{
+		FPropertyFlagsChangeDetector Detector(PropDef);
+		FPropertyBase& VarProperty = PropDef.GetPropertyBase();
+		VarProperty.PropertyFlags |= CPF_InstancedReference & (~VarProperty.DisallowFlags);
+		AddedFlags |= Detector.Changes().Added;
+		return Detector.Changes();
+	}
+
 	static FProperty* CreateEngineType(FUnrealPropertyDefinitionInfo& PropDef, FFieldVariant Scope, const FName& Name, EObjectFlags ObjectFlags)
 	{
 		const FPropertyBase& VarProperty = PropDef.GetPropertyBase();
@@ -2193,6 +2267,11 @@ struct FPropertyTypeTraitsStaticArray : public FPropertyTypeTraitsBase
 		CreatePropertyHelper<false>(PropDef);
 	}
 
+	static FPropertyFlagsChanges PostParseFinalize(FUnrealPropertyDefinitionInfo& PropDef, EPropertyFlags& AddedFlags)
+	{
+		return PostParseFinalizeHelper<false>(PropDef, AddedFlags);
+	}
+
 	static FProperty* CreateEngineType(FUnrealPropertyDefinitionInfo& PropDef, FFieldVariant Scope, const FName& Name, EObjectFlags ObjectFlags)
 	{
 
@@ -2297,16 +2376,26 @@ struct FPropertyTypeTraitsDynamicArray : public FPropertyTypeTraitsBase
 		InnerVarProperty.MetaData.Empty();
 		PropDef.SetValuePropDef(FPropertyTraits::CreateProperty(InnerVarProperty, PropDef, PropDef.GetFName(), 
 			PropDef.GetVariableCategory(), ACCESS_Public, PropDef.GetArrayDimensions(), PropDef.GetUnrealSourceFile(), PropDef.GetLineNumber(), PropDef.GetParsePosition()));
-		FUnrealPropertyDefinitionInfo& InnerPropDef = PropDef.GetValuePropDef();
 
-		VarProperty.PropertyFlags = InnerPropDef.GetPropertyBase().PropertyFlags;
-		VarProperty.MetaData.Append(MoveTemp(InnerPropDef.GetPropertyBase().MetaData)); // Move any added meta data items to the container
-		PropDef.SetAllocatorType(VarProperty.AllocatorType);
+		FUnrealPropertyDefinitionInfo& ValuePropDef = PropDef.GetValuePropDef();
+		VarProperty.PropertyFlags = ValuePropDef.GetPropertyBase().PropertyFlags;
+		ValuePropDef.GetPropertyBase().PropertyFlags &= CPF_PropagateToArrayInner;
+	}
+
+	static FPropertyFlagsChanges PostParseFinalize(FUnrealPropertyDefinitionInfo& PropDef, EPropertyFlags& AddedFlags)
+	{
+		FPropertyFlagsChangeDetector Detector(PropDef);
+		FPropertyBase& VarProperty = PropDef.GetPropertyBase();
+		FUnrealPropertyDefinitionInfo& ValuePropDef = PropDef.GetValuePropDef();
+
+		EPropertyFlags ValueAddedFlags = CPF_None;
+		VarProperty.PropertyFlags |= PostParseFinalizeHelper<true>(ValuePropDef, ValueAddedFlags).Added;
 
-		// Propagate flags
-		InnerPropDef.GetPropertyBase().PropertyFlags &= CPF_PropagateToArrayInner;
+		VarProperty.MetaData.Append(MoveTemp(ValuePropDef.GetPropertyBase().MetaData)); // Move any added meta data items to the container
+		PropDef.SetAllocatorType(VarProperty.AllocatorType);
 
-		PropagateFlagsFromInnerAndHandlePersistentInstanceMetadata(VarProperty.PropertyFlags, PropDef.GetPropertyBase().MetaData, InnerPropDef);
+		PropagateFlagsFromInnerAndHandlePersistentInstanceMetadata(VarProperty.PropertyFlags, PropDef.GetPropertyBase().MetaData, ValuePropDef);
+		return Detector.Changes();
 	}
 
 	static FProperty* CreateEngineType(FUnrealPropertyDefinitionInfo& PropDef, FFieldVariant Scope, const FName& Name, EObjectFlags ObjectFlags)
@@ -2368,15 +2457,25 @@ struct FPropertyTypeTraitsSet : public FPropertyTypeTraitsBase
 		InnerVarProperty.MetaData.Empty();
 		PropDef.SetValuePropDef(FPropertyTraits::CreateProperty(InnerVarProperty, PropDef, PropDef.GetFName(),
 			PropDef.GetVariableCategory(), ACCESS_Public, PropDef.GetArrayDimensions(), PropDef.GetUnrealSourceFile(), PropDef.GetLineNumber(), PropDef.GetParsePosition()));
-		FUnrealPropertyDefinitionInfo& InnerPropDef = PropDef.GetValuePropDef();
 
-		VarProperty.PropertyFlags = InnerPropDef.GetPropertyBase().PropertyFlags;
-		VarProperty.MetaData.Append(MoveTemp(InnerPropDef.GetPropertyBase().MetaData)); // Move any added meta data items to the container
+		FUnrealPropertyDefinitionInfo& ValuePropDef = PropDef.GetValuePropDef();
+		VarProperty.PropertyFlags = ValuePropDef.GetPropertyBase().PropertyFlags;
+		ValuePropDef.GetPropertyBase().PropertyFlags &= CPF_PropagateToSetElement;
+	}
+
+	static FPropertyFlagsChanges PostParseFinalize(FUnrealPropertyDefinitionInfo& PropDef, EPropertyFlags& AddedFlags)
+	{
+		FPropertyFlagsChangeDetector Detector(PropDef);
+		FPropertyBase& VarProperty = PropDef.GetPropertyBase();
+		FUnrealPropertyDefinitionInfo& ValuePropDef = PropDef.GetValuePropDef();
 
-		// Propagate flags
-		InnerPropDef.GetPropertyBase().PropertyFlags &= CPF_PropagateToSetElement;
+		EPropertyFlags ValueAddedFlags = CPF_None;
+		VarProperty.PropertyFlags |= PostParseFinalizeHelper<true>(ValuePropDef, ValueAddedFlags).Added;
+
+		VarProperty.MetaData.Append(MoveTemp(ValuePropDef.GetPropertyBase().MetaData)); // Move any added meta data items to the container
 
-		PropagateFlagsFromInnerAndHandlePersistentInstanceMetadata(VarProperty.PropertyFlags, PropDef.GetPropertyBase().MetaData, InnerPropDef);
+		PropagateFlagsFromInnerAndHandlePersistentInstanceMetadata(VarProperty.PropertyFlags, PropDef.GetPropertyBase().MetaData, ValuePropDef);
+		return Detector.Changes();
 	}
 
 	static FProperty* CreateEngineType(FUnrealPropertyDefinitionInfo& PropDef, FFieldVariant Scope, const FName& Name, EObjectFlags ObjectFlags)
@@ -2436,7 +2535,6 @@ struct FPropertyTypeTraitsMap : public FPropertyTypeTraitsBase
 		PropDef.SetKeyPropDef(FPropertyTraits::CreateProperty(*PropDef.GetPropertyBase().MapKeyProp, PropDef, *(PropDef.GetFName().ToString() + TEXT("_Key")), 
 			PropDef.GetVariableCategory(), ACCESS_Public, PropDef.GetArrayDimensions(), PropDef.GetUnrealSourceFile(), PropDef.GetLineNumber(), PropDef.GetParsePosition()));
 		FUnrealPropertyDefinitionInfo& KeyPropDef = PropDef.GetKeyPropDef();
-		KeyPropDef.GetPropertyBase().MetaData.Empty();
 
 		FPropertyBase ValueVarProperty = PropDef.GetPropertyBase();
 		ValueVarProperty.ArrayType = EArrayType::None;
@@ -2448,15 +2546,32 @@ struct FPropertyTypeTraitsMap : public FPropertyTypeTraitsBase
 		FUnrealPropertyDefinitionInfo& ValuePropDef = PropDef.GetValuePropDef();
 
 		VarProperty.PropertyFlags = ValuePropDef.GetPropertyBase().PropertyFlags;
-		VarProperty.MetaData.Append(MoveTemp(ValuePropDef.GetPropertyBase().MetaData)); // Move any added meta data items to the container
-		PropDef.SetAllocatorType(VarProperty.AllocatorType);//@TODO - This is a duplicate now
-
-		// Propagate flags
 		KeyPropDef.GetPropertyBase().PropertyFlags &= CPF_PropagateToMapKey;
 		ValuePropDef.GetPropertyBase().PropertyFlags &= CPF_PropagateToMapValue;
 
+		PropDef.SetAllocatorType(VarProperty.AllocatorType);//@TODO - This is a duplicate now
+	}
+
+	static FPropertyFlagsChanges PostParseFinalize(FUnrealPropertyDefinitionInfo& PropDef, EPropertyFlags& AddedFlags)
+	{
+		FPropertyFlagsChangeDetector Detector(PropDef);
+		FPropertyBase& VarProperty = PropDef.GetPropertyBase();
+		FUnrealPropertyDefinitionInfo& KeyPropDef = PropDef.GetKeyPropDef();
+		FUnrealPropertyDefinitionInfo& ValuePropDef = PropDef.GetValuePropDef();
+
+		EPropertyFlags KeyAddedFlags = CPF_None;
+		PostParseFinalizeHelper<true>(KeyPropDef, KeyAddedFlags);
+		KeyPropDef.GetPropertyBase().MetaData.Empty();
+
+		EPropertyFlags ValueAddedFlags = CPF_None;
+		VarProperty.PropertyFlags |= PostParseFinalizeHelper<true>(ValuePropDef, ValueAddedFlags).Added;
+		KeyPropDef.GetPropertyBase().PropertyFlags |= ValueAddedFlags;
+
+		VarProperty.MetaData.Append(MoveTemp(ValuePropDef.GetPropertyBase().MetaData)); // Move any added meta data items to the container
+
 		PropagateFlagsFromInnerAndHandlePersistentInstanceMetadata(VarProperty.PropertyFlags, PropDef.GetPropertyBase().MapKeyProp->MetaData, KeyPropDef);
 		PropagateFlagsFromInnerAndHandlePersistentInstanceMetadata(VarProperty.PropertyFlags, PropDef.GetPropertyBase().MetaData, ValuePropDef);
+		return Detector.Changes();
 	}
 
 	static FProperty* CreateEngineType(FUnrealPropertyDefinitionInfo& PropDef, FFieldVariant Scope, const FName& Name, EObjectFlags ObjectFlags)
@@ -2691,3 +2806,13 @@ bool FPropertyTraits::IsValidFieldClass(FName FieldClassName)
 
 	return FFieldClass::GetNameToFieldClassMap().Find(FieldClassName) != nullptr;
 }
+
+void FPropertyTraits::PostParseFinalize(FUnrealPropertyDefinitionInfo& PropDef)
+{
+	EPropertyFlags AddedFlags = CPF_None;
+	if (PostParseFinalizeHelper<true>(PropDef, AddedFlags).Added != CPF_None)
+	{
+		PropDef.GetOuter()->OnPostParsePropertyFlagsChanged(PropDef);
+	}
+}
+
diff --git a/Engine/Source/Programs/UnrealHeaderTool/Private/PropertyTypes.h b/Engine/Source/Programs/UnrealHeaderTool/Private/PropertyTypes.h
--- a/Engine/Source/Programs/UnrealHeaderTool/Private/PropertyTypes.h
+++ b/Engine/Source/Programs/UnrealHeaderTool/Private/PropertyTypes.h
@@ -84,4 +84,11 @@ struct FPropertyTraits
 	 * Tests to see if the given field class name is valid
 	 */
 	static bool IsValidFieldClass(FName FieldClassName);
+
+	/**
+	 * Handle any final initialization after parsing has completed
+	 *
+	 * @param PropDef The property in question
+	 */
+	static void PostParseFinalize(FUnrealPropertyDefinitionInfo& PropDef);
 };
diff --git a/Engine/Source/Programs/UnrealHeaderTool/Private/UnrealHeaderTool.h b/Engine/Source/Programs/UnrealHeaderTool/Private/UnrealHeaderTool.h
--- a/Engine/Source/Programs/UnrealHeaderTool/Private/UnrealHeaderTool.h
+++ b/Engine/Source/Programs/UnrealHeaderTool/Private/UnrealHeaderTool.h
@@ -10,6 +10,9 @@
 // we might never need this tracking since it is always by pointer.
 #define UHT_ENABLE_PTR_PROPERTY_TAG 0
 
+// Enable concurrent parsing of header files
+#define UHT_ENABLE_CONCURRENT_PARSING 1
+
 struct FHeaderParserNames
 {
 	static const FName NAME_IsConversionRoot;
diff --git a/Engine/Source/Programs/UnrealHeaderTool/Private/UnrealTypeDefinitionInfo.cpp b/Engine/Source/Programs/UnrealHeaderTool/Private/UnrealTypeDefinitionInfo.cpp
--- a/Engine/Source/Programs/UnrealHeaderTool/Private/UnrealTypeDefinitionInfo.cpp
+++ b/Engine/Source/Programs/UnrealHeaderTool/Private/UnrealTypeDefinitionInfo.cpp
@@ -357,14 +357,15 @@ void FUnrealTypeDefinitionInfo::CreateUObjectEngineTypes(ECreateEngineTypesPhase
 	}
 }
 
-void FUnrealTypeDefinitionInfo::PostParseFinalize()
+void FUnrealTypeDefinitionInfo::PostParseFinalize(EPostParseFinalizePhase Phase)
 {
-	switch (PostParseFinalizeState)
+	EFinalizeState& FinalizeState = PostParseFinalizeState[uint8(Phase)];
+	switch (FinalizeState)
 	{
 	case EFinalizeState::None:
-		PostParseFinalizeState = EFinalizeState::InProgress;
-		PostParseFinalizeInternal();
-		PostParseFinalizeState = EFinalizeState::Finished;
+		FinalizeState = EFinalizeState::InProgress;
+		PostParseFinalizeInternal(Phase);
+		FinalizeState = EFinalizeState::Finished;
 		break;
 	case EFinalizeState::InProgress:
 		checkf(false, TEXT("Recursive call to PostParseFinalize detected"))
@@ -558,13 +559,77 @@ void FUnrealTypeDefinitionInfo::ValidateMetaDataFormat(const FName InKey, ECheck
 	}
 }
 
-void FUnrealPropertyDefinitionInfo::PostParseFinalizeInternal()
+void FUnrealPropertyDefinitionInfo::PostParseFinalizeInternal(EPostParseFinalizePhase Phase)
 {
-	if (GetPropertySafe() == nullptr)
+	switch (Phase)
 	{
-		FPropertyTraits::CreateEngineType(SharedThis(this));
+	case EPostParseFinalizePhase::PreCreateEngineTypes:
+		FPropertyTraits::PostParseFinalize(*this);
+		break;
+
+	case EPostParseFinalizePhase::PostCreateEngineTypes:
+
+		// Due to structures not being fully parsed at parse time, we have to do blueprint checks here
+		if (HasAnyPropertyFlags(CPF_BlueprintAssignable | CPF_BlueprintCallable) && GetPropertyBase().IsMulticastDelegateOrMulticastDelegateStaticArray())
+		{
+			if (FUnrealFunctionDefinitionInfo* SourceDelegateFunctionDef = GetPropertyBase().FunctionDef) // should be set...
+			{
+				FUnrealStructDefinitionInfo* StructDef = GetOwnerStruct();
+				for (TSharedRef<FUnrealPropertyDefinitionInfo> FuncParamDef : SourceDelegateFunctionDef->GetProperties())
+				{
+					if (!FPropertyTraits::IsSupportedByBlueprint(*FuncParamDef, false))
+					{
+						FString ExtendedCPPType;
+						FString CPPType = FuncParamDef->GetCPPType(&ExtendedCPPType);
+						LogError(TEXT("Type '%s%s' is not supported by blueprint. %s.%s"), *CPPType, *ExtendedCPPType, *SourceDelegateFunctionDef->GetName(), *FuncParamDef->GetName());
+					}
+
+					if (FuncParamDef->HasAllPropertyFlags(CPF_OutParm) && !FuncParamDef->HasAllPropertyFlags(CPF_ConstParm))
+					{
+						const bool bClassGeneratedFromBP = StructDef != nullptr && StructDef->IsDynamic();
+						const bool bAllowedArrayRefFromBP = bClassGeneratedFromBP && FuncParamDef->IsDynamicArray();
+						if (!bAllowedArrayRefFromBP)
+						{
+							LogError(TEXT("BlueprintAssignable delegates do not support non-const references at the moment. Function: %s Parameter: '%s'"), *SourceDelegateFunctionDef->GetName(), *FuncParamDef->GetName());
+						}
+					}
+				}
+			}
+		}
+		else if (HasAnyPropertyFlags(CPF_BlueprintVisible))
+		{
+			FUnrealStructDefinitionInfo* StructDef = GetOwnerStruct();
+			FUnrealClassDefinitionInfo* ClassDef = UHTCast<FUnrealClassDefinitionInfo>(StructDef);
+			FUnrealScriptStructDefinitionInfo* ScriptStructDef = UHTCast<FUnrealScriptStructDefinitionInfo>(StructDef);
+			if (ClassDef != nullptr || ScriptStructDef != nullptr)
+			{
+				if (ScriptStructDef != nullptr && !StructDef->GetBoolMetaDataHierarchical(FHeaderParserNames::NAME_BlueprintType))
+				{
+					LogError(TEXT("Cannot expose property to blueprints in a struct that is not a BlueprintType. %s.%s"), *StructDef->GetName(), *GetName());
+				}
+
+				if (IsStaticArray())
+				{
+					LogError(TEXT("Static array cannot be exposed to blueprint %s.%s"), *StructDef->GetName(), *GetName());
+				}
+
+				if (!FPropertyTraits::IsSupportedByBlueprint(*this, true))
+				{
+					FString ExtendedCPPType;
+					FString CPPType = GetCPPType(&ExtendedCPPType);
+					LogError(TEXT("Type '%s%s' is not supported by blueprint. %s.%s"), *CPPType, *ExtendedCPPType, *StructDef->GetName(), *GetName());
+				}
+			}
+		}
+
+		//@TODO: The creation of engine type should be part of the engine type creation system
+		if (GetPropertySafe() == nullptr)
+		{
+			FPropertyTraits::CreateEngineType(SharedThis(this));
+		}
+		TypePackageName = GetTypePackageNameHelper(*this);
+		break;
 	}
-	TypePackageName = GetTypePackageNameHelper(*this);
 }
 
 bool FUnrealPropertyDefinitionInfo::IsDynamic() const
@@ -596,6 +661,13 @@ void FUnrealPropertyDefinitionInfo::SetDelegateFunctionSignature(FUnrealFunction
 		DelegateProperty->SignatureFunction = DelegateFunctionDef.GetFunction();
 	}
 	PropertyBase.FunctionDef = &DelegateFunctionDef;
+
+	// This can be invoked multiple times.  Just add the first instance.
+	if (!bSignatureSet)
+	{
+		bSignatureSet = true;
+		DelegateFunctionDef.AddReferencingProperty(*this);
+	}
 }
 
 FString FUnrealPropertyDefinitionInfo::GetEngineClassName() const
@@ -1007,22 +1079,34 @@ void FUnrealPackageDefinitionInfo::CreateUObjectEngineTypesInternal(ECreateEngin
 	}
 }
 
-void FUnrealPackageDefinitionInfo::PostParseFinalizeInternal()
+void FUnrealPackageDefinitionInfo::PostParseFinalizeInternal(EPostParseFinalizePhase Phase)
 {
-	UPackage* Package = GetPackage();
+	FUnrealObjectDefinitionInfo::PostParseFinalizeInternal(Phase);
 
-	FString PackageName = Package->GetName();
-	PackageName.ReplaceInline(TEXT("/"), TEXT("_"), ESearchCase::CaseSensitive);
+	switch (Phase)
+	{
+	case EPostParseFinalizePhase::PreCreateEngineTypes:
+		break;
+
+	case EPostParseFinalizePhase::PostCreateEngineTypes:
+	{
+		UPackage* Package = GetPackage();
 
-	SingletonName.Appendf(TEXT("Z_Construct_UPackage_%s()"), *PackageName);
-	SingletonNameChopped = SingletonName.LeftChop(2);
-	ExternDecl.Appendf(TEXT("\tUPackage* %s;\r\n"), *SingletonName);
+		FString PackageName = Package->GetName();
+		PackageName.ReplaceInline(TEXT("/"), TEXT("_"), ESearchCase::CaseSensitive);
+
+		SingletonName.Appendf(TEXT("Z_Construct_UPackage_%s()"), *PackageName);
+		SingletonNameChopped = SingletonName.LeftChop(2);
+		ExternDecl.Appendf(TEXT("\tUPackage* %s;\r\n"), *SingletonName);
+		break;
+	}
+	}
 
 	for (TSharedRef<FUnrealSourceFile>& LocalSourceFile : GetAllSourceFiles())
 	{
 		for (TSharedRef<FUnrealTypeDefinitionInfo>& TypeDef : LocalSourceFile->GetDefinedTypes())
 		{
-			TypeDef->PostParseFinalize();
+			TypeDef->PostParseFinalize(Phase);
 		}
 	}
 }
@@ -1035,24 +1119,36 @@ void FUnrealPackageDefinitionInfo::AddCrossModuleReference(TSet<FString>* Unique
 	}
 }
 
-void FUnrealFieldDefinitionInfo::PostParseFinalizeInternal()
+void FUnrealFieldDefinitionInfo::PostParseFinalizeInternal(EPostParseFinalizePhase Phase)
 {
-	const FString& ClassName = GetEngineClassName(true);
-	const FString& PackageShortName = GetPackageDef().GetShortUpperName();
+	FUnrealObjectDefinitionInfo::PostParseFinalizeInternal(Phase);
+
+	switch (Phase)
+	{
+	case EPostParseFinalizePhase::PreCreateEngineTypes:
+		break;
 
-	FUHTStringBuilder Out;
-	GenerateSingletonName(Out, this, false);
-	ExternDecl[0].Appendf(TEXT("\t%s_API U%s* %s;\r\n"), *PackageShortName, *ClassName, *Out);
-	SingletonName[0] = Out;
-	SingletonNameChopped[0] = SingletonName[0].LeftChop(2);
+	case EPostParseFinalizePhase::PostCreateEngineTypes:
+	{
+		const FString& ClassName = GetEngineClassName(true);
+		const FString& PackageShortName = GetPackageDef().GetShortUpperName();
+
+		FUHTStringBuilder Out;
+		GenerateSingletonName(Out, this, false);
+		ExternDecl[0].Appendf(TEXT("\t%s_API U%s* %s;\r\n"), *PackageShortName, *ClassName, *Out);
+		SingletonName[0] = Out;
+		SingletonNameChopped[0] = SingletonName[0].LeftChop(2);
 
-	Out.Reset();
-	GenerateSingletonName(Out, this, true);
-	ExternDecl[1].Appendf(TEXT("\t%s_API U%s* %s;\r\n"), *PackageShortName, *ClassName, *Out);
-	SingletonName[1] = Out;
-	SingletonNameChopped[1] = SingletonName[1].LeftChop(2);
+		Out.Reset();
+		GenerateSingletonName(Out, this, true);
+		ExternDecl[1].Appendf(TEXT("\t%s_API U%s* %s;\r\n"), *PackageShortName, *ClassName, *Out);
+		SingletonName[1] = Out;
+		SingletonNameChopped[1] = SingletonName[1].LeftChop(2);
 
-	TypePackageName = GetTypePackageNameHelper(*this);
+		TypePackageName = GetTypePackageNameHelper(*this);
+		break;
+	}
+	}
 }
 
 FText FUnrealFieldDefinitionInfo::GetToolTipText(bool bShortTooltip) const
@@ -1146,6 +1242,39 @@ FUnrealClassDefinitionInfo* FUnrealFieldDefinitionInfo::GetOwnerClass() const
 	return nullptr;
 }
 
+void FUnrealFieldDefinitionInfo::AddReferencingProperty(FUnrealPropertyDefinitionInfo& PropertyDef)
+{
+	check(PropertyDef.NextReferencingProperty == nullptr);
+	for (FUnrealPropertyDefinitionInfo* Tail = ReferencingProperties.load(std::memory_order_relaxed);;)
+	{
+		PropertyDef.NextReferencingProperty = Tail;
+		if (ReferencingProperties.compare_exchange_weak(Tail, &PropertyDef, std::memory_order_release, std::memory_order_relaxed))
+		{
+			return;
+		}
+	}
+}
+
+FUnrealPropertyDefinitionInfo* FUnrealFieldDefinitionInfo::GetFirstReferencingProperty()
+{
+	return ReferencingProperties.load(std::memory_order_acquire);
+}
+
+void FUnrealFieldDefinitionInfo::PostParseFinalizeReferencedProperties()
+{
+	for (FUnrealPropertyDefinitionInfo* PropDef = GetFirstReferencingProperty(); PropDef; PropDef = PropDef->NextReferencingProperty)
+	{
+		// A referencing property might be the child property of another property.  We invoke PostParseFinalize on the outermost 
+		// property.
+		FUnrealPropertyDefinitionInfo* OuterPropDef = PropDef;
+		while (FUnrealPropertyDefinitionInfo* NextPropDef = UHTCast<FUnrealPropertyDefinitionInfo>(OuterPropDef->GetOuter()))
+		{
+			OuterPropDef = NextPropDef;
+		}
+		FPropertyTraits::PostParseFinalize(*OuterPropDef);
+	}
+}
+
 FUnrealEnumDefinitionInfo::FUnrealEnumDefinitionInfo(FUnrealSourceFile& InSourceFile, int32 InLineNumber, FString&& InNameCPP, FName InName, UEnum::ECppForm InCppForm, EUnderlyingEnumType InUnderlyingType)
 	: FUnrealFieldDefinitionInfo(InSourceFile, InLineNumber, MoveTemp(InNameCPP), InName, InSourceFile.GetPackageDef())
 	, UnderlyingType(InUnderlyingType)
@@ -1495,42 +1624,53 @@ void FUnrealStructDefinitionInfo::CreateUObjectEngineTypesInternal(ECreateEngine
 	}
 }
 
-void FUnrealStructDefinitionInfo::PostParseFinalizeInternal()
+void FUnrealStructDefinitionInfo::PostParseFinalizeInternal(EPostParseFinalizePhase Phase)
 {
-	FUnrealFieldDefinitionInfo::PostParseFinalizeInternal();
+	FUnrealFieldDefinitionInfo::PostParseFinalizeInternal(Phase);
 
 	if (SuperStructInfo.Struct)
 	{
-		SuperStructInfo.Struct->PostParseFinalize();
+		SuperStructInfo.Struct->PostParseFinalize(Phase);
 	}
 
 	for (const FBaseStructInfo& Info : BaseStructInfos)
 	{
 		if (Info.Struct)
 		{
-			Info.Struct->PostParseFinalize();
+			Info.Struct->PostParseFinalize(Phase);
 		}
 	}
 
 	for (TSharedRef<FUnrealFunctionDefinitionInfo> FunctionDef : Functions)
 	{
-		FunctionDef->PostParseFinalize();
+		FunctionDef->PostParseFinalize(Phase);
 	}
 
 	for (TSharedRef<FUnrealPropertyDefinitionInfo> PropertyDef : Properties)
 	{
-		PropertyDef->PostParseFinalize();
+		PropertyDef->PostParseFinalize(Phase);
 	}
 
-	GetStruct()->Bind();
+	switch (Phase)
+	{
+	case EPostParseFinalizePhase::PreCreateEngineTypes:
+		break;
 
-	// Internals will assert of we are relinking an intrinsic 
-	bool bRelinkExistingProperties = true;
-	if (FUnrealClassDefinitionInfo* ClassDef = UHTCast<FUnrealClassDefinitionInfo>(this))
+	case EPostParseFinalizePhase::PostCreateEngineTypes:
 	{
-		bRelinkExistingProperties = !ClassDef->HasAnyClassFlags(CLASS_Intrinsic);
+
+		GetStruct()->Bind();
+
+		// Internals will assert of we are relinking an intrinsic 
+		bool bRelinkExistingProperties = true;
+		if (FUnrealClassDefinitionInfo* ClassDef = UHTCast<FUnrealClassDefinitionInfo>(this))
+		{
+			bRelinkExistingProperties = !ClassDef->HasAnyClassFlags(CLASS_Intrinsic);
+		}
+		GetStruct()->StaticLink(bRelinkExistingProperties);
+		break;
+	}
 	}
-	GetStruct()->StaticLink(bRelinkExistingProperties);
 }
 
 void FUnrealStructDefinitionInfo::CreatePropertyEngineTypes()
@@ -1581,6 +1721,11 @@ FUnrealScriptStructDefinitionInfo::FUnrealScriptStructDefinitionInfo(FUnrealSour
 	: FUnrealStructDefinitionInfo(InSourceFile, InLineNumber, MoveTemp(InNameCPP), InName, InSourceFile.GetPackageDef())
 { }
 
+FUnrealScriptStructDefinitionInfo* FUnrealScriptStructDefinitionInfo::GetSuperScriptStruct() const
+{
+	return UHTCast<FUnrealScriptStructDefinitionInfo>(GetSuperStruct());
+}
+
 uint32 FUnrealScriptStructDefinitionInfo::GetHash(bool bIncludeNoExport) const
 {
 	if (!bIncludeNoExport)
@@ -1625,9 +1770,56 @@ void FUnrealScriptStructDefinitionInfo::CreateUObjectEngineTypesInternal(ECreate
 	}
 }
 
-void FUnrealScriptStructDefinitionInfo::PostParseFinalizeInternal()
+void FUnrealScriptStructDefinitionInfo::PostParseFinalizeInternal(EPostParseFinalizePhase Phase)
+{
+	FUnrealStructDefinitionInfo::PostParseFinalizeInternal(Phase);
+
+	switch (Phase)
+	{
+	case EPostParseFinalizePhase::PreCreateEngineTypes:
+		// Collect the instanced reference state for the structure and propagate it to any
+		// property that references the structure.
+		if (FUnrealScriptStructDefinitionInfo* Base = GetSuperScriptStruct())
+		{
+			if (Base->HasAnyStructFlags(STRUCT_HasInstancedReference))
+			{
+				SetHasInstancedReference();
+			}
+		}
+		if (!HasAnyStructFlags(STRUCT_HasInstancedReference))
+		{
+			for (TSharedRef<FUnrealPropertyDefinitionInfo> PropertyDef : GetProperties())
+			{
+				if (PropertyDef->ContainsInstancedObjectProperty())
+				{
+					// Invoke this method to make sure we forward
+					OnPostParsePropertyFlagsChanged(*PropertyDef);
+					break;
+				}
+			}
+		}
+		break;
+
+	case EPostParseFinalizePhase::PostCreateEngineTypes:
+		break;
+	}
+}
+
+void FUnrealScriptStructDefinitionInfo::OnPostParsePropertyFlagsChanged(FUnrealPropertyDefinitionInfo& PropertyDef)
+{
+	if (!HasAnyStructFlags(STRUCT_HasInstancedReference))
+	{
+		if (PropertyDef.ContainsInstancedObjectProperty())
+		{
+			SetHasInstancedReference();
+		}
+	}
+}
+
+void FUnrealScriptStructDefinitionInfo::SetHasInstancedReference()
 {
-	FUnrealStructDefinitionInfo::PostParseFinalizeInternal();
+	SetStructFlags(STRUCT_HasInstancedReference);
+	PostParseFinalizeReferencedProperties();
 }
 
 FUnrealClassDefinitionInfo::FUnrealClassDefinitionInfo(FUnrealSourceFile& InSourceFile, int32 InLineNumber, FString&& InNameCPP, FName InName, bool bInIsInterface)
@@ -1818,66 +2010,96 @@ void FUnrealClassDefinitionInfo::CreateUObjectEngineTypesInternal(ECreateEngineT
 	}
 }
 
-void FUnrealClassDefinitionInfo::PostParseFinalizeInternal()
+void FUnrealClassDefinitionInfo::PostParseFinalizeInternal(EPostParseFinalizePhase Phase)
 {
-	FUnrealStructDefinitionInfo::PostParseFinalizeInternal();
+	FUnrealStructDefinitionInfo::PostParseFinalizeInternal(Phase);
 
-	if (IsInterface() && GetStructMetaData().ParsedInterface == EParsedInterface::ParsedUInterface)
+	switch (Phase)
 	{
-		FString UName = GetNameCPP();
-		FString IName = TEXT("I") + UName.RightChop(1);
-		Throwf(TEXT("UInterface '%s' parsed without a corresponding '%s'"), *UName, *IName);
-	}
-
-	FHeaderParser::CheckSparseClassData(*this);
+	case EPostParseFinalizePhase::PreCreateEngineTypes:
+		break;
 
-	// Collect the class replication properties
-	if (FUnrealClassDefinitionInfo* SuperClassDef = GetSuperClass())
+	case EPostParseFinalizePhase::PostCreateEngineTypes:
 	{
-		ClassReps = SuperClassDef->GetClassReps();
-		FirstOwnedClassRep = ClassReps.Num();
-	}
-	for (TSharedRef<FUnrealPropertyDefinitionInfo> PropertyDef : GetProperties())
-	{
-		if (PropertyDef->HasAnyPropertyFlags(CPF_Net))
+
+		// Pick up any newly set flags from the super class (i.e. HasInstancedReference)
+		if (FUnrealClassDefinitionInfo* SuperClassDef = GetSuperClass())
 		{
-			ClassReps.Add(&*PropertyDef);
+			SetClassFlags(SuperClassDef->GetClassFlags() & CLASS_ScriptInherit);
 		}
-	}
 
-	// Initialize the class object
-	UClass* Class = GetClass();
+		// Check to see if we have any instances referenced
+		if (!HasAnyClassFlags(CLASS_HasInstancedReference))
+		{
+			for (TSharedRef<FUnrealPropertyDefinitionInfo> PropertyDef : GetProperties())
+			{
+				if (PropertyDef->ContainsInstancedObjectProperty())
+				{
+					SetClassFlags(CLASS_HasInstancedReference);
+					break;
+				}
+			}
+		}
 
-	// Clear the property size
-	Class->PropertiesSize = 0;
+		if (IsInterface() && GetStructMetaData().ParsedInterface == EParsedInterface::ParsedUInterface)
+		{
+			FString UName = GetNameCPP();
+			FString IName = TEXT("I") + UName.RightChop(1);
+			Throwf(TEXT("UInterface '%s' parsed without a corresponding '%s'"), *UName, *IName);
+		}
 
-	// Make the visible to the package
-	Class->ClearFlags(RF_Transient);
-	check(Class->HasAnyFlags(RF_Public));
-	check(Class->HasAnyFlags(RF_Standalone));
+		FHeaderParser::CheckSparseClassData(*this);
 
-	// Finalize all of the children introduced in this class
-	for (TSharedRef<FUnrealFunctionDefinitionInfo> FunctionDef : GetFunctions())
-	{
-		UFunction* Function = FunctionDef->GetFunction();
-		Class->AddFunctionToFunctionMap(Function, Function->GetFName());
-	}
+		// Collect the class replication properties
+		if (FUnrealClassDefinitionInfo* SuperClassDef = GetSuperClass())
+		{
+			ClassReps = SuperClassDef->GetClassReps();
+			FirstOwnedClassRep = ClassReps.Num();
+		}
+		for (TSharedRef<FUnrealPropertyDefinitionInfo> PropertyDef : GetProperties())
+		{
+			if (PropertyDef->HasAnyPropertyFlags(CPF_Net))
+			{
+				ClassReps.Add(&*PropertyDef);
+			}
+		}
 
-	if (!HasAnyClassFlags(CLASS_Native))
-	{
-		//Class->UnMark(EObjectMark(OBJECTMARK_TagImp | OBJECTMARK_TagExp));
+		// Initialize the class object
+		UClass* Class = GetClass();
+
+		// Clear the property size
+		Class->PropertiesSize = 0;
+
+		// Make the visible to the package
+		Class->ClearFlags(RF_Transient);
+		check(Class->HasAnyFlags(RF_Public));
+		check(Class->HasAnyFlags(RF_Standalone));
+
+		// Finalize all of the children introduced in this class
+		for (TSharedRef<FUnrealFunctionDefinitionInfo> FunctionDef : GetFunctions())
+		{
+			UFunction* Function = FunctionDef->GetFunction();
+			Class->AddFunctionToFunctionMap(Function, Function->GetFName());
+		}
+
+		if (!HasAnyClassFlags(CLASS_Native))
+		{
+			//Class->UnMark(EObjectMark(OBJECTMARK_TagImp | OBJECTMARK_TagExp));
+		}
+		else if (!HasAnyClassFlags(CLASS_NoExport | CLASS_Intrinsic))
+		{
+			GetPackageDef().SetWriteClassesH(true);
+			//Class->UnMark(OBJECTMARK_TagImp);
+			//Class->Mark(OBJECTMARK_TagExp);
+		}
+
+		// This needs to be done outside of parallel blocks because it will modify UClass memory.
+		// Later calls to SetUpUhtReplicationData inside parallel blocks should be fine, because
+		// they will see the memory has already been set up, and just return the parent pointer.
+		Class->SetUpUhtReplicationData();
+		break;
 	}
-	else if (!HasAnyClassFlags(CLASS_NoExport | CLASS_Intrinsic))
-	{
-		GetPackageDef().SetWriteClassesH(true);
-		//Class->UnMark(OBJECTMARK_TagImp);
-		//Class->Mark(OBJECTMARK_TagExp);
 	}
-
-	// This needs to be done outside of parallel blocks because it will modify UClass memory.
-	// Later calls to SetUpUhtReplicationData inside parallel blocks should be fine, because
-	// they will see the memory has already been set up, and just return the parent pointer.
-	Class->SetUpUhtReplicationData();
 }
 
 bool FUnrealClassDefinitionInfo::ImplementsInterface(const FUnrealClassDefinitionInfo& SomeInterface) const
@@ -2549,53 +2771,68 @@ void FUnrealFunctionDefinitionInfo::CreateUObjectEngineTypesInternal(ECreateEngi
 	}
 }
 
-void FUnrealFunctionDefinitionInfo::PostParseFinalizeInternal()
+void FUnrealFunctionDefinitionInfo::PostParseFinalizeInternal(EPostParseFinalizePhase Phase)
 {
 	// Invoke the base class finalization
-	FUnrealStructDefinitionInfo::PostParseFinalizeInternal();
+	FUnrealStructDefinitionInfo::PostParseFinalizeInternal(Phase);
 
-	UFunction* Function = GetFunction();
+	switch (Phase)
+	{
+	case EPostParseFinalizePhase::PreCreateEngineTypes:
+		break;
 
-	// The following code is only performed on functions in a class.  
-	if (UHTCast<FUnrealClassDefinitionInfo>(GetOuter()) != nullptr)
+	case EPostParseFinalizePhase::PostCreateEngineTypes:
 	{
 
-		// Fix up any structs that were used as a parameter in a delegate before being defined
-		if (HasAnyFunctionFlags(FUNC_Delegate))
+		// Due to structures that might not be fully parsed at parse time, do blueprint validation here
+		if (GetFunctionType() == EFunctionType::Function && HasAnyFunctionFlags(FUNC_BlueprintCallable | FUNC_BlueprintEvent))
 		{
 			for (TSharedRef<FUnrealPropertyDefinitionInfo> PropertyDef : GetProperties())
 			{
-				if (PropertyDef->IsStructOrStructStaticArray())
+				if (PropertyDef->IsStaticArray())
 				{
-					FUnrealScriptStructDefinitionInfo& StructDef = UHTCastChecked<FUnrealScriptStructDefinitionInfo>(PropertyDef->GetPropertyBase().ClassDef);
-					if (StructDef.HasAnyStructFlags(STRUCT_HasInstancedReference))
-					{
-						PropertyDef->SetPropertyFlags(CPF_ContainsInstancedReference);
-					}
+					Throwf(TEXT("Static array cannot be exposed to blueprint. Function: %s Parameter %s\n"), *GetName(), *PropertyDef->GetName());
+				}
+
+				if (!FPropertyTraits::IsSupportedByBlueprint(*PropertyDef, false))
+				{
+					FPropertyTraits::IsSupportedByBlueprint(*PropertyDef, false);
+					FString ExtendedCPPType;
+					FString CPPType = PropertyDef->GetCPPType(&ExtendedCPPType);
+					LogError(TEXT("Type '%s%s' is not supported by blueprint. %s.%s"), *CPPType, *ExtendedCPPType, *GetName(), *PropertyDef->GetName());
 				}
 			}
 		}
 
-		Function->StaticLink(true);
+		UFunction* Function = GetFunction();
 
-		// Compute the function parameter size, propagate some flags to the outer function, and save the return offset
-		// Must be done in a second phase, as StaticLink resets various fields again!
-		Function->ParmsSize = 0;
-		for (TSharedRef<FUnrealPropertyDefinitionInfo> PropertyDef : GetProperties())
+		// The following code is only performed on functions in a class.  
+		if (UHTCast<FUnrealClassDefinitionInfo>(GetOuter()) != nullptr)
 		{
-			if (PropertyDef->HasSpecificPropertyFlags(CPF_ReturnParm | CPF_OutParm, CPF_OutParm))
-			{
-				SetFunctionFlags(FUNC_HasOutParms);
-			}
 
-			if (PropertyDef->IsStructOrStructStaticArray())
+			Function->StaticLink(true);
+
+			// Compute the function parameter size, propagate some flags to the outer function, and save the return offset
+			// Must be done in a second phase, as StaticLink resets various fields again!
+			Function->ParmsSize = 0;
+			for (TSharedRef<FUnrealPropertyDefinitionInfo> PropertyDef : GetProperties())
 			{
-				FUnrealScriptStructDefinitionInfo& StructDef = UHTCastChecked<FUnrealScriptStructDefinitionInfo>(PropertyDef->GetPropertyBase().ClassDef);
-				if (StructDef.HasDefaults())
+				if (PropertyDef->HasSpecificPropertyFlags(CPF_ReturnParm | CPF_OutParm, CPF_OutParm))
 				{
-					SetFunctionFlags(FUNC_HasDefaults);
+					SetFunctionFlags(FUNC_HasOutParms);
+				}
+
+				if (PropertyDef->IsStructOrStructStaticArray())
+				{
+					FUnrealScriptStructDefinitionInfo& StructDef = UHTCastChecked<FUnrealScriptStructDefinitionInfo>(PropertyDef->GetPropertyBase().ClassDef);
+					if (StructDef.HasDefaults())
+					{
+						SetFunctionFlags(FUNC_HasDefaults);
+					}
 				}
 			}
 		}
+		break;
+	}
 	}
 }
diff --git a/Engine/Source/Programs/UnrealHeaderTool/Private/UnrealTypeDefinitionInfo.h b/Engine/Source/Programs/UnrealHeaderTool/Private/UnrealTypeDefinitionInfo.h
--- a/Engine/Source/Programs/UnrealHeaderTool/Private/UnrealTypeDefinitionInfo.h
+++ b/Engine/Source/Programs/UnrealHeaderTool/Private/UnrealTypeDefinitionInfo.h
@@ -71,6 +71,20 @@ enum class EFunctionType
 	SparseDelegate,
 };
 
+enum class ECreateEngineTypesPhase : uint8
+{
+	Phase1,				// Create the Class, ScriptStruct, and Enum types
+	Phase2,				// Create the functions and properties
+	MAX_Phases,
+};
+
+enum class EPostParseFinalizePhase : uint8
+{
+	PreCreateEngineTypes,
+	PostCreateEngineTypes,
+	MAX_Phases,
+};
+
 /**
  * UHT Specific implementation of meta data support
  */
@@ -219,14 +233,6 @@ class FUnrealTypeDefinitionInfo
 	: public TSharedFromThis<FUnrealTypeDefinitionInfo>
 	, public FUHTMessageProvider
 {
-public:
-	enum class ECreateEngineTypesPhase : uint8
-	{
-		Phase1,				// Create the Class, ScriptStruct, and Enum types
-		Phase2,				// Create the functions and properties
-		MAX_Phases,
-	};
-
 private:
 	enum class EFinalizeState : uint8
 	{
@@ -391,6 +397,11 @@ public:
 		return *Info;
 	}
 
+	/**
+	 * Handles the propagation of property flags
+	 */
+	virtual void OnPostParsePropertyFlagsChanged(FUnrealPropertyDefinitionInfo& PropertyDef) {}
+
 	/**
 	 * Create UObject engine types.
 	 */
@@ -399,7 +410,7 @@ public:
 	/**
 	 * Perform any post parsing finalization and validation
 	 */
-	void PostParseFinalize();
+	void PostParseFinalize(EPostParseFinalizePhase Phase);
 
 	/**
 	 * Return the compilation scope associated with this object
@@ -584,7 +595,7 @@ protected:
 	/**
 	 * Perform any post parsing finalization and validation
 	 */
-	virtual void PostParseFinalizeInternal() {}
+	virtual void PostParseFinalizeInternal(EPostParseFinalizePhase Phase) {}
 
 private:
 	FString NameCPP;
@@ -593,7 +604,7 @@ private:
 	int32 LineNumber = 0;
 	std::atomic<uint32> Hash = 0;
 	EFinalizeState CreateUObectEngineTypesState[uint8(ECreateEngineTypesPhase::MAX_Phases)] = { EFinalizeState::None, EFinalizeState::None };
-	EFinalizeState PostParseFinalizeState = EFinalizeState::None;
+	EFinalizeState PostParseFinalizeState[uint8(EPostParseFinalizePhase::MAX_Phases)] = { EFinalizeState::None, EFinalizeState::None };
 };
 
 /**
@@ -1208,6 +1219,9 @@ public:
 	FUnrealObjectDefinitionInfo* GetOwnerObject() const;
 	FUnrealStructDefinitionInfo* GetOwnerStruct() const;
 
+	/** Linked list of properties referencing a specific field */
+	FUnrealPropertyDefinitionInfo* NextReferencingProperty = nullptr;
+
 protected:
 #if UHT_ENABLE_ENGINE_TYPE_CHECKS
 	virtual void CheckFindMetaData(const FName& Key, const FString* ValuePtr) const
@@ -1223,7 +1237,7 @@ protected:
 	/**
 	 * Perform any post parsing finalization and validation
 	 */
-	virtual void PostParseFinalizeInternal() override;
+	virtual void PostParseFinalizeInternal(EPostParseFinalizePhase Phase) override;
 
 private:
 	FName Name;
@@ -1238,6 +1252,7 @@ private:
 	EAccessSpecifier AccessSpecifier = ACCESS_Public;
 	EAllocatorType AllocatorType = EAllocatorType::Default;
 	bool bUnsized = false;
+	bool bSignatureSet = false; // temporary flag to know if we have set the signature
 };
 
 /**
@@ -1618,7 +1633,7 @@ protected:
 	/**
 	 * Perform any post parsing finalization and validation
 	 */
-	virtual void PostParseFinalizeInternal() override;
+	virtual void PostParseFinalizeInternal(EPostParseFinalizePhase Phase) override;
 
 private:
 	const FManifestModule& Module;
@@ -1734,17 +1749,35 @@ public:
 	 */
 	FUnrealClassDefinitionInfo* GetOwnerClass() const;
 
+	/**
+	 * Add a referencing property
+	 */
+	void AddReferencingProperty(FUnrealPropertyDefinitionInfo& PropertyDef);
+
+	/**
+	 * Return the first property that references this field
+	 */
+	FUnrealPropertyDefinitionInfo* GetFirstReferencingProperty();
+
+	/**
+	 * Invoke the post parse finalize method on all referenced properties
+	 */
+	void PostParseFinalizeReferencedProperties();
+
 protected:
 	/**
 	 * Perform any post parsing finalization and validation
 	 */
-	virtual void PostParseFinalizeInternal() override;
+	virtual void PostParseFinalizeInternal(EPostParseFinalizePhase Phase) override;
 
 private:
 	FString SingletonName[2];
 	FString SingletonNameChopped[2];
 	FString ExternDecl[2];
 	FString TypePackageName;
+
+	/** Linked list of all properties that reference this field */
+	std::atomic<FUnrealPropertyDefinitionInfo*> ReferencingProperties;
 };
 
 /**
@@ -2254,7 +2287,7 @@ protected:
 	/**
 	 * Perform any post parsing finalization and validation
 	 */
-	virtual void PostParseFinalizeInternal() override;
+	virtual void PostParseFinalizeInternal(EPostParseFinalizePhase Phase) override;
 
 private:
 	TSharedPtr<FScope> StructScope;
@@ -2322,6 +2355,11 @@ public:
 		return static_cast<UScriptStruct*>(GetObjectSafe());
 	}
 
+	/**
+	 * Return the super script struct
+	 */
+	FUnrealScriptStructDefinitionInfo* GetSuperScriptStruct() const;
+
 	/**
 	 * Return the struct flags
 	 */
@@ -2418,6 +2456,11 @@ public:
 		MacroDeclaredLineNumber = InMacroDeclaredLineNumber;
 	}
 
+	/**
+	 * Handles the propagation of property flags
+	 */
+	virtual void OnPostParsePropertyFlagsChanged(FUnrealPropertyDefinitionInfo& PropertyDef) override;
+
 protected:
 	/**
 	 * Create UObject engine types
@@ -2427,9 +2470,12 @@ protected:
 	/**
 	 * Perform any post parsing finalization and validation
 	 */
-	virtual void PostParseFinalizeInternal() override;
+	virtual void PostParseFinalizeInternal(EPostParseFinalizePhase Phase) override;
 
 private:
+
+	void SetHasInstancedReference();
+
 	int32 MacroDeclaredLineNumber = INDEX_NONE;
 	EStructFlags StructFlags = STRUCT_NoFlags;
 };
@@ -2635,7 +2681,7 @@ protected:
 	/**
 	 * Perform any post parsing finalization and validation
 	 */
-	virtual void PostParseFinalizeInternal() override;
+	virtual void PostParseFinalizeInternal(EPostParseFinalizePhase Phase) override;
 
 private:
 
@@ -3082,7 +3128,7 @@ protected:
 	/**
 	 * Perform any post parsing finalization and validation
 	 */
-	virtual void PostParseFinalizeInternal() override;
+	virtual void PostParseFinalizeInternal(EPostParseFinalizePhase Phase) override;
 
 private:
 	/** Merges all 'show' categories */
