Fixed global delegate functions to support changes when doing live coding.
Fixed issue where changing a function signature in an interface wouldn't update nodes in a blueprint.
Changed the API of the deferred registry to be more inline with UE coding standards.

#rnx

diff --git a/Engine/Source/Developer/Windows/LiveCoding/Private/LiveCodingModule.cpp b/Engine/Source/Developer/Windows/LiveCoding/Private/LiveCodingModule.cpp
--- a/Engine/Source/Developer/Windows/LiveCoding/Private/LiveCodingModule.cpp
+++ b/Engine/Source/Developer/Windows/LiveCoding/Private/LiveCodingModule.cpp
@@ -87,6 +87,10 @@ public:
 	{
 	}
 
+	virtual void NotifyChange(UPackage* New, UPackage* Old) override
+	{
+	}
+
 	virtual bool GetEnableReinstancing(bool bHasChanged) const
 	{
 		if (bHasChanged && !bEnabledMessage)
diff --git a/Engine/Source/Developer/Windows/LiveCodingServer/Private/External/LC_SymbolPatterns.cpp b/Engine/Source/Developer/Windows/LiveCodingServer/Private/External/LC_SymbolPatterns.cpp
--- a/Engine/Source/Developer/Windows/LiveCodingServer/Private/External/LC_SymbolPatterns.cpp
+++ b/Engine/Source/Developer/Windows/LiveCodingServer/Private/External/LC_SymbolPatterns.cpp
@@ -222,21 +222,23 @@ namespace symbolPatterns
 	};
 
 	// BEGIN EPIC MOD
-	extern const char* const UE_STATICS_BLOCK_PATTERNS[5] =
+	extern const char* const UE_STATICS_BLOCK_PATTERNS[6] =
 	{
 		"*Z_Construct_UClass_*_Statics*",
 		"*Z_Construct_UEnum_*_Statics*",
 		"*Z_Construct_UFunction_*_Statics*",
 		"*Z_Construct_UDelegateFunction_*_Statics*",
 		"*Z_Construct_UScriptStruct_*_Statics*",
+		"*Z_Construct_UPackage_*_Statics*",
 	};
 
-	extern const char* const UE_REGISTER_PATTERNS[4] =
+	extern const char* const UE_REGISTER_PATTERNS[5] =
 	{
 		"?AutoInitialize",
 		"?Z_CompiledInDeferEnum_UEnum_",
 		"?Z_CompiledInDeferStruct_UScriptStruct_",
 		"?Z_CompiledInDeferCppStructOps_UScriptStruct_",
+		"?Z_CompiledInDeferPackage_UPackage_",
 	};
 	// END EPIC MOD
 }
diff --git a/Engine/Source/Developer/Windows/LiveCodingServer/Private/External/LC_SymbolPatterns.h b/Engine/Source/Developer/Windows/LiveCodingServer/Private/External/LC_SymbolPatterns.h
--- a/Engine/Source/Developer/Windows/LiveCodingServer/Private/External/LC_SymbolPatterns.h
+++ b/Engine/Source/Developer/Windows/LiveCodingServer/Private/External/LC_SymbolPatterns.h
@@ -43,7 +43,7 @@ namespace symbolPatterns
 	extern const char* const ANONYMOUS_NAMESPACE_PATTERN;
 
 	// BEGIN EPIC MOD
-	extern const char* const UE_STATICS_BLOCK_PATTERNS[5];
-	extern const char* const UE_REGISTER_PATTERNS[4];
+	extern const char* const UE_STATICS_BLOCK_PATTERNS[6];
+	extern const char* const UE_REGISTER_PATTERNS[5];
 	// END EPIC MOD
 }
diff --git a/Engine/Source/Editor/BlueprintGraph/Classes/K2Node_CustomEvent.h b/Engine/Source/Editor/BlueprintGraph/Classes/K2Node_CustomEvent.h
--- a/Engine/Source/Editor/BlueprintGraph/Classes/K2Node_CustomEvent.h
+++ b/Engine/Source/Editor/BlueprintGraph/Classes/K2Node_CustomEvent.h
@@ -59,6 +59,7 @@ class UK2Node_CustomEvent : public UK2Node_Event
 	virtual void GetMenuActions(FBlueprintActionDatabaseRegistrar& ActionRegistrar) const override;
 	virtual void FixupPinStringDataReferences(FArchive* SavingArchive) override;
 	virtual bool NodeCausesStructuralBlueprintChange() const override { return true; }
+	virtual bool HasExternalDependencies(TArray<class UStruct*>* OptionalOutput) const override;
 	//~ End UK2Node Interface
 
 	//~ Begin UK2Node_EditablePinBase Interface
diff --git a/Engine/Source/Editor/BlueprintGraph/Private/K2Node_CustomEvent.cpp b/Engine/Source/Editor/BlueprintGraph/Private/K2Node_CustomEvent.cpp
--- a/Engine/Source/Editor/BlueprintGraph/Private/K2Node_CustomEvent.cpp
+++ b/Engine/Source/Editor/BlueprintGraph/Private/K2Node_CustomEvent.cpp
@@ -575,4 +575,24 @@ FEdGraphNodeDeprecationResponse UK2Node_CustomEvent::GetDeprecationResponse(EEdG
 	return Response;
 }
 
+bool UK2Node_CustomEvent::HasExternalDependencies(TArray<class UStruct*>* OptionalOutput) const
+{
+	bool bResult = false;
+
+	// We use the dependencies of the linked node instead of the resulting function signature because a globally defined 
+	// delegate won't match a dependency check (it has no owner). 
+	const UEdGraphPin* DelegateOutPin = FindPin(DelegateOutputName);
+	const UEdGraphPin* LinkedPin = (DelegateOutPin && DelegateOutPin->LinkedTo.Num() && DelegateOutPin->LinkedTo[0]) ? FBlueprintEditorUtils::FindFirstCompilerRelevantLinkedPin(DelegateOutPin->LinkedTo[0]) : nullptr;
+	if (LinkedPin)
+	{
+		if (UK2Node* OtherNode = Cast<UK2Node>(LinkedPin->GetOwningNode()))
+		{
+			bResult = OtherNode->HasExternalDependencies(OptionalOutput);
+		}
+	}
+
+	bResult |= Super::HasExternalDependencies(OptionalOutput);
+	return bResult;
+}
+
 #undef LOCTEXT_NAMESPACE
diff --git a/Engine/Source/Editor/UnrealEd/Private/Kismet2/BlueprintEditorUtils.cpp b/Engine/Source/Editor/UnrealEd/Private/Kismet2/BlueprintEditorUtils.cpp
--- a/Engine/Source/Editor/UnrealEd/Private/Kismet2/BlueprintEditorUtils.cpp
+++ b/Engine/Source/Editor/UnrealEd/Private/Kismet2/BlueprintEditorUtils.cpp
@@ -650,6 +650,13 @@ void FBlueprintEditorUtils::RefreshExternalBlueprintDependencyNodes(UBlueprint*
 						{
 							bShouldRefresh |= OwnerClass &&
 								(OwnerClass->IsChildOf(RefreshOnlyChild) || OwnerClass->GetAuthoritativeClass()->IsChildOf(RefreshOnlyChild));
+							if (!bShouldRefresh && Struct->IsA<UFunction>() && OwnerClass->HasAnyClassFlags(CLASS_Interface))
+							{
+								if (UClass* RefreshClass = Cast<UClass>(RefreshOnlyChild))
+								{
+									bShouldRefresh = RefreshClass->ImplementsInterface(OwnerClass);
+								}
+							}
 						}						
 					}
 					if (bShouldRefresh)
diff --git a/Engine/Source/Editor/UnrealEd/Private/Kismet2/ReloadUtilities.cpp b/Engine/Source/Editor/UnrealEd/Private/Kismet2/ReloadUtilities.cpp
--- a/Engine/Source/Editor/UnrealEd/Private/Kismet2/ReloadUtilities.cpp
+++ b/Engine/Source/Editor/UnrealEd/Private/Kismet2/ReloadUtilities.cpp
@@ -713,8 +713,9 @@ FReload::~FReload()
 	EndReload();
 
 	TStringBuilder<256> Builder;
-	if (ClassStats.HasValues() || StructStats.HasValues() || EnumStats.HasValues() || NumFunctionsRemapped != 0 || NumScriptStructsRemapped != 0)
+	if (PackageStats.HasValues() || ClassStats.HasValues() || StructStats.HasValues() || EnumStats.HasValues() || NumFunctionsRemapped != 0 || NumScriptStructsRemapped != 0)
 	{
+		FormatStats(Builder, TEXT("package"), TEXT("packages"), PackageStats);
 		FormatStats(Builder, TEXT("class"), TEXT("classes"), ClassStats);
 		FormatStats(Builder, TEXT("enum"), TEXT("enums"), EnumStats);
 		FormatStats(Builder, TEXT("scriptstruct"), TEXT("scriptstructs"), StructStats);
@@ -843,6 +844,20 @@ void FReload::NotifyChange(UScriptStruct* New, UScriptStruct* Old)
 	}
 }
 
+void FReload::NotifyChange(UPackage* New, UPackage* Old)
+{
+	if (Old != nullptr)
+	{
+		++PackageStats.Changed;
+	}
+	else
+	{
+		++PackageStats.New;
+	}
+
+	Packages.AddUnique(New);
+}
+
 namespace
 {
 	template<typename T>
diff --git a/Engine/Source/Editor/UnrealEd/Public/Kismet2/ReloadUtilities.h b/Engine/Source/Editor/UnrealEd/Public/Kismet2/ReloadUtilities.h
--- a/Engine/Source/Editor/UnrealEd/Public/Kismet2/ReloadUtilities.h
+++ b/Engine/Source/Editor/UnrealEd/Public/Kismet2/ReloadUtilities.h
@@ -40,6 +40,7 @@ public:
 	virtual void NotifyChange(UClass* New, UClass* Old) override;
 	virtual void NotifyChange(UEnum* New, UEnum* Old) override;
 	virtual void NotifyChange(UScriptStruct* New, UScriptStruct* Old) override;
+	virtual void NotifyChange(UPackage* New, UPackage* Old) override;
 	virtual void Reinstance() override;
 
 	/**
@@ -153,6 +154,7 @@ private:
 	FReinstanceStats ClassStats;
 	FReinstanceStats EnumStats;
 	FReinstanceStats StructStats;
+	FReinstanceStats PackageStats;
 	int32 NumFunctionsRemapped = 0;
 	int32 NumScriptStructsRemapped = 0;
 	mutable bool bEnabledMessage = false;
diff --git a/Engine/Source/Programs/UnrealHeaderTool/Private/CodeGenerator.cpp b/Engine/Source/Programs/UnrealHeaderTool/Private/CodeGenerator.cpp
--- a/Engine/Source/Programs/UnrealHeaderTool/Private/CodeGenerator.cpp
+++ b/Engine/Source/Programs/UnrealHeaderTool/Private/CodeGenerator.cpp
@@ -754,11 +754,11 @@ void FNativeClassHeaderGenerator::GetPropertyTag(FUHTStringBuilder& Out, FUnreal
 	{
 #if UHT_ENABLE_VALUE_PROPERTY_TAG
 	case EUHTPropertyType::Enum:
-		PropertyBase.EnumDef->GetHashTag(Out);
+		PropertyBase.EnumDef->GetHashTag(PropDef, Out);
 		break;
 
 	case EUHTPropertyType::Struct:
-		PropertyBase.ScriptStructDef->GetHashTag(Out);
+		PropertyBase.ScriptStructDef->GetHashTag(PropDef, Out);
 		break;
 #endif
 
@@ -774,6 +774,13 @@ void FNativeClassHeaderGenerator::GetPropertyTag(FUHTStringBuilder& Out, FUnreal
 		break;
 #endif
 
+#if UHT_ENABLE_DELEGATE_PROPERTY_TAG
+	case EUHTPropertyType::Delegate:
+	case EUHTPropertyType::MulticastDelegate:
+		PropertyBase.FunctionDef->GetHashTag(PropDef, Out);
+		break;
+#endif
+
 	default:
 		break;
 	}
@@ -1826,7 +1833,14 @@ bool IsDelegateFunction(FUnrealFieldDefinitionInfo& FieldDef)
 void FNativeClassHeaderGenerator::ExportGeneratedPackageInitCode(FOutputDevice& Out, const TCHAR* InDeclarations, uint32 Hash)
 {
 	UPackage* Package = PackageDef.GetPackage();
-	const FString& SingletonName = GetPackageSingletonName(PackageDef, nullptr);
+	const FString& SingletonName = GetPackageSingletonNameFuncAddr(PackageDef, nullptr);
+
+	FString PackageName = Package->GetName();
+	PackageName.ReplaceInline(TEXT("/"), TEXT("_"), ESearchCase::CaseSensitive);
+	FString BodyHashFn = FString::Printf(TEXT("Z_UPackage_%s_BodyHash"), *PackageName);
+	FString DeclarationsHashFn = FString::Printf(TEXT("Z_UPackage_%s_DeclarationsHash"), *PackageName);
+
+	uint32 DeclarationsHash = GenerateTextHash(InDeclarations);
 
 	TArray<FUnrealFieldDefinitionInfo*> Singletons;
 	for (TSharedRef<FUnrealSourceFile>& SourceFile : PackageDef.GetAllSourceFiles())
@@ -1855,9 +1869,18 @@ void FNativeClassHeaderGenerator::ExportGeneratedPackageInitCode(FOutputDevice&
 	FOutputDeviceNull OutputDeviceNull;
 	FString MetaDataParams = OutputMetaDataCodeForObject(OutputDeviceNull, Out, PackageDef, TEXT("Package_MetaDataParams"), TEXT(""), TEXT("\t\t\t"));
 
-	Out.Logf(TEXT("\tUPackage* %s\r\n"), *SingletonName);
+	Out.Logf(TEXT("\tuint32 %s()\t\r\n\t{\r\n\t\treturn 0x%08X;\r\n\t}\r\n"), *BodyHashFn, Hash);
+	Out.Logf(TEXT("\tuint32 %s()\t\r\n\t{\r\n\t\treturn 0x%08X;\r\n\t}\r\n"), *DeclarationsHashFn, DeclarationsHash);
+
+	Out.Logf(TEXT("\tstatic FPackageRegistrationInfo& Z_Registration_Info_UPackage_%s()\r\n"), *PackageName);
 	Out.Logf(TEXT("\t{\r\n"));
-	Out.Logf(TEXT("\t\tstatic UPackage* ReturnPackage = nullptr;\r\n"));
+	Out.Logf(TEXT("\t\tstatic FPackageRegistrationInfo info;\r\n"));
+	Out.Logf(TEXT("\t\treturn info;\r\n"));
+	Out.Logf(TEXT("\t}\r\n"));
+
+	Out.Logf(TEXT("\tUPackage* %s()\r\n"), *SingletonName);
+	Out.Logf(TEXT("\t{\r\n"));
+	Out.Logf(TEXT("\t\tstatic UPackage*& ReturnPackage = Z_Registration_Info_UPackage_%s().OuterSingleton;\r\n"), *PackageName);
 	Out.Logf(TEXT("\t\tif (!ReturnPackage)\r\n"));
 	Out.Logf(TEXT("\t\t{\r\n"));
 
@@ -1865,7 +1888,7 @@ void FNativeClassHeaderGenerator::ExportGeneratedPackageInitCode(FOutputDevice&
 	const TCHAR* SingletonCount;
 	if (!Singletons.IsEmpty())
 	{
-		Out.Logf(TEXT("\t\t\tstatic UObject* (*const SingletonFuncArray[])() = {\r\n"));
+		Out.Logf(TEXT("\t\t\tUObject* (*const SingletonFuncArray[])() = {\r\n"));
 		for (FUnrealFieldDefinitionInfo* FieldDef : Singletons)
 		{
 			const FString& Name = FieldDef->GetSingletonNameChopped(true);
@@ -1882,19 +1905,29 @@ void FNativeClassHeaderGenerator::ExportGeneratedPackageInitCode(FOutputDevice&
 		SingletonCount = TEXT("0");
 	}
 
-	Out.Logf(TEXT("\t\t\tstatic const UECodeGen_Private::FPackageParams PackageParams = {\r\n"));
+	Out.Logf(TEXT("\t\t\tconst UECodeGen_Private::FPackageParams PackageParams = {\r\n"));
 	Out.Logf(TEXT("\t\t\t\t%s,\r\n"), *CreateUTF8LiteralString(Package->GetName()));
 	Out.Logf(TEXT("\t\t\t\t%s,\r\n"), SingletonArray);
 	Out.Logf(TEXT("\t\t\t\t%s,\r\n"), SingletonCount);
 	Out.Logf(TEXT("\t\t\t\tPKG_CompiledIn | 0x%08X,\r\n"), Package->GetPackageFlags() & (PKG_ClientOptional | PKG_ServerSideOnly | PKG_EditorOnly | PKG_Developer | PKG_UncookedOnly));
 	Out.Logf(TEXT("\t\t\t\t0x%08X,\r\n"), Hash);
-	Out.Logf(TEXT("\t\t\t\t0x%08X,\r\n"), GenerateTextHash(InDeclarations));
+	Out.Logf(TEXT("\t\t\t\t0x%08X,\r\n"), DeclarationsHash);
 	Out.Logf(TEXT("\t\t\t\t%s\r\n"), *MetaDataParams);
 	Out.Logf(TEXT("\t\t\t};\r\n"));
 	Out.Logf(TEXT("\t\t\tUECodeGen_Private::ConstructUPackage(ReturnPackage, PackageParams);\r\n"));
 	Out.Logf(TEXT("\t\t}\r\n"));
 	Out.Logf(TEXT("\t\treturn ReturnPackage;\r\n"));
 	Out.Logf(TEXT("\t}\r\n"));
+
+	// Do not change the Z_CompiledInDeferPackage_UPackage_ without changing LC_SymbolPatterns
+	Out.Logf(TEXT("\tstatic FRegisterCompiledInInfo Z_CompiledInDeferPackage_UPackage_%s(%s, TEXT(\"%s\"), Z_Registration_Info_UPackage_%s(), CONSTUCT_RELOAD_VERSION_INFO(FPackageReloadVersionInfo, %s(), %s()));\r\n"),
+		*PackageName,
+		*SingletonName,
+		*Package->GetName(),
+		*PackageName,
+		*BodyHashFn,
+		*DeclarationsHashFn
+	);
 }
 
 void FNativeClassHeaderGenerator::ExportNativeGeneratedInitCode(FOutputDevice& Out, FOutputDevice& OutDeclarations, FReferenceGatherers& OutReferenceGatherers, const FUnrealSourceFile& SourceFile, FUnrealClassDefinitionInfo& ClassDef, FUHTStringBuilder& OutFriendText) const
@@ -2025,7 +2058,7 @@ void FNativeClassHeaderGenerator::ExportNativeGeneratedInitCode(FOutputDevice& O
 				}
 
 				FUHTStringBuilder FuncHashTag;
-				FunctionDef->GetHashTag(FuncHashTag);
+				FunctionDef->GetHashTag(ClassDef, FuncHashTag);
 
 				StaticDefinitions.Logf(
 					TEXT("%s\t\t{ &%s, %s },%s\r\n%s"),
@@ -2089,7 +2122,7 @@ void FNativeClassHeaderGenerator::ExportNativeGeneratedInitCode(FOutputDevice& O
 				}
 
 				FUHTStringBuilder IntHash;
-				InterClassDef.GetHashTag(IntHash);
+				InterClassDef.GetHashTag(ClassDef, IntHash);
 				StaticDefinitions.Logf(
 					TEXT("\t\t\t{ %s, %s, %s }, %s\r\n"),
 					*InterClassDef.GetSingletonNameChopped(false),
@@ -2209,7 +2242,7 @@ void FNativeClassHeaderGenerator::ExportNativeGeneratedInitCode(FOutputDevice& O
 	FUnrealClassDefinitionInfo* SuperClassDef = ClassDef.GetSuperClass();
 	if (SuperClassDef && !SuperClassDef->HasAnyClassFlags(CLASS_Intrinsic))
 	{
-		BaseClassHash = SuperClassDef->GetHash();
+		BaseClassHash = SuperClassDef->GetHash(ClassDef);
 	}
 	GeneratedClassRegisterFunctionText.Logf(TEXT("\r\n// %u\r\n"), BaseClassHash);
 
@@ -5321,35 +5354,6 @@ static bool HasDynamicOuter(FUnrealFieldDefinitionInfo& FieldDef)
 	return FieldOuter && FieldOuter->IsDynamic();
 }
 
-static void RecordPackageSingletons(
-	FUnrealSourceFile& SourceFile,
-	const TArray<FUnrealScriptStructDefinitionInfo*>& StructDefs,
-	const TArray<FUnrealFunctionDefinitionInfo*>& DelegateDefs)
-{
-	TArray<FUnrealFieldDefinitionInfo*> Singletons;
-	Singletons.Reserve(StructDefs.Num() + DelegateDefs.Num());
-	for (FUnrealScriptStructDefinitionInfo* StructDef : StructDefs)
-	{
-		if (StructDef->HasAnyStructFlags(STRUCT_NoExport) && !HasDynamicOuter(*StructDef))
-		{
-			Singletons.Add(StructDef);
-		}
-	}
-
-	for (FUnrealFunctionDefinitionInfo* DelegateDef : DelegateDefs)
-	{
-		if (!HasDynamicOuter(*DelegateDef))
-		{
-			Singletons.Add(DelegateDef);
-		}
-	}
-
-	if (Singletons.Num())
-	{
-		SourceFile.GetSingletons().Append(MoveTemp(Singletons));
-	}
-}
-
 bool FNativeClassHeaderGenerator::WriteSource(const FManifestModule& Module, FGeneratedFileInfo& FileInfo, const FString& InBodyText, FUnrealSourceFile* InSourceFile, const TSet<FString>& InCrossModuleReferences)
 {
 	// Collect the includes if this is from a source file
@@ -5478,23 +5482,25 @@ void FNativeClassHeaderGenerator::GenerateSourceFiles(
 
 					FReferenceGatherers ReferenceGatherers(&GeneratedCPP.CrossModuleReferences, GeneratedCPP.ForwardDeclarations);
 
-					TArray<FUnrealEnumDefinitionInfo*> EnumDefs;
-					TArray<FUnrealScriptStructDefinitionInfo*> StructDefs;
-					TArray<FUnrealFunctionDefinitionInfo*> DelegateFunctionDefs;
-					SourceFile.GetScope()->SplitTypesIntoArrays(EnumDefs, StructDefs, DelegateFunctionDefs);
-
-					RecordPackageSingletons(SourceFile, StructDefs, DelegateFunctionDefs);
+					TArray<FUnrealFieldDefinitionInfo*> Types;
+					SourceFile.GetScope()->GatherTypes(Types);
 
-					// Sort by line number to help make sure things are in order
-					auto Predicate = [](const FUnrealFieldDefinitionInfo& Lhs, const FUnrealFieldDefinitionInfo& Rhs)
+					// Make sure that all the types have a definition range.
+					for (FUnrealFieldDefinitionInfo* TypeDef : Types)
 					{
-						return Lhs.GetLineNumber() < Rhs.GetLineNumber();
-					};
-					EnumDefs.StableSort(Predicate);
-					StructDefs.StableSort(Predicate);
+						TypeDef->ValidateDefinitionRange();
+					}
 
-					// Reverse the containers as they come out in the reverse order of declaration
-					Algo::Reverse(DelegateFunctionDefs);
+					// Sort by the end of the definition range.  Since classes can contain delegates, we can't use the definition line
+					// since that will place the class ahead of the delegate.
+					Types.StableSort([](const FUnrealFieldDefinitionInfo& Lhs, const FUnrealFieldDefinitionInfo& Rhs)
+						{
+							return Lhs.GetDefinitionRange().End < Rhs.GetDefinitionRange().End;
+						}
+					);
+
+					TArray<FUnrealFieldDefinitionInfo*> Singletons;
+					Singletons.Reserve(Types.Num());
 
 					const FString FileDefineName = SourceFile.GetFileDefineName();
 					const FString& StrippedFilename = SourceFile.GetStrippedFilename();
@@ -5507,53 +5513,59 @@ void FNativeClassHeaderGenerator::GenerateSourceFiles(
 						LINE_TERMINATOR,
 						*FileDefineName, *StrippedFilename, *StrippedFilename, *FileDefineName);
 
-					// Export enums declared in non-UClass headers.
-					for (FUnrealEnumDefinitionInfo* EnumDef : EnumDefs)
+					for (FUnrealFieldDefinitionInfo* FieldDef : Types)
 					{
-						// Is this ever not the case?
-						if (EnumDef->GetOuter()->GetObject()->IsA(UPackage::StaticClass()))
+						if (FUnrealEnumDefinitionInfo* EnumDef = UHTCast<FUnrealEnumDefinitionInfo>(FieldDef))
 						{
-							GeneratedFunctionDeclarations.Log(EnumDef->GetExternDecl(true));
-							Generator.ExportGeneratedEnumInitCode(GeneratedCPPText, ReferenceGatherers, SourceFile, *EnumDef);
+							// Is this ever not the case?
+							if (EnumDef->GetOuter()->GetObject()->IsA(UPackage::StaticClass()))
+							{
+								GeneratedFunctionDeclarations.Log(EnumDef->GetExternDecl(true));
+								Generator.ExportGeneratedEnumInitCode(GeneratedCPPText, ReferenceGatherers, SourceFile, *EnumDef);
+							}
+						}
+						else if (FUnrealScriptStructDefinitionInfo* ScriptStructDef = UHTCast<FUnrealScriptStructDefinitionInfo>(FieldDef))
+						{
+							if (ScriptStructDef->HasAnyStructFlags(STRUCT_NoExport) && !HasDynamicOuter(*ScriptStructDef))
+							{
+								Singletons.Add(ScriptStructDef);
+							}
+							GeneratedFunctionDeclarations.Log(ScriptStructDef->GetExternDecl(true));
+							Generator.ExportGeneratedStructBodyMacros(GeneratedHeaderText, GeneratedCPPText, ReferenceGatherers, SourceFile, *ScriptStructDef);
+						}
+						else if (FUnrealFunctionDefinitionInfo* FunctionDef = UHTCast<FUnrealFunctionDefinitionInfo>(FieldDef))
+						{
+							if (!HasDynamicOuter(*FunctionDef))
+							{
+								Singletons.Add(FunctionDef);
+							}
+							GeneratedFunctionDeclarations.Log(FunctionDef->GetExternDecl(true));
+							Generator.ExportDelegateDeclaration(GeneratedCPPText, ReferenceGatherers, SourceFile, *FunctionDef);
+							Generator.ExportDelegateDefinition(GeneratedHeaderText, ReferenceGatherers, SourceFile, *FunctionDef);
+						}
+						else if (FUnrealClassDefinitionInfo* ClassDef = UHTCast<FUnrealClassDefinitionInfo>(FieldDef))
+						{
+							if (!ClassDef->HasAnyClassFlags(CLASS_Intrinsic))
+							{
+								Generator.ExportClassFromSourceFileInner(GeneratedHeaderText, GeneratedCPPText, GeneratedFunctionDeclarations, ReferenceGatherers, *ClassDef, SourceFile, GeneratedCPP.ExportFlags);
+							}
 						}
 					}
 
-					// export boilerplate macros for structs
-					// reverse the order.
-					for (FUnrealScriptStructDefinitionInfo* StructDef : StructDefs)
-					{
-						GeneratedFunctionDeclarations.Log(StructDef->GetExternDecl(true));
-						Generator.ExportGeneratedStructBodyMacros(GeneratedHeaderText, GeneratedCPPText, ReferenceGatherers, SourceFile, *StructDef);
-					}
-
-					// export delegate definitions
-					for (FUnrealFunctionDefinitionInfo* FunctionDef : DelegateFunctionDefs)
-					{
-						GeneratedFunctionDeclarations.Log(FunctionDef->GetExternDecl(true));
-						Generator.ExportDelegateDeclaration(GeneratedCPPText, ReferenceGatherers, SourceFile, *FunctionDef);
-					}
-
-					// export delegate wrapper function implementations
-					for (FUnrealFunctionDefinitionInfo* FunctionDef : DelegateFunctionDefs)
-					{
-						Generator.ExportDelegateDefinition(GeneratedHeaderText, ReferenceGatherers, SourceFile, *FunctionDef);
-					}
+					GeneratedHeaderText.Log(TEXT("#undef CURRENT_FILE_ID\r\n"));
+					GeneratedHeaderText.Logf(TEXT("#define CURRENT_FILE_ID %s\r\n\r\n\r\n"), *SourceFile.GetFileId());
 
-					for (const TSharedRef<FUnrealTypeDefinitionInfo>& TypeDef : SourceFile.GetDefinedClasses())
+					for (FUnrealFieldDefinitionInfo* FieldDef : Types)
 					{
-						FUnrealClassDefinitionInfo& ClassDef = TypeDef->AsClassChecked();
-						if (!ClassDef.HasAnyClassFlags(CLASS_Intrinsic))
+						if (FUnrealEnumDefinitionInfo* EnumDef = UHTCast<FUnrealEnumDefinitionInfo>(FieldDef))
 						{
-							Generator.ExportClassFromSourceFileInner(GeneratedHeaderText, GeneratedCPPText, GeneratedFunctionDeclarations, ReferenceGatherers, ClassDef, SourceFile, GeneratedCPP.ExportFlags);
+							Generator.ExportEnum(GeneratedHeaderText, *EnumDef);
 						}
 					}
 
-					GeneratedHeaderText.Log(TEXT("#undef CURRENT_FILE_ID\r\n"));
-					GeneratedHeaderText.Logf(TEXT("#define CURRENT_FILE_ID %s\r\n\r\n\r\n"), *SourceFile.GetFileId());
-
-					for (FUnrealEnumDefinitionInfo* EnumDef : EnumDefs)
+					if (Singletons.Num())
 					{
-						Generator.ExportEnum(GeneratedHeaderText, *EnumDef);
+						SourceFile.GetSingletons().Append(MoveTemp(Singletons));
 					}
 				}
 			);
diff --git a/Engine/Source/Programs/UnrealHeaderTool/Private/HeaderParser.cpp b/Engine/Source/Programs/UnrealHeaderTool/Private/HeaderParser.cpp
--- a/Engine/Source/Programs/UnrealHeaderTool/Private/HeaderParser.cpp
+++ b/Engine/Source/Programs/UnrealHeaderTool/Private/HeaderParser.cpp
@@ -1052,6 +1052,8 @@ FUnrealEnumDefinitionInfo& FHeaderParser::CompileEnum()
 		Throwf(TEXT("Invalid BlueprintType enum base - currently only uint8 supported"));
 	}
 
+	EnumDef.GetDefinitionRange().Start = &Input[InputPos];
+
 	// Get opening brace.
 	RequireSymbol( TEXT('{'), TEXT("'Enum'") );
 
@@ -1221,6 +1223,8 @@ FUnrealEnumDefinitionInfo& FHeaderParser::CompileEnum()
 		RequireSymbol( TEXT('}'), TEXT("'Enum'") );
 	}
 
+	EnumDef.GetDefinitionRange().End = &Input[InputPos];
+
 	// Register the list of enum names.
 	if (!EnumDef.SetEnums(EnumNames, CppForm, Flags))
 	{
@@ -1851,6 +1855,8 @@ FUnrealScriptStructDefinitionInfo& FHeaderParser::CompileStructDeclaration()
 	// Register the metadata
 	FUHTMetaData::RemapAndAddMetaData(StructDef, MoveTemp(MetaData));
 
+	StructDef.GetDefinitionRange().Start = &Input[InputPos];
+
 	// Get opening brace.
 	RequireSymbol( TEXT('{'), TEXT("'struct'") );
 
@@ -2020,6 +2026,8 @@ FUnrealScriptStructDefinitionInfo& FHeaderParser::CompileStructDeclaration()
 		}
 	}
 
+	StructDef.GetDefinitionRange().End = &Input[InputPos];
+
 	// Validation
 	bool bStructBodyFound = StructDef.GetMacroDeclaredLineNumber() != INDEX_NONE;
 	bool bExported        = !!(StructFlags & STRUCT_Native);
@@ -5913,6 +5921,8 @@ FUnrealFunctionDefinitionInfo& FHeaderParser::CompileDelegateDeclaration(const T
 		FuncInfo.FunctionFlags |= FUNC_MulticastDelegate;
 	}
 
+	const TCHAR* StartPos = &Input[InputPos];
+
 	// Now parse the macro body
 	RequireSymbol(TEXT('('), CurrentScopeName);
 
@@ -5959,6 +5969,7 @@ FUnrealFunctionDefinitionInfo& FHeaderParser::CompileDelegateDeclaration(const T
 	}
 
 	FUnrealFunctionDefinitionInfo& DelegateSignatureFunctionDef = CreateFunction(MoveTemp(FuncInfo), bIsSparse ? EFunctionType::SparseDelegate : EFunctionType::Delegate);
+	DelegateSignatureFunctionDef.GetDefinitionRange().Start = StartPos;
 	const FFuncInfo& FuncDefFuncInfo = DelegateSignatureFunctionDef.GetFunctionData();
 
 	// determine whether this function should be 'const'
@@ -6011,6 +6022,8 @@ FUnrealFunctionDefinitionInfo& FHeaderParser::CompileDelegateDeclaration(const T
 		RequireSymbol(TEXT(')'), TEXT("Delegate Declaration"));
 	}
 
+	DelegateSignatureFunctionDef.GetDefinitionRange().End = &Input[InputPos];
+
 	// The macro line must be set here
 	DelegateSignatureFunctionDef.GetFunctionData().MacroLine = InputLine;
 
@@ -6313,6 +6326,8 @@ FUnrealFunctionDefinitionInfo& FHeaderParser::CompileFunctionDeclaration()
 		Throwf(TEXT("Missing %s name"), TypeOfFunction);
 	}
 
+	const TCHAR* StartPos = &Input[InputPos];
+
 	if ( !MatchSymbol(TEXT('(')) )
 	{
 		Throwf(TEXT("Bad %s definition"), TypeOfFunction);
@@ -6358,6 +6373,7 @@ FUnrealFunctionDefinitionInfo& FHeaderParser::CompileFunctionDeclaration()
 	}
 
 	FUnrealFunctionDefinitionInfo& FuncDef = CreateFunction(MoveTemp(FuncInfo), EFunctionType::Function);
+	FuncDef.GetDefinitionRange().Start = StartPos;
 	const FFuncInfo& FuncDefFuncInfo = FuncDef.GetFunctionData();
 	FUnrealFunctionDefinitionInfo* SuperFuncDef = FuncDef.GetSuperFunction();
 
@@ -6494,6 +6510,8 @@ FUnrealFunctionDefinitionInfo& FHeaderParser::CompileFunctionDeclaration()
 		Throwf(TEXT("Function '%s': Base implementation of RPCs cannot be in a state. Add a stub outside state scope."), *FuncDef.GetName());
 	}
 
+	FuncDef.GetDefinitionRange().End = &Input[InputPos];
+
 	// Just declaring a function, so end the nesting.
 	PostPopFunctionDeclaration(FuncDef);
 
diff --git a/Engine/Source/Programs/UnrealHeaderTool/Private/Scope.cpp b/Engine/Source/Programs/UnrealHeaderTool/Private/Scope.cpp
--- a/Engine/Source/Programs/UnrealHeaderTool/Private/Scope.cpp
+++ b/Engine/Source/Programs/UnrealHeaderTool/Private/Scope.cpp
@@ -21,52 +21,41 @@ void FScope::AddType(FUnrealFieldDefinitionInfo& Type)
 	TypeMap.Add(Type.GetFName(), &Type);
 }
 
-/**
- * Dispatch type to one of three arrays Enums, Structs and DelegateFunctions.
- *
- * @param Type Input type.
- * @param Enums (Output parameter) Array to fill with enums.
- * @param Structs (Output parameter) Array to fill with structs.
- * @param DelegateFunctions (Output parameter) Array to fill with delegate functions.
- */
-void DispatchType(FUnrealFieldDefinitionInfo& FieldDef, TArray<FUnrealEnumDefinitionInfo*>& Enums, TArray<FUnrealScriptStructDefinitionInfo*>& Structs, TArray<FUnrealFunctionDefinitionInfo*>& DelegateFunctions)
+void FScope::GatherTypes(TArray<FUnrealFieldDefinitionInfo*>& Types)
 {
-	if (UHTCast<FUnrealClassDefinitionInfo>(FieldDef) != nullptr)
-	{
-		// Inner scopes.
-		FieldDef.GetScope()->SplitTypesIntoArrays(Enums, Structs, DelegateFunctions);
-	}
-	else if (FUnrealEnumDefinitionInfo* EnumDef = UHTCast<FUnrealEnumDefinitionInfo>(FieldDef))
-	{
-		Enums.Add(EnumDef);
-	}
-	else if (FUnrealScriptStructDefinitionInfo* ScriptStructDef = UHTCast<FUnrealScriptStructDefinitionInfo>(FieldDef))
-	{
-		Structs.Add(ScriptStructDef);
-	}
-	else if (FUnrealFunctionDefinitionInfo* FunctionDef = UHTCast<FUnrealFunctionDefinitionInfo>(FieldDef))
+	for (TPair<FName, FUnrealFieldDefinitionInfo*>& TypePair : TypeMap)
 	{
-		if (FunctionDef->IsDelegateFunction())
+		FUnrealFieldDefinitionInfo* FieldDef = TypePair.Value;
+		if (FUnrealClassDefinitionInfo* ClassDef = UHTCast<FUnrealClassDefinitionInfo>(FieldDef))
 		{
-			bool bAdded = false;
-			if (FunctionDef->GetSuperFunction() == nullptr)
+			// Inner scopes.
+			ClassDef->GetScope()->GatherTypes(Types);
+			Types.Add(ClassDef);
+		}
+		else if (FUnrealEnumDefinitionInfo* EnumDef = UHTCast<FUnrealEnumDefinitionInfo>(FieldDef))
+		{
+			Types.Add(EnumDef);
+		}
+		else if (FUnrealScriptStructDefinitionInfo* ScriptStructDef = UHTCast<FUnrealScriptStructDefinitionInfo>(FieldDef))
+		{
+			Types.Add(ScriptStructDef);
+		}
+		else if (FUnrealFunctionDefinitionInfo* FunctionDef = UHTCast<FUnrealFunctionDefinitionInfo>(FieldDef))
+		{
+			if (FunctionDef->IsDelegateFunction())
 			{
-				DelegateFunctions.Add(FunctionDef);
-				bAdded = true;
+				bool bAdded = false;
+				if (FunctionDef->GetSuperFunction() == nullptr)
+				{
+					Types.Add(FunctionDef);
+					bAdded = true;
+				}
+				check(bAdded);
 			}
-			check(bAdded);
 		}
 	}
 }
 
-void FScope::SplitTypesIntoArrays(TArray<FUnrealEnumDefinitionInfo*>& Enums, TArray<FUnrealScriptStructDefinitionInfo*>& Structs, TArray<FUnrealFunctionDefinitionInfo*>& DelegateFunctions)
-{
-	for (TPair<FName, FUnrealFieldDefinitionInfo*>& TypePair : TypeMap)
-	{
-		DispatchType(*TypePair.Value, Enums, Structs, DelegateFunctions);
-	}
-}
-
 FUnrealFieldDefinitionInfo* FScope::FindTypeByName(FName Name)
 {
 	if (!Name.IsNone())
diff --git a/Engine/Source/Programs/UnrealHeaderTool/Private/Scope.h b/Engine/Source/Programs/UnrealHeaderTool/Private/Scope.h
--- a/Engine/Source/Programs/UnrealHeaderTool/Private/Scope.h
+++ b/Engine/Source/Programs/UnrealHeaderTool/Private/Scope.h
@@ -124,13 +124,11 @@ public:
 	}
 
 	/**
-	 * Gets structs, enums and delegate functions from this scope.
-	 *
-	 * @param Enums (Output parameter) List of enums from this scope.
-	 * @param Structs (Output parameter) List of structs from this scope.
-	 * @param DelegateFunctions (Output parameter) List of delegate properties from this scope.
+	 * Collect a list of all types declared in this scope.  This includes any types defined in other types which 
+	 * at this time is not currently supported.  Functions that aren't delegates are not collected.
+	 * The types are returned in the order they are declared.
 	 */
-	void SplitTypesIntoArrays(TArray<FUnrealEnumDefinitionInfo*>& Enums, TArray<FUnrealScriptStructDefinitionInfo*>& Structs, TArray<FUnrealFunctionDefinitionInfo*>& DelegateFunctions);
+	void GatherTypes(TArray<FUnrealFieldDefinitionInfo*>& Types);
 
 	/**
 	 * Gets scope name.
diff --git a/Engine/Source/Programs/UnrealHeaderTool/Private/UnrealHeaderTool.h b/Engine/Source/Programs/UnrealHeaderTool/Private/UnrealHeaderTool.h
--- a/Engine/Source/Programs/UnrealHeaderTool/Private/UnrealHeaderTool.h
+++ b/Engine/Source/Programs/UnrealHeaderTool/Private/UnrealHeaderTool.h
@@ -5,6 +5,9 @@
 // Enable property tags for enums and structures
 #define UHT_ENABLE_VALUE_PROPERTY_TAG 1
 
+// Enable property tags for delegates
+#define UHT_ENABLE_DELEGATE_PROPERTY_TAG 1
+
 // Enable property tags for class and interfaces wrapped in porters 
 // DO NOT ENABLE: This does not handle forward declarations.  Also,
 // we might never need this tracking since it is always by pointer.
diff --git a/Engine/Source/Programs/UnrealHeaderTool/Private/UnrealTypeDefinitionInfo.cpp b/Engine/Source/Programs/UnrealHeaderTool/Private/UnrealTypeDefinitionInfo.cpp
--- a/Engine/Source/Programs/UnrealHeaderTool/Private/UnrealTypeDefinitionInfo.cpp
+++ b/Engine/Source/Programs/UnrealHeaderTool/Private/UnrealTypeDefinitionInfo.cpp
@@ -314,18 +314,18 @@ void FUnrealTypeDefinitionInfo::SetHash(uint32 InHash)
 	Hash = InHash;
 }
 
-uint32 FUnrealTypeDefinitionInfo::GetHash(bool bIncludeNoExport) const
+uint32 FUnrealTypeDefinitionInfo::GetHash(FUnrealTypeDefinitionInfo& ReferencingDef, bool bIncludeNoExport) const
 {
 	if (Hash == 0)
 	{
-		Throwf(TEXT("Attempt to fetch the generated hash for type \"%s\" before it has been generated.  Include dependencies, topological sort, or job graph is in error."), *GetNameCPP());
+		ReferencingDef.Throwf(TEXT("Attempt to fetch the generated hash for type \"%s\" before it has been generated.  Include dependencies, topological sort, or job graph is in error."), *GetNameCPP());
 	}
 	return Hash;
 }
 
-void FUnrealTypeDefinitionInfo::GetHashTag(FUHTStringBuilder& Out) const
+void FUnrealTypeDefinitionInfo::GetHashTag(FUnrealTypeDefinitionInfo& ReferencingDef, FUHTStringBuilder& Out) const
 {
-	uint32 TempHash = GetHash(false);
+	uint32 TempHash = GetHash(ReferencingDef, false);
 	if (TempHash != 0)
 	{
 		if (Out.IsEmpty())
@@ -662,6 +662,10 @@ void FUnrealPropertyDefinitionInfo::SetDelegateFunctionSignature(FUnrealFunction
 	}
 	PropertyBase.FunctionDef = &DelegateFunctionDef;
 
+#if UHT_ENABLE_DELEGATE_PROPERTY_TAG
+	GetUnrealSourceFile().AddTypeDefIncludeIfNeeded(&DelegateFunctionDef);
+#endif
+
 	// This can be invoked multiple times.  Just add the first instance.
 	if (!bSignatureSet)
 	{
@@ -1726,7 +1730,7 @@ FUnrealScriptStructDefinitionInfo* FUnrealScriptStructDefinitionInfo::GetSuperSc
 	return UHTCast<FUnrealScriptStructDefinitionInfo>(GetSuperStruct());
 }
 
-uint32 FUnrealScriptStructDefinitionInfo::GetHash(bool bIncludeNoExport) const
+uint32 FUnrealScriptStructDefinitionInfo::GetHash(FUnrealTypeDefinitionInfo& ReferencingDef, bool bIncludeNoExport) const
 {
 	if (!bIncludeNoExport)
 	{
@@ -1735,7 +1739,7 @@ uint32 FUnrealScriptStructDefinitionInfo::GetHash(bool bIncludeNoExport) const
 			return 0;
 		}
 	}
-	return FUnrealStructDefinitionInfo::GetHash(bIncludeNoExport);
+	return FUnrealStructDefinitionInfo::GetHash(ReferencingDef, bIncludeNoExport);
 }
 
 void FUnrealScriptStructDefinitionInfo::CreateUObjectEngineTypesInternal(ECreateEngineTypesPhase Phase)
@@ -1942,7 +1946,7 @@ FUnrealClassDefinitionInfo* FUnrealClassDefinitionInfo::FindScriptClass(const FS
 	return nullptr;
 }
 
-uint32 FUnrealClassDefinitionInfo::GetHash(bool bIncludeNoExport) const
+uint32 FUnrealClassDefinitionInfo::GetHash(FUnrealTypeDefinitionInfo& ReferencingDef, bool bIncludeNoExport) const
 {
 	if (!bIncludeNoExport)
 	{
@@ -1951,7 +1955,7 @@ uint32 FUnrealClassDefinitionInfo::GetHash(bool bIncludeNoExport) const
 			return 0;
 		}
 	}
-	return FUnrealStructDefinitionInfo::GetHash(bIncludeNoExport);
+	return FUnrealStructDefinitionInfo::GetHash(ReferencingDef, bIncludeNoExport);
 }
 
 void FUnrealClassDefinitionInfo::CreateUObjectEngineTypesInternal(ECreateEngineTypesPhase Phase)
diff --git a/Engine/Source/Programs/UnrealHeaderTool/Private/UnrealTypeDefinitionInfo.h b/Engine/Source/Programs/UnrealHeaderTool/Private/UnrealTypeDefinitionInfo.h
--- a/Engine/Source/Programs/UnrealHeaderTool/Private/UnrealTypeDefinitionInfo.h
+++ b/Engine/Source/Programs/UnrealHeaderTool/Private/UnrealTypeDefinitionInfo.h
@@ -504,12 +504,12 @@ public:
 	/**
 	 * Return the previously set hash. This method will assert if the hash has not been set.
 	 */
-	virtual uint32 GetHash(bool bIncludeNoExport = true) const;
+	virtual uint32 GetHash(FUnrealTypeDefinitionInfo& ReferencingDef, bool bIncludeNoExport = true) const;
 
 	/**
 	 * Return the hash as a code comment.
 	 */
-	void GetHashTag(FUHTStringBuilder& Out) const;
+	void GetHashTag(FUnrealTypeDefinitionInfo& ReferencingDef, FUHTStringBuilder& Out) const;
 
 	/**
 	 * Add meta data for the given definition
@@ -1653,6 +1653,13 @@ private:
 class FUnrealFieldDefinitionInfo 
 	: public FUnrealObjectDefinitionInfo
 {
+public:
+	struct FDefinitionRange
+	{
+		const TCHAR* Start = nullptr;
+		const TCHAR* End = nullptr;
+	};
+
 protected:
 	using FUnrealObjectDefinitionInfo::FUnrealObjectDefinitionInfo;
 
@@ -1764,6 +1771,33 @@ public:
 	 */
 	void PostParseFinalizeReferencedProperties();
 
+	/**
+	 * Return the definition range of the structure
+	 */
+	FDefinitionRange& GetDefinitionRange()
+	{
+		return DefinitionRange;
+	}
+
+	/**
+	 * Return the definition range of the structure
+	 */
+	const FDefinitionRange& GetDefinitionRange() const
+	{
+		return DefinitionRange;
+	}
+
+	/**
+	 * Verify we have a valid definition range
+	 */
+	void ValidateDefinitionRange()
+	{
+		if (DefinitionRange.End <= DefinitionRange.Start)
+		{
+			Throwf(TEXT("The definition range is invalid. Most probably caused by previous parsing error."));
+		}
+	}
+
 protected:
 	/**
 	 * Perform any post parsing finalization and validation
@@ -1775,6 +1809,7 @@ private:
 	FString SingletonNameChopped[2];
 	FString ExternDecl[2];
 	FString TypePackageName;
+	FDefinitionRange DefinitionRange;
 
 	/** Linked list of all properties that reference this field */
 	std::atomic<FUnrealPropertyDefinitionInfo*> ReferencingProperties = nullptr;
@@ -2004,13 +2039,6 @@ public:
 		FUnrealStructDefinitionInfo* Struct = nullptr;
 	};
 
-	struct FDefinitionRange
-	{
-		const TCHAR* Start = nullptr;
-		const TCHAR* End = nullptr;
-	};
-
-
 public:
 	virtual FUnrealStructDefinitionInfo* AsStruct() override
 	{
@@ -2245,25 +2273,6 @@ public:
 		bHasGeneratedBody = true;
 	}
 
-	/**
-	 * Verify we have a valid definition range
-	 */
-	void ValidateDefinitionRange()
-	{
-		if (DefinitionRange.End <= DefinitionRange.Start)
-		{
-			Throwf(TEXT("The class definition range is invalid. Most probably caused by previous parsing error."));
-		}
-	}
-
-	/**
-	 * Return the definition range of the structure
-	 */
-	FDefinitionRange& GetDefinitionRange()
-	{
-		return DefinitionRange;
-	}
-
 	/**
 	 * Get the RigVM information
 	 */
@@ -2302,8 +2311,6 @@ private:
 	FBaseStructInfo SuperStructInfo;
 	TArray<FBaseStructInfo> BaseStructInfos;
 
-	FDefinitionRange DefinitionRange;
-
 	FRigVMStructInfo RigVMInfo;
 
 	EGeneratedCodeVersion GeneratedCodeVersion = FUHTConfig::Get().DefaultGeneratedCodeVersion;
@@ -2328,7 +2335,7 @@ public:
 		return this;
 	}
 
-	virtual uint32 GetHash(bool bIncludeNoExport = true) const override;
+	virtual uint32 GetHash(FUnrealTypeDefinitionInfo& ReferencingDef, bool bIncludeNoExport = true) const override;
 
 	/**
 	 * Return the engine class name
@@ -2755,7 +2762,7 @@ public:
 	 */
 	virtual const TCHAR* GetPrefixCPP() const;
 
-	virtual uint32 GetHash(bool bIncludeNoExport = true) const override;
+	virtual uint32 GetHash(FUnrealTypeDefinitionInfo& ReferencingDef, bool bIncludeNoExport = true) const override;
 
 	/**
 	 * Return the Engine instance associated with the compiler instance
diff --git a/Engine/Source/Runtime/CoreUObject/Private/UObject/DeferredRegistry.cpp b/Engine/Source/Runtime/CoreUObject/Private/UObject/DeferredRegistry.cpp
--- a/Engine/Source/Runtime/CoreUObject/Private/UObject/DeferredRegistry.cpp
+++ b/Engine/Source/Runtime/CoreUObject/Private/UObject/DeferredRegistry.cpp
@@ -6,43 +6,43 @@
 
 #if WITH_RELOAD
 
-void ReloadRenameObject(UScriptStruct& Object, const TCHAR* RenamePrefix)
+void ReloadProcessObject(UScriptStruct* ScriptStruct, const TCHAR* RenamePrefix)
 {
 	// Make sure the old struct is not used by anything
-	Object.ClearFlags(RF_Standalone | RF_Public);
-	Object.RemoveFromRoot();
-	const FName OldRename = MakeUniqueObjectName(GetTransientPackage(), Object.GetClass(), *FString::Printf(TEXT("%s_%s"), RenamePrefix, *Object.GetName()));
-	Object.Rename(*OldRename.ToString(), GetTransientPackage());
+	ScriptStruct->ClearFlags(RF_Standalone | RF_Public);
+	ScriptStruct->RemoveFromRoot();
+	const FName OldRename = MakeUniqueObjectName(GetTransientPackage(), ScriptStruct->GetClass(), *FString::Printf(TEXT("%s_%s"), RenamePrefix, *ScriptStruct->GetName()));
+	ScriptStruct->Rename(*OldRename.ToString(), GetTransientPackage());
 }
 
-void ReloadRenameObject(UEnum& Object, const TCHAR* RenamePrefix)
+void ReloadProcessObject(UEnum* Enum, const TCHAR* RenamePrefix)
 {
 	// Make sure the old struct is not used by anything
-	Object.ClearFlags(RF_Standalone | RF_Public);
-	Object.RemoveFromRoot();
-	const FName OldRename = MakeUniqueObjectName(GetTransientPackage(), Object.GetClass(), *FString::Printf(TEXT("%s_%s"), RenamePrefix, *Object.GetName()));
-	Object.Rename(*OldRename.ToString(), GetTransientPackage());
+	Enum->ClearFlags(RF_Standalone | RF_Public);
+	Enum->RemoveFromRoot();
+	const FName OldRename = MakeUniqueObjectName(GetTransientPackage(), Enum->GetClass(), *FString::Printf(TEXT("%s_%s"), RenamePrefix, *Enum->GetName()));
+	Enum->Rename(*OldRename.ToString(), GetTransientPackage());
 
-	Object.RemoveNamesFromMasterList();
+	Enum->RemoveNamesFromMasterList();
 }
 
-void ReloadRenameObject(UClass& Class, const TCHAR* RenamePrefix)
+void ReloadProcessObject(UClass* Class, const TCHAR* RenamePrefix)
 {
-	FString NameWithoutPrefix = UObjectBase::RemoveClassPrefix(*Class.GetName());
+	FString NameWithoutPrefix = UObjectBase::RemoveClassPrefix(*Class->GetName());
 
 	// Rename the old class and move it to transient package
-	Class.RemoveFromRoot();
-	Class.ClearFlags(RF_Standalone | RF_Public);
-	Class.GetDefaultObject()->RemoveFromRoot();
-	Class.GetDefaultObject()->ClearFlags(RF_Standalone | RF_Public);
-	const FName OldClassRename = MakeUniqueObjectName(GetTransientPackage(), Class.GetClass(), *FString::Printf(TEXT("%s_%s"), RenamePrefix, *NameWithoutPrefix));
-	Class.Rename(*OldClassRename.ToString(), GetTransientPackage());
-	Class.SetFlags(RF_Transient);
-	Class.AddToRoot();
+	Class->RemoveFromRoot();
+	Class->ClearFlags(RF_Standalone | RF_Public);
+	Class->GetDefaultObject()->RemoveFromRoot();
+	Class->GetDefaultObject()->ClearFlags(RF_Standalone | RF_Public);
+	const FName OldClassRename = MakeUniqueObjectName(GetTransientPackage(), Class->GetClass(), *FString::Printf(TEXT("%s_%s"), RenamePrefix, *NameWithoutPrefix));
+	Class->Rename(*OldClassRename.ToString(), GetTransientPackage());
+	Class->SetFlags(RF_Transient);
+	Class->AddToRoot();
 
 	// Make sure enums de-register their names BEFORE we create the new class, otherwise there will be name conflicts
 	TArray<UObject*> ClassSubobjects;
-	GetObjectsWithOuter(&Class, ClassSubobjects);
+	GetObjectsWithOuter(Class, ClassSubobjects);
 	for (auto ClassSubobject : ClassSubobjects)
 	{
 		if (auto Enum = dynamic_cast<UEnum*>(ClassSubobject))
@@ -53,7 +53,7 @@ void ReloadRenameObject(UClass& Class, const TCHAR* RenamePrefix)
 
 	// Reset singletons for any child fields
 #if WITH_LIVE_CODING
-	for (UField* Field = Class.Children; Field != nullptr; Field = Field->Next)
+	for (UField* Field = Class->Children; Field != nullptr; Field = Field->Next)
 	{
 		UFunction* Function = Cast<UFunction>(Field);
 		if (Function != nullptr && Function->SingletonPtr != nullptr)
@@ -64,4 +64,17 @@ void ReloadRenameObject(UClass& Class, const TCHAR* RenamePrefix)
 #endif
 }
 
+void ReloadProcessObject(UPackage* Package, const TCHAR* RenamePrefix)
+{
+#if WITH_LIVE_CODING
+	for (UFunction* Function : Package->Delegates)
+	{
+		if (Function->SingletonPtr != nullptr)
+		{
+			*Function->SingletonPtr = nullptr;
+		}
+	}
+#endif
+}
+
 #endif // WITH_RELOAD
diff --git a/Engine/Source/Runtime/CoreUObject/Private/UObject/DeferredRegistry.h b/Engine/Source/Runtime/CoreUObject/Private/UObject/DeferredRegistry.h
--- a/Engine/Source/Runtime/CoreUObject/Private/UObject/DeferredRegistry.h
+++ b/Engine/Source/Runtime/CoreUObject/Private/UObject/DeferredRegistry.h
@@ -23,9 +23,15 @@ inline bool operator != (const FEnumReloadVersionInfo& lhs, const FEnumReloadVer
 	return lhs.Hash != rhs.Hash;
 }
 
-void ReloadRenameObject(UScriptStruct& Object, const TCHAR* RenamePrefix);
-void ReloadRenameObject(UEnum& Object, const TCHAR* RenamePrefix);
-void ReloadRenameObject(UClass& Object, const TCHAR* RenamePrefix);
+inline bool operator != (const FPackageReloadVersionInfo& lhs, const FPackageReloadVersionInfo& rhs)
+{
+	return lhs.BodyHash != rhs.BodyHash || lhs.DeclarationsHash != rhs.DeclarationsHash;
+}
+
+void ReloadProcessObject(UScriptStruct* Object, const TCHAR* RenamePrefix);
+void ReloadProcessObject(UEnum* Object, const TCHAR* RenamePrefix);
+void ReloadProcessObject(UClass* Object, const TCHAR* RenamePrefix);
+void ReloadProcessObject(UPackage* Object, const TCHAR* RenamePrefix);
 #endif
 
 PRAGMA_DISABLE_DEPRECATION_WARNINGS
@@ -83,7 +89,7 @@ public:
 
 		bool bHasChanged = !ExistingInfo || (*ExistingInfo)->ReloadVersionInfo != InVersion;
 		InInfo.ReloadVersionInfo = InVersion;
-		TType* OldSingleton = ExistingInfo ? (*ExistingInfo)->InnerSingleton : nullptr;
+		TType* OldSingleton = ExistingInfo ? ((*ExistingInfo)->InnerSingleton ? (*ExistingInfo)->InnerSingleton : (*ExistingInfo)->OuterSingleton) : nullptr;
 
 		bool bAdd = true;
 		if (ExistingInfo)
@@ -234,9 +240,9 @@ public:
 	}
 
 	/**
-	* Rename and existing objects that have changed
+	* Process existing objects that have changed
 	*/
-	void RenameChangedObjects()
+	void ProcessChangedObjects(bool InvokeOuterRegisterFunction = false)
 	{
 #if WITH_RELOAD
 		IReload* Reload = GetActiveReloadInterface();
@@ -248,11 +254,19 @@ public:
 				if (Registrant.bHasChanged && Registrant.OldSingleton != nullptr)
 				{
 					UE_LOG(LogClass, Log, TEXT("%s Reload."), *Registrant.OldSingleton->GetName());
-					ReloadRenameObject(*Registrant.OldSingleton, Prefix);
 
 					// Reset the cached class construct info.
 					Registrant.Info->InnerSingleton = nullptr;
 					Registrant.Info->OuterSingleton = nullptr;
+
+					// Perform any other processing needed for existing objects
+					ReloadProcessObject(Registrant.OldSingleton, Prefix);
+
+					// If requested, invoke the registration function to make sure it has been recreated
+					if (InvokeOuterRegisterFunction)
+					{
+						Registrant.OuterRegisterFn();
+					}
 				}
 			}
 		}
@@ -349,5 +363,6 @@ inline void TDeferredRegistry<FClassRegistrationInfo>::UpdateSingletons(FClassRe
 using FClassDeferredRegistry = TDeferredRegistry<FClassRegistrationInfo>;
 using FEnumDeferredRegistry = TDeferredRegistry<FEnumRegistrationInfo>;
 using FStructDeferredRegistry = TDeferredRegistry<FStructRegistrationInfo>;
+using FPackageDeferredRegistry = TDeferredRegistry<FPackageRegistrationInfo>;
 
 PRAGMA_ENABLE_DEPRECATION_WARNINGS
diff --git a/Engine/Source/Runtime/CoreUObject/Private/UObject/UObjectBase.cpp b/Engine/Source/Runtime/CoreUObject/Private/UObject/UObjectBase.cpp
--- a/Engine/Source/Runtime/CoreUObject/Private/UObject/UObjectBase.cpp
+++ b/Engine/Source/Runtime/CoreUObject/Private/UObject/UObjectBase.cpp
@@ -753,6 +753,16 @@ void UObjectCompiledInDefer(UClass *(*InRegister)(), UClass *(*InStaticClass)(),
 	}
 }
 
+// UPackage registration
+
+void RegisterCompiledInInfo(UPackage* (*InOuterRegister)(), const TCHAR* InPackageName, FPackageRegistrationInfo& InInfo, const FPackageReloadVersionInfo& InVersionInfo)
+{
+#if WITH_RELOAD
+	check(InOuterRegister);
+	FPackageDeferredRegistry::Get().AddRegistration(reinterpret_cast<class UPackage* (*)()>(InOuterRegister), nullptr, TEXT(""), InPackageName, InInfo, InVersionInfo, nullptr);
+#endif
+}
+
 /** Register all loaded classes */
 void UClassRegisterAllCompiledInClasses()
 {
@@ -763,7 +773,7 @@ void UClassRegisterAllCompiledInClasses()
 
 	FClassDeferredRegistry& Registry = FClassDeferredRegistry::Get();
 
-	Registry.RenameChangedObjects();
+	Registry.ProcessChangedObjects();
 
 	for (const FClassDeferredRegistry::FRegistrant& Registrant : Registry.GetRegistrations())
 	{
@@ -931,12 +941,14 @@ void ProcessNewlyLoadedUObjects(FName Package, bool bCanProcessNewlyLoadedObject
 	LLM_SCOPE(ELLMTag::UObject);
 	DECLARE_SCOPE_CYCLE_COUNTER(TEXT("ProcessNewlyLoadedUObjects"), STAT_ProcessNewlyLoadedUObjects, STATGROUP_ObjectVerbose);
 
+	FPackageDeferredRegistry& PackageRegistry = FPackageDeferredRegistry::Get();
 	FClassDeferredRegistry& ClassRegistry = FClassDeferredRegistry::Get();
 	FStructDeferredRegistry& StructRegistry = FStructDeferredRegistry::Get();
 	FEnumDeferredRegistry& EnumRegistry = FEnumDeferredRegistry::Get();
 
-	StructRegistry.RenameChangedObjects();
-	EnumRegistry.RenameChangedObjects();
+	PackageRegistry.ProcessChangedObjects(true);
+	StructRegistry.ProcessChangedObjects();
+	EnumRegistry.ProcessChangedObjects();
 
 	UClassRegisterAllCompiledInClasses();
 
@@ -960,6 +972,7 @@ void ProcessNewlyLoadedUObjects(FName Package, bool bCanProcessNewlyLoadedObject
 	if (Reload != nullptr)
 	{
 		UClassReplaceReloadClasses(); // Legacy
+		PackageRegistry.NotifyReload(*Reload);
 		EnumRegistry.NotifyReload(*Reload);
 		StructRegistry.NotifyReload(*Reload);
 		ClassRegistry.NotifyReload(*Reload);
@@ -967,6 +980,7 @@ void ProcessNewlyLoadedUObjects(FName Package, bool bCanProcessNewlyLoadedObject
 	}
 #endif
 
+	PackageRegistry.EmptyRegistrations();
 	EnumRegistry.EmptyRegistrations();
 	StructRegistry.EmptyRegistrations();
 	ClassRegistry.EmptyRegistrations();
diff --git a/Engine/Source/Runtime/CoreUObject/Private/UObject/UObjectGlobals.cpp b/Engine/Source/Runtime/CoreUObject/Private/UObject/UObjectGlobals.cpp
--- a/Engine/Source/Runtime/CoreUObject/Private/UObject/UObjectGlobals.cpp
+++ b/Engine/Source/Runtime/CoreUObject/Private/UObject/UObjectGlobals.cpp
@@ -4918,10 +4918,25 @@ namespace UECodeGen_Private
 		NewPackage->SetGuid(FGuid(Params.BodyCRC, Params.DeclarationsCRC, 0u, 0u));
 		PRAGMA_ENABLE_DEPRECATION_WARNINGS
 
+#if WITH_RELOAD
+		TArray<UFunction*> Delegates;
+		Delegates.Reserve(Params.NumSingletons);
+#endif
 		for (UObject* (*const *SingletonFunc)() = Params.SingletonFuncArray, *(*const *SingletonFuncEnd)() = SingletonFunc + Params.NumSingletons; SingletonFunc != SingletonFuncEnd; ++SingletonFunc)
 		{
+#if WITH_RELOAD
+			UObject* Object = (*SingletonFunc)();
+			if (UFunction* Function = Cast<UFunction>(Object))
+			{
+				Delegates.Add(Function);
+			}
+#else
 			(*SingletonFunc)();
+#endif
 		}
+#if WITH_RELOAD
+		NewPackage->Delegates = MoveTemp(Delegates);
+#endif
 	}
 
 	void ConstructUClass(UClass*& OutClass, const FClassParams& Params)
diff --git a/Engine/Source/Runtime/CoreUObject/Public/UObject/Package.h b/Engine/Source/Runtime/CoreUObject/Public/UObject/Package.h
--- a/Engine/Source/Runtime/CoreUObject/Public/UObject/Package.h
+++ b/Engine/Source/Runtime/CoreUObject/Public/UObject/Package.h
@@ -251,6 +251,11 @@ public:
 	/** size of the file for this package; if the package was not loaded from a file or was a forced export in another package, this will be zero */
 	uint64 FileSize;			// TODO: strip from runtime?
 
+#if WITH_RELOAD
+	/** Link list of delegates registered to the package.  The next pointer chain can't be used for this. */
+	TArray<UFunction*> Delegates;
+#endif
+
 #if WITH_EDITORONLY_DATA
 	/** Editor only: Thumbnails stored in this package */
 	TUniquePtr< FThumbnailMap > ThumbnailMap;
diff --git a/Engine/Source/Runtime/CoreUObject/Public/UObject/Reload.h b/Engine/Source/Runtime/CoreUObject/Public/UObject/Reload.h
--- a/Engine/Source/Runtime/CoreUObject/Public/UObject/Reload.h
+++ b/Engine/Source/Runtime/CoreUObject/Public/UObject/Reload.h
@@ -47,6 +47,7 @@ public:
 	virtual void NotifyChange(UClass* New, UClass* Old) = 0;
 	virtual void NotifyChange(UEnum* New, UEnum* Old) = 0;
 	virtual void NotifyChange(UScriptStruct* New, UScriptStruct* Old) = 0;
+	virtual void NotifyChange(UPackage* New, UPackage* Old) = 0;
 
 	/**
 	 * Perform the re-instancing 
diff --git a/Engine/Source/Runtime/CoreUObject/Public/UObject/UObjectBase.h b/Engine/Source/Runtime/CoreUObject/Public/UObject/UObjectBase.h
--- a/Engine/Source/Runtime/CoreUObject/Public/UObject/UObjectBase.h
+++ b/Engine/Source/Runtime/CoreUObject/Public/UObject/UObjectBase.h
@@ -350,17 +350,17 @@ struct FEnumReloadVersionInfo
 };
 
 /**
- * Registration information for structures
+ * Registration information for enums
  */
 using FEnumRegistrationInfo = TRegistrationInfo<UEnum, FEnumReloadVersionInfo>;
 
 /**
- * Adds a static structure registration and version information. The InInfo parameter must be static.
+ * Adds a static enum registration and version information. The InInfo parameter must be static.
  */
 COREUOBJECT_API void RegisterCompiledInInfo(class UEnum* (*InOuterRegister)(), const TCHAR* InPackageName, const TCHAR* InName, FEnumRegistrationInfo& InInfo, const FEnumReloadVersionInfo& InVersionInfo);
 
 /**
- * Adds a dynamic struct registration.
+ * Adds a dynamic enum registration.
  */
 COREUOBJECT_API void RegisterCompiledInInfo(class UEnum* (*InRegister)(), const TCHAR* PackageName, const TCHAR* ObjectName, const TCHAR* DynamicPackageName, const TCHAR* DynamicPathName);
 
@@ -369,6 +369,27 @@ COREUOBJECT_API void RegisterCompiledInInfo(class UEnum* (*InRegister)(), const
  */
 COREUOBJECT_API class UEnum* GetStaticEnum(class UEnum* (*InRegister)(), UObject* EnumOuter, const TCHAR* EnumName);
 
+/**
+ * Reload version information for packages 
+ */
+struct FPackageReloadVersionInfo
+{
+#if WITH_RELOAD
+	uint32 BodyHash = 0;
+	uint32 DeclarationsHash = 0;
+#endif
+};
+
+/**
+ * Registration information for packages
+ */
+using FPackageRegistrationInfo = TRegistrationInfo<UPackage, FPackageReloadVersionInfo>;
+
+/**
+ * Adds a static package registration and version information. The InInfo parameter must be static.
+ */
+COREUOBJECT_API void RegisterCompiledInInfo(UPackage* (*InOuterRegister)(), const TCHAR* InPackageName, FPackageRegistrationInfo& InInfo, const FPackageReloadVersionInfo& InVersionInfo);
+
 /**
  * Base class for deferred native class registration
  */
