Fix issues with the TObjectPtr upgrade toolchain in preparation of another pass to make ShooterGame compliant and enforce the use of TObjectPtr for all member properties.

diff --git a/Engine/Programs/UnrealHeaderTool/Config/DefaultEngine.ini b/Engine/Programs/UnrealHeaderTool/Config/DefaultEngine.ini
--- a/Engine/Programs/UnrealHeaderTool/Config/DefaultEngine.ini
+++ b/Engine/Programs/UnrealHeaderTool/Config/DefaultEngine.ini
@@ -32,7 +32,7 @@ gc.MaxObjectsInProgram=262144
 +TypeRedirects=(OldType="TAssetSubclassOf", NewType="TSoftClassPtr")
 
 [UnrealHeaderTool]
-; This will increase the size of the UHT log files, but the size should return
-; to normal after native pointer upgrades are completed
-NativePointerMemberBehavior=AllowAndLog
-ObjectPtrMemberBehavior=AllowAndLog
+EngineNativePointerMemberBehavior=AllowAndLog
+EngineObjectPtrMemberBehavior=AllowSilently
+NonEngineNativePointerMemberBehavior=AllowAndLog
+NonEngineObjectPtrMemberBehavior=AllowSilently
diff --git a/Engine/Source/Programs/UnrealHeaderTool/Private/BaseParser.h b/Engine/Source/Programs/UnrealHeaderTool/Private/BaseParser.h
--- a/Engine/Source/Programs/UnrealHeaderTool/Private/BaseParser.h
+++ b/Engine/Source/Programs/UnrealHeaderTool/Private/BaseParser.h
@@ -55,9 +55,11 @@ struct FUHTConfig
 	// Default version of generated code. Defaults to oldest possible, unless specified otherwise in config.
 	EGeneratedCodeVersion DefaultGeneratedCodeVersion = EGeneratedCodeVersion::V1;
 
-	EPointerMemberBehavior NativePointerMemberBehavior = EPointerMemberBehavior::AllowSilently;
+	EPointerMemberBehavior EngineNativePointerMemberBehavior = EPointerMemberBehavior::AllowSilently;
+	EPointerMemberBehavior EngineObjectPtrMemberBehavior = EPointerMemberBehavior::AllowSilently;
 
-	EPointerMemberBehavior ObjectPtrMemberBehavior = EPointerMemberBehavior::AllowSilently;
+	EPointerMemberBehavior NonEngineNativePointerMemberBehavior = EPointerMemberBehavior::AllowSilently;
+	EPointerMemberBehavior NonEngineObjectPtrMemberBehavior = EPointerMemberBehavior::AllowSilently;
 
 private:
 	FUHTConfig();
diff --git a/Engine/Source/Programs/UnrealHeaderTool/Private/HeaderParser.cpp b/Engine/Source/Programs/UnrealHeaderTool/Private/HeaderParser.cpp
--- a/Engine/Source/Programs/UnrealHeaderTool/Private/HeaderParser.cpp
+++ b/Engine/Source/Programs/UnrealHeaderTool/Private/HeaderParser.cpp
@@ -709,8 +709,10 @@ FUHTConfig::FUHTConfig()
 	const FName StructsWithTPrefixKey(TEXT("StructsWithTPrefix"));
 	const FName DelegateParameterCountStringsKey(TEXT("DelegateParameterCountStrings"));
 	const FName GeneratedCodeVersionKey(TEXT("GeneratedCodeVersion"));
-	const FName NativePointerMemberBehaviorKey(TEXT("NativePointerMemberBehavior"));
-	const FName ObjectPtrMemberBehaviorKey(TEXT("ObjectPtrMemberBehavior"));
+	const FName EngineNativePointerMemberBehaviorKey(TEXT("EngineNativePointerMemberBehavior"));
+	const FName EngineObjectPtrMemberBehaviorKey(TEXT("EngineObjectPtrMemberBehavior"));
+	const FName NonEngineNativePointerMemberBehaviorKey(TEXT("NonEngineNativePointerMemberBehavior"));
+	const FName NonEngineObjectPtrMemberBehaviorKey(TEXT("NonEngineObjectPtrMemberBehavior"));
 
 	FConfigSection* ConfigSection = GConfig->GetSectionPrivate(TEXT("UnrealHeaderTool"), false, true, GEngineIni);
 	if (ConfigSection)
@@ -743,13 +745,21 @@ FUHTConfig::FUHTConfig()
 			{
 				DefaultGeneratedCodeVersion = ToGeneratedCodeVersion(It.Value().GetValue());
 			}
-			else if (It.Key() == NativePointerMemberBehaviorKey)
+			else if (It.Key() == EngineNativePointerMemberBehaviorKey)
 			{
-				NativePointerMemberBehavior = ToPointerMemberBehavior(It.Value().GetValue());
+				EngineNativePointerMemberBehavior = ToPointerMemberBehavior(It.Value().GetValue());
 			}
-			else if (It.Key() == ObjectPtrMemberBehaviorKey)
+			else if (It.Key() == EngineObjectPtrMemberBehaviorKey)
 			{
-				ObjectPtrMemberBehavior = ToPointerMemberBehavior(It.Value().GetValue());
+				EngineObjectPtrMemberBehavior = ToPointerMemberBehavior(It.Value().GetValue());
+			}
+			else if (It.Key() == NonEngineNativePointerMemberBehaviorKey)
+			{
+				NonEngineNativePointerMemberBehavior = ToPointerMemberBehavior(It.Value().GetValue());
+			}
+			else if (It.Key() == NonEngineObjectPtrMemberBehaviorKey)
+			{
+				NonEngineObjectPtrMemberBehavior = ToPointerMemberBehavior(It.Value().GetValue());
 			}
 		}
 	}
@@ -4084,7 +4094,8 @@ void FHeaderParser::GetVarType(
 						// Optionally emit messages about native pointer members and swallow trailing 'const' after pointer properties
 						if (VariableCategory == EVariableCategory::Member)
 						{
-							ConditionalLogPointerUsage(UHTConfig.NativePointerMemberBehavior, TEXT("Native pointer"), FString(InputPos - VarStartPos, Input + VarStartPos).TrimStartAndEnd().ReplaceCharWithEscapedChar());
+							ConditionalLogPointerUsage(bIsCurrentModulePartOfEngine ? UHTConfig.EngineNativePointerMemberBehavior : UHTConfig.NonEngineNativePointerMemberBehavior,
+								TEXT("Native pointer"), FString(InputPos - VarStartPos, Input + VarStartPos).TrimStartAndEnd().ReplaceCharWithEscapedChar());
 
 							MatchIdentifier(TEXT("const"), ESearchCase::CaseSensitive);
 						}
@@ -4093,7 +4104,8 @@ void FHeaderParser::GetVarType(
 					}
 					else if ((PropertyType == CPT_ObjectPtrReference) && (VariableCategory == EVariableCategory::Member))
 					{
-						ConditionalLogPointerUsage(UHTConfig.ObjectPtrMemberBehavior, TEXT("ObjectPtr"), FString(InputPos - VarStartPos, Input + VarStartPos).TrimStartAndEnd().ReplaceCharWithEscapedChar());
+						ConditionalLogPointerUsage(bIsCurrentModulePartOfEngine ? UHTConfig.EngineObjectPtrMemberBehavior : UHTConfig.NonEngineObjectPtrMemberBehavior,
+							TEXT("ObjectPtr"), FString(InputPos - VarStartPos, Input + VarStartPos).TrimStartAndEnd().ReplaceCharWithEscapedChar());
 					}
 
 					// Imply const if it's a parameter that is a pointer to a const class
@@ -8694,11 +8706,11 @@ bool FHeaderParser::CheckUIMinMaxRangeFromMetaData(const FString& UIMin, const F
 
 void FHeaderParser::ConditionalLogPointerUsage(EPointerMemberBehavior PointerMemberBehavior, const TCHAR* PointerTypeDesc, FString&& PointerTypeDecl)
 {
-	switch (UHTConfig.NativePointerMemberBehavior)
+	switch (PointerMemberBehavior)
 	{
 		case EPointerMemberBehavior::Disallow:
 		{
-			UE_LOG(LogCompile, Error, TEXT("%s usage in member declaration detected in '%s', line %d [[%s]]"), PointerTypeDesc, *Filename, InputLine, *PointerTypeDecl);
+			LogError(TEXT("%s usage in member declaration detected in '%s', line %d [[%s]]"), PointerTypeDesc, *Filename, InputLine, *PointerTypeDecl);
 		}
 		break;
 		case EPointerMemberBehavior::AllowAndLog:
diff --git a/Engine/Source/Programs/UnrealObjectPtrTool/Private/UnrealObjectPtrTool.cpp b/Engine/Source/Programs/UnrealObjectPtrTool/Private/UnrealObjectPtrTool.cpp
--- a/Engine/Source/Programs/UnrealObjectPtrTool/Private/UnrealObjectPtrTool.cpp
+++ b/Engine/Source/Programs/UnrealObjectPtrTool/Private/UnrealObjectPtrTool.cpp
@@ -486,7 +486,7 @@ bool FPointerUpgrader::SaveStringToFile(const FStringView& InString, EByteOrderM
 
 	if (InBOM == EByteOrderMarkerType::UTF8)
 	{
-		UTF8CHAR UTF8BOM[] = {0xEF, 0xBB, 0xBF};
+		UTF8CHAR UTF8BOM[] = {(UTF8CHAR)0xEF, (UTF8CHAR)0xBB, (UTF8CHAR)0xBF};
 		Ar->Serialize(&UTF8BOM, UE_ARRAY_COUNT(UTF8BOM) * sizeof(UTF8CHAR));
 
 		FTCHARToUTF8 UTF8String(InString.GetData(), InString.Len());
@@ -596,6 +596,7 @@ FString FPointerUpgrader::PerformReverseTypenameUpgrade(const FString& TypeName)
 
 INT32_MAIN_INT32_ARGC_TCHAR_ARGV()
 {
+	FTaskTagScope Scope(ETaskTag::EGameThread);
 	FString CmdLine = FCommandLine::BuildFromArgV(nullptr, ArgC, ArgV, nullptr);
 	GEngineLoop.PreInit(*CmdLine);
 
