Cleaning up the compiler directive code.
Add error message for when WITH_EDITOR used in shipping code to hide a property.

#rnx

diff --git a/Engine/Source/Programs/UnrealHeaderTool/Private/HeaderParser.cpp b/Engine/Source/Programs/UnrealHeaderTool/Private/HeaderParser.cpp
--- a/Engine/Source/Programs/UnrealHeaderTool/Private/HeaderParser.cpp
+++ b/Engine/Source/Programs/UnrealHeaderTool/Private/HeaderParser.cpp
@@ -1031,7 +1031,7 @@ FUnrealEnumDefinitionInfo& FHeaderParser::CompileEnum()
 		}
 	}
 
-	if (CompilerDirectiveStack.Num() > 0 && (CompilerDirectiveStack.Last() & ECompilerDirective::WithEditorOnlyData) != 0)
+	if ((GetCurrentCompilerDirective() & ECompilerDirective::WithEditorOnlyData) != 0)
 	{
 		EnumDef.MakeEditorOnly();
 	}
@@ -1935,11 +1935,7 @@ FUnrealScriptStructDefinitionInfo& FHeaderParser::CompileStructDeclaration()
 		}
 		else if (Token.Matches(TEXT('#')) && MatchIdentifier(TEXT("endif"), ESearchCase::CaseSensitive))
 		{
-			if (CompilerDirectiveStack.Num() < 1)
-			{
-				Throwf(TEXT("Unmatched '#endif' in class or global scope"));
-			}
-			CompilerDirectiveStack.Pop();
+			PopCompilerDirective();
 			// Do nothing and hope that the if code below worked out OK earlier
 		}
 		else if ( Token.Matches(TEXT('#')) && MatchIdentifier(TEXT("if"), ESearchCase::CaseSensitive) )
@@ -2725,11 +2721,7 @@ void FHeaderParser::CompileDirective()
 	}
 	else if (Directive.Matches(TEXT("endif"), ESearchCase::CaseSensitive))
 	{
-		if (CompilerDirectiveStack.Num() < 1)
-		{
-			Throwf(TEXT("Unmatched '#endif' in class or global scope"));
-		}
-		CompilerDirectiveStack.Pop();
+		PopCompilerDirective();
 	}
 	else if (Directive.Matches(TEXT("define"), ESearchCase::CaseSensitive))
 	{
@@ -2838,10 +2830,27 @@ void FHeaderParser::GetVarType(
 		}
 	}
 
-	if (CompilerDirectiveStack.Num() > 0 && (CompilerDirectiveStack.Last()&ECompilerDirective::WithEditorOnlyData) != 0)
+	uint32 CurrentCompilerDirective = GetCurrentCompilerDirective();
+	if ((CurrentCompilerDirective & ECompilerDirective::WithEditorOnlyData) != 0)
 	{
 		Flags |= CPF_EditorOnly;
 	}
+	else if ((CurrentCompilerDirective & ECompilerDirective::WithEditor) != 0)
+	{
+		// Checking for this error is a bit tricky given legacy code.  
+		// 1) If already wrapped in WITH_EDITORONLY_DATA (see above), then we ignore the error via the else 
+		// 2) Ignore any module that is an editor module
+		const FManifestModule& Module = Scope->GetFileScope()->GetSourceFile()->GetPackageDef().GetModule();
+		const bool bIsEditorModule =
+			Module.ModuleType == EBuildModuleType::EngineEditor ||
+			Module.ModuleType == EBuildModuleType::GameEditor || 
+			Module.ModuleType == EBuildModuleType::EngineUncooked || 
+			Module.ModuleType == EBuildModuleType::GameUncooked;
+		if (VariableCategory == EVariableCategory::Member && !bIsEditorModule)
+		{
+			LogError(TEXT("UProperties should not be wrapped by WITH_EDITOR, use WITH_EDITORONLY_DATA instead."));
+		}
+	}
 
 	// Store the start and end positions of the parsed type
 	if (ParsedVarIndexRange)
@@ -4935,9 +4944,9 @@ bool FHeaderParser::CompileDeclaration(TArray<FUnrealFunctionDefinitionInfo*>& D
 					if (ArchiveType != ESerializerArchiveType::None)
 					{
 						// Found what we want!
-						if (CompilerDirectiveStack.Num() == 0 || (CompilerDirectiveStack.Num() == 1 && CompilerDirectiveStack[0] == ECompilerDirective::WithEditorOnlyData))
+						if (uint32 CurrentCompilerDirective = GetCurrentCompilerDirective(); CurrentCompilerDirective == 0 || CurrentCompilerDirective == ECompilerDirective::WithEditorOnlyData)
 						{
-							FString EnclosingDefine = CompilerDirectiveStack.Num() > 0 ? TEXT("WITH_EDITORONLY_DATA") : TEXT("");
+							FString EnclosingDefine = CurrentCompilerDirective != 0 ? TEXT("WITH_EDITORONLY_DATA") : TEXT("");
 
 							FUnrealClassDefinitionInfo& ClassDef = UHTCastChecked<FUnrealClassDefinitionInfo>(GetCurrentClassDef());
 							ClassDef.AddArchiveType(ArchiveType);
@@ -6141,7 +6150,7 @@ FUnrealFunctionDefinitionInfo& FHeaderParser::CompileFunctionDeclaration()
 		FuncInfo.FunctionFlags |= FUNC_Event;
 	}
 
-	if (CompilerDirectiveStack.Num() > 0 && (CompilerDirectiveStack.Last()&ECompilerDirective::WithEditor) != 0)
+	if ((GetCurrentCompilerDirective() & ECompilerDirective::WithEditor) != 0)
 	{
 		FuncInfo.FunctionFlags |= FUNC_EditorOnly;
 	}
diff --git a/Engine/Source/Programs/UnrealHeaderTool/Private/HeaderParser.h b/Engine/Source/Programs/UnrealHeaderTool/Private/HeaderParser.h
--- a/Engine/Source/Programs/UnrealHeaderTool/Private/HeaderParser.h
+++ b/Engine/Source/Programs/UnrealHeaderTool/Private/HeaderParser.h
@@ -277,10 +277,26 @@ protected:
 	 */
 	TArray<uint32> CompilerDirectiveStack;
 
+	// Return the top level compiler directive state
+	uint32 GetCurrentCompilerDirective() const
+	{
+		return CompilerDirectiveStack.IsEmpty() ? 0 : CompilerDirectiveStack.Last();
+	}
+
 	// Pushes the Directive specified to the CompilerDirectiveStack according to the rules described above
-	void FORCEINLINE PushCompilerDirective(ECompilerDirective::Type Directive)
+	void PushCompilerDirective(ECompilerDirective::Type Directive)
+	{
+		CompilerDirectiveStack.Push(Directive | GetCurrentCompilerDirective());
+	}
+
+	// Removes the top most compiler directive stack entry
+	void PopCompilerDirective()
 	{
-		CompilerDirectiveStack.Push(CompilerDirectiveStack.Num()>0 ? (CompilerDirectiveStack[CompilerDirectiveStack.Num()-1] | Directive) : Directive);
+		if (CompilerDirectiveStack.Num() < 1)
+		{
+			Throwf(TEXT("Unmatched '#endif' in class or global scope"));
+		}
+		CompilerDirectiveStack.Pop();
 	}
 
 	/**
