Added support for a UEnum not containing a literal zero value.  If that is the case, a static_assert is added to verify that one of the enum values is zero.

diff --git a/Engine/Source/Programs/UnrealHeaderTool/Private/CodeGenerator.cpp b/Engine/Source/Programs/UnrealHeaderTool/Private/CodeGenerator.cpp
--- a/Engine/Source/Programs/UnrealHeaderTool/Private/CodeGenerator.cpp
+++ b/Engine/Source/Programs/UnrealHeaderTool/Private/CodeGenerator.cpp
@@ -3887,6 +3887,39 @@ void FNativeClassHeaderGenerator::ExportGeneratedEnumInitCode(FOutputDevice& Out
 
 	const FString& PackageSingletonName = (bIsDynamic ? EnumDef.GetTypePackageName() : GetPackageSingletonName(UHTCastChecked<FUnrealPackageDefinitionInfo>(EnumDef.GetOuter()), OutReferenceGatherers.UniqueCrossModuleReferences));
 
+	// If we don't have a zero 0 then we emit a static assert to verify we have one
+	if (!EnumDef.IsValidEnumValue(0) && EnumDef.HasMetaData(FHeaderParserNames::NAME_BlueprintType))
+	{
+		bool bHasUnparsedValue = false;
+		for (const TPair<FName, int64>& Enum : EnumDef.GetEnums())
+		{
+			if (Enum.Value == -1)
+			{
+				bHasUnparsedValue = true;
+				break;
+			}
+		}
+
+		if (bHasUnparsedValue)
+		{
+			Out.Logf(TEXT("\tstatic_assert("));
+			bool bDoneFirst = false;
+			for (const TPair<FName, int64>& Enum : EnumDef.GetEnums())
+			{
+				if (Enum.Value == -1)
+				{
+					if (bDoneFirst)
+					{
+						Out.Logf(TEXT("||"));
+					}
+					bDoneFirst = true;
+					Out.Logf(TEXT("!int64(%s)"), *Enum.Key.ToString());
+				}
+			}
+			Out.Logf(TEXT(", \"'%s' does not have a 0 entry!(This is a problem when the enum is initalized by default)\");\r\n"), *EnumDef.GetName());
+		}
+	}
+
 	if (!bIsDynamic)
 	{
 		Out.Logf(TEXT("\tstatic FEnumRegistrationInfo& Z_Registration_Info_UEnum_%s()\r\n"), *EnumNameCpp);
diff --git a/Engine/Source/Programs/UnrealHeaderTool/Private/HeaderParser.cpp b/Engine/Source/Programs/UnrealHeaderTool/Private/HeaderParser.cpp
--- a/Engine/Source/Programs/UnrealHeaderTool/Private/HeaderParser.cpp
+++ b/Engine/Source/Programs/UnrealHeaderTool/Private/HeaderParser.cpp
@@ -1107,6 +1107,7 @@ FUnrealEnumDefinitionInfo& FHeaderParser::CompileEnum()
 	FMetaData TagMetaData;
 	TArray<TMap<FName, FString>> EntryMetaData;
 
+	bool bHasUnparsedValue = false;
 	TArray<TPair<FName, int64>> EnumNames;
 	int64 CurrentEnumValue = 0;
 	while (GetIdentifier(TagToken))
@@ -1127,6 +1128,7 @@ FUnrealEnumDefinitionInfo& FHeaderParser::CompileEnum()
 			{
 				// We didn't parse a literal, so set an invalid value
 				NewEnumValue = -1;
+				bHasUnparsedValue = true;
 			}
 
 			// Skip tokens until we encounter a comma, a closing brace or a UMETA declaration
@@ -1157,6 +1159,7 @@ FUnrealEnumDefinitionInfo& FHeaderParser::CompileEnum()
 				// There are tokens after the initializer so it's not a standalone literal,
 				// so set it to an invalid value.
 				NewEnumValue = -1;
+				bHasUnparsedValue = true;
 			}
 
 			CurrentEnumValue = NewEnumValue;
@@ -1252,7 +1255,7 @@ FUnrealEnumDefinitionInfo& FHeaderParser::CompileEnum()
 
 	CheckDocumentationPolicyForEnum(EnumDef, EnumValueMetaData, EntryMetaData);
 
-	if (!EnumDef.IsValidEnumValue(0) && EnumMetaData.Contains(FHeaderParserNames::NAME_BlueprintType))
+	if (!EnumDef.IsValidEnumValue(0) && EnumMetaData.Contains(FHeaderParserNames::NAME_BlueprintType) && !bHasUnparsedValue)
 	{
 		EnumDef.LogWarning(TEXT("'%s' does not have a 0 entry! (This is a problem when the enum is initalized by default)"), *EnumDef.GetName());
 	}
diff --git a/Engine/Source/Programs/UnrealHeaderTool/Private/UnrealTypeDefinitionInfo.h b/Engine/Source/Programs/UnrealHeaderTool/Private/UnrealTypeDefinitionInfo.h
--- a/Engine/Source/Programs/UnrealHeaderTool/Private/UnrealTypeDefinitionInfo.h
+++ b/Engine/Source/Programs/UnrealHeaderTool/Private/UnrealTypeDefinitionInfo.h
@@ -1872,6 +1872,14 @@ public:
 	 */
 	bool ContainsExistingMax() const;
 
+	/**
+	 * Return the enumeration values
+	 */
+	const TArray<TPair<FName, int64>>& GetEnums()
+	{
+		return Names;
+	}
+
 	/**
 	 * @return The number of enum names.
 	 */
