Removed the Class/Struct meta data class.
Fixed some API that were mistakenly setup as generic struct API instead of class API.
Removed some unused API.
No functional changes, just code reorg.

#rnx

diff --git a/Engine/Source/Programs/UnrealHeaderTool/Private/CodeGenerator.cpp b/Engine/Source/Programs/UnrealHeaderTool/Private/CodeGenerator.cpp
--- a/Engine/Source/Programs/UnrealHeaderTool/Private/CodeGenerator.cpp
+++ b/Engine/Source/Programs/UnrealHeaderTool/Private/CodeGenerator.cpp
@@ -2090,7 +2090,7 @@ void FNativeClassHeaderGenerator::ExportNativeGeneratedInitCode(FOutputDevice& O
 			FunctionsCount = TEXT("0");
 		}
 
-		if (ClassDef.GetStructMetaData().bObjectInitializerConstructorDeclared)
+		if (ClassDef.IsObjectInitializerConstructorDeclared())
 		{
 			ClassDef.SetMetaData(NAME_ObjectInitializerConstructorDeclared, TEXT(""));
 		}
@@ -2664,7 +2664,7 @@ FString GetGeneratedMacroDeprecationWarning(const TCHAR* MacroName)
 FString GetPreservedAccessSpecifierString(FUnrealClassDefinitionInfo& ClassDef)
 {
 	FString PreservedAccessSpecifier;
-	switch (ClassDef.GetStructMetaData().GeneratedBodyMacroAccessSpecifier)
+	switch (ClassDef.GetGeneratedBodyMacroAccessSpecifier())
 	{
 	case EAccessSpecifier::ACCESS_Private:
 		PreservedAccessSpecifier = "private:";
@@ -2702,7 +2702,6 @@ void FNativeClassHeaderGenerator::ExportClassFromSourceFileInner(
 	FUHTStringBuilder EnhancedUObjectConstructorsMacroCall;
 
 	UClass* Class = ClassDef.GetClass();
-	FStructMetaData& StructData = ClassDef.GetStructMetaData();
 
 	FUnrealClassDefinitionInfo* SuperClassDef = ClassDef.GetSuperClass();
 	UClass* SuperClass = SuperClassDef ? SuperClassDef->GetClass() : nullptr;
@@ -2735,14 +2734,14 @@ void FNativeClassHeaderGenerator::ExportClassFromSourceFileInner(
 			ExportEventParm(UClassMacroContent, OutReferenceGatherers.ForwardDeclarations, *FunctionDef, /*Indent=*/ 1, /*bOutputConstructor=*/ true, EExportingState::Normal);
 		}
 
-		FString MacroName = SourceFile.GetGeneratedMacroName(StructData, TEXT("_EVENT_PARMS"));
+		FString MacroName = SourceFile.GetGeneratedMacroName(ClassDef.GetGeneratedBodyLine(), TEXT("_EVENT_PARMS"));
 		WriteMacro(OutGeneratedHeaderText, MacroName, UClassMacroContent);
 		PrologMacroCalls.Logf(TEXT("\t%s\r\n"), *MacroName);
 
 		// VM -> C++ proxies (events and delegates).
 		FOutputDeviceNull NullOutput;
 		FOutputDevice& CallbackOut = ClassDef.HasAnyClassFlags(CLASS_NoExport) ? NullOutput : OutCpp;
-		FString CallbackWrappersMacroName = SourceFile.GetGeneratedMacroName(StructData, TEXT("_CALLBACK_WRAPPERS"));
+		FString CallbackWrappersMacroName = SourceFile.GetGeneratedMacroName(ClassDef.GetGeneratedBodyLine(), TEXT("_CALLBACK_WRAPPERS"));
 		ExportCallbackFunctions(
 			OutGeneratedHeaderText,
 			CallbackOut,
@@ -2789,7 +2788,7 @@ void FNativeClassHeaderGenerator::ExportClassFromSourceFileInner(
 		FUHTStringBuilder Boilerplate, BoilerPlateCPP;
 		const TCHAR* MacroNameHeader;
 		const TCHAR* MacroNameCPP;
-		GeneratedSerializeFunctionHeaderMacroName = SourceFile.GetGeneratedMacroName(StructData, TEXT("_ARCHIVESERIALIZER"));
+		GeneratedSerializeFunctionHeaderMacroName = SourceFile.GetGeneratedMacroName(ClassDef.GetGeneratedBodyLine(), TEXT("_ARCHIVESERIALIZER"));
 
 		if (ClassDef.GetArchiveType() == ESerializerArchiveType::StructuredArchiveRecord)
 		{
@@ -2859,12 +2858,12 @@ void FNativeClassHeaderGenerator::ExportClassFromSourceFileInner(
 
 		if (ClassDef.HasAnyClassFlags(CLASS_Interface))
 		{
-			ExportConstructorsMacros(OutGeneratedHeaderText, OutCpp, StandardUObjectConstructorsMacroCall, EnhancedUObjectConstructorsMacroCall, SourceFile.GetGeneratedMacroName(StructData), ClassDef, *APIArg);
+			ExportConstructorsMacros(OutGeneratedHeaderText, OutCpp, StandardUObjectConstructorsMacroCall, EnhancedUObjectConstructorsMacroCall, SourceFile.GetGeneratedMacroName(ClassDef.GetGeneratedBodyLine()), ClassDef, *APIArg);
 
-			FString InterfaceMacroName = SourceFile.GetGeneratedMacroName(StructData, TEXT("_GENERATED_UINTERFACE_BODY"));
+			FString InterfaceMacroName = SourceFile.GetGeneratedMacroName(ClassDef.GetGeneratedBodyLine(), TEXT("_GENERATED_UINTERFACE_BODY"));
 			OutGeneratedHeaderText.Log(Macroize(*(InterfaceMacroName + TEXT("()")), *Boilerplate));
 
-			int32 ClassGeneratedBodyLine = StructData.GetGeneratedBodyLine();
+			int32 ClassGeneratedBodyLine = ClassDef.GetGeneratedBodyLine();
 
 			FString DeprecationWarning = GetGeneratedMacroDeprecationWarning(TEXT("GENERATED_UINTERFACE_BODY"));
 
@@ -2926,11 +2925,11 @@ void FNativeClassHeaderGenerator::ExportClassFromSourceFileInner(
 				WriteReplicatedMacroData(*ClassCPPName, *APIArg, ClassDef, SuperClass, InterfaceBoilerplate, SourceFile, OutFlags);
 			}
 
-			FString NoPureDeclsMacroName = SourceFile.GetGeneratedMacroName(StructData, TEXT("_INCLASS_IINTERFACE_NO_PURE_DECLS"));
+			FString NoPureDeclsMacroName = SourceFile.GetGeneratedMacroName(ClassDef.GetGeneratedBodyLine(), TEXT("_INCLASS_IINTERFACE_NO_PURE_DECLS"));
 			WriteMacro(OutGeneratedHeaderText, NoPureDeclsMacroName, InterfaceBoilerplate);
 			ClassNoPureDeclsMacroCalls.Logf(TEXT("\t%s\r\n"), *NoPureDeclsMacroName);
 
-			FString MacroName = SourceFile.GetGeneratedMacroName(StructData, TEXT("_INCLASS_IINTERFACE"));
+			FString MacroName = SourceFile.GetGeneratedMacroName(ClassDef.GetGeneratedBodyLine(), TEXT("_INCLASS_IINTERFACE"));
 			WriteMacro(OutGeneratedHeaderText, MacroName, InterfaceBoilerplate);
 			ClassMacroCalls.Logf(TEXT("\t%s\r\n"), *MacroName);
 		}
@@ -2954,21 +2953,21 @@ void FNativeClassHeaderGenerator::ExportClassFromSourceFileInner(
 			}
 
 			{
-				FString NoPureDeclsMacroName = SourceFile.GetGeneratedMacroName(StructData, TEXT("_INCLASS_NO_PURE_DECLS"));
+				FString NoPureDeclsMacroName = SourceFile.GetGeneratedMacroName(ClassDef.GetGeneratedBodyLine(), TEXT("_INCLASS_NO_PURE_DECLS"));
 				WriteMacro(OutGeneratedHeaderText, NoPureDeclsMacroName, Boilerplate);
 				ClassNoPureDeclsMacroCalls.Logf(TEXT("\t%s\r\n"), *NoPureDeclsMacroName);
 
-				FString MacroName = SourceFile.GetGeneratedMacroName(StructData, TEXT("_INCLASS"));
+				FString MacroName = SourceFile.GetGeneratedMacroName(ClassDef.GetGeneratedBodyLine(), TEXT("_INCLASS"));
 				WriteMacro(OutGeneratedHeaderText, MacroName, Boilerplate);
 				ClassMacroCalls.Logf(TEXT("\t%s\r\n"), *MacroName);
 
-				ExportConstructorsMacros(OutGeneratedHeaderText, OutCpp, StandardUObjectConstructorsMacroCall, EnhancedUObjectConstructorsMacroCall, SourceFile.GetGeneratedMacroName(StructData), ClassDef, *APIArg);
+				ExportConstructorsMacros(OutGeneratedHeaderText, OutCpp, StandardUObjectConstructorsMacroCall, EnhancedUObjectConstructorsMacroCall, SourceFile.GetGeneratedMacroName(ClassDef.GetGeneratedBodyLine()), ClassDef, *APIArg);
 			}
 		}
 	}
 
 	{
-		FString MacroName = SourceFile.GetGeneratedMacroName(StructData.GetPrologLine(), TEXT("_PROLOG"));
+		FString MacroName = SourceFile.GetGeneratedMacroName(ClassDef.GetPrologLine(), TEXT("_PROLOG"));
 		WriteMacro(OutGeneratedHeaderText, MacroName, PrologMacroCalls);
 	}
 
@@ -2987,7 +2986,7 @@ void FNativeClassHeaderGenerator::ExportClassFromSourceFileInner(
 		if (bIsIInterface)
 		{
 			MacroName = TEXT("GENERATED_IINTERFACE_BODY()");
-			GeneratedBodyLine = StructData.GetInterfaceGeneratedBodyLine();
+			GeneratedBodyLine = ClassDef.GetInterfaceGeneratedBodyLine();
 			LegacyGeneratedBody = ClassMacroCalls;
 			GeneratedBody = ClassNoPureDeclsMacroCalls;
 		}
@@ -2995,7 +2994,7 @@ void FNativeClassHeaderGenerator::ExportClassFromSourceFileInner(
 		{
 			MacroName = TEXT("GENERATED_UCLASS_BODY()");
 			DeprecationWarning = GetGeneratedMacroDeprecationWarning(MacroName);
-			GeneratedBodyLine = StructData.GetGeneratedBodyLine();
+			GeneratedBodyLine = ClassDef.GetGeneratedBodyLine();
 			LegacyGeneratedBody = FString::Printf(TEXT("%s%s"), *ClassMacroCalls, *StandardUObjectConstructorsMacroCall);
 			GeneratedBody = FString::Printf(TEXT("%s%s"), *ClassNoPureDeclsMacroCalls, *EnhancedUObjectConstructorsMacroCall);
 		}
@@ -3042,7 +3041,7 @@ void ExportCopyConstructorDefinition(FOutputDevice& Out, const TCHAR* API, const
  */
 void ExportVTableHelperCtorAndCaller(FOutputDevice& Out, FUnrealClassDefinitionInfo& ClassDef, const TCHAR* API, const TCHAR* ClassCPPName)
 {
-	if (!ClassDef.GetStructMetaData().bCustomVTableHelperConstructorDeclared)
+	if (!ClassDef.IsCustomVTableHelperConstructorDeclared())
 	{
 		Out.Logf(TEXT("\tDECLARE_VTABLE_PTR_HELPER_CTOR(%s_API, %s);" LINE_TERMINATOR), API, ClassCPPName);
 	}
@@ -3062,7 +3061,7 @@ void ExportStandardConstructorsMacro(FOutputDevice& Out, FUnrealClassDefinitionI
 	{
 		Out.Logf(TEXT("\t/** Standard constructor, called after all reflected properties have been initialized */\r\n"));
 		Out.Logf(TEXT("\t%s_API %s(const FObjectInitializer& ObjectInitializer%s);\r\n"), API, ClassCPPName,
-			ClassDef.GetStructMetaData().bDefaultConstructorDeclared ? TEXT("") : TEXT(" = FObjectInitializer::Get()"));
+			ClassDef.IsDefaultConstructorDeclared() ? TEXT("") : TEXT(" = FObjectInitializer::Get()"));
 	}
 	if (ClassDef.HasAnyClassFlags(CLASS_Abstract))
 	{
@@ -3087,8 +3086,7 @@ void ExportStandardConstructorsMacro(FOutputDevice& Out, FUnrealClassDefinitionI
  */
 void ExportConstructorDefinition(FOutputDevice& Out, FUnrealClassDefinitionInfo& ClassDef, const TCHAR* API, const TCHAR* ClassCPPName)
 {
-	FStructMetaData& StructData = ClassDef.GetStructMetaData();
-	if (!StructData.bConstructorDeclared)
+	if (!ClassDef.IsConstructorDeclared())
 	{
 		Out.Logf(TEXT("\t/** Standard constructor, called after all reflected properties have been initialized */\r\n"));
 
@@ -3105,27 +3103,26 @@ void ExportConstructorDefinition(FOutputDevice& Out, FUnrealClassDefinitionInfo&
 				// Since the SourceFile array provided to the ParallelFor is in dependency order and does not allow cyclic dependencies, 
 				// we can be certain that another thread has started processing the file containing our SuperClass before this
 				// file would have been assigned out,  so we just have to wait
-				FStructMetaData& SuperClassData = SuperClassDef->GetStructMetaData();
-				while (!SuperClassData.bConstructorDeclared)
+				while (!SuperClassDef->IsConstructorDeclared())
 				{
 					FPlatformProcess::Sleep(0.01);
 				}
 
-				bSuperClassObjectInitializerConstructorDeclared = SuperClassData.bObjectInitializerConstructorDeclared;
+				bSuperClassObjectInitializerConstructorDeclared = SuperClassDef->IsObjectInitializerConstructorDeclared();
 			}
 		}
 		if (bSuperClassObjectInitializerConstructorDeclared)
 		{
 			Out.Logf(TEXT("\t%s_API %s(const FObjectInitializer& ObjectInitializer = FObjectInitializer::Get()) : Super(ObjectInitializer) { };\r\n"), API, ClassCPPName);
-			StructData.bObjectInitializerConstructorDeclared = true;
+			ClassDef.MarkObjectInitializerConstructorDeclared();
 		}
 		else
 		{
 			Out.Logf(TEXT("\t%s_API %s() { };\r\n"), API, ClassCPPName);
-			StructData.bDefaultConstructorDeclared = true;
+			ClassDef.MarkDefaultConstructorDeclared();
 		}
 
-		StructData.bConstructorDeclared = true;
+		ClassDef.MarkConstructorDeclared();
 	}
 	ExportCopyConstructorDefinition(Out, API, ClassCPPName);
 }
@@ -3138,7 +3135,7 @@ void ExportConstructorDefinition(FOutputDevice& Out, FUnrealClassDefinitionInfo&
  */
 void ExportDefaultConstructorCallDefinition(FOutputDevice& Out, FUnrealClassDefinitionInfo& ClassDef, const TCHAR* ClassCPPName)
 {
-	if (ClassDef.GetStructMetaData().bObjectInitializerConstructorDeclared)
+	if (ClassDef.IsObjectInitializerConstructorDeclared())
 	{
 		if (ClassDef.HasAnyClassFlags(CLASS_Abstract))
 		{
@@ -3149,7 +3146,7 @@ void ExportDefaultConstructorCallDefinition(FOutputDevice& Out, FUnrealClassDefi
 			Out.Logf(TEXT("\tDEFINE_DEFAULT_OBJECT_INITIALIZER_CONSTRUCTOR_CALL(%s)\r\n"), ClassCPPName);
 		}
 	}
-	else if (ClassDef.GetStructMetaData().bDefaultConstructorDeclared)
+	else if (ClassDef.IsDefaultConstructorDeclared())
 	{
 		if (ClassDef.HasAnyClassFlags(CLASS_Abstract))
 		{
@@ -3208,7 +3205,7 @@ void FNativeClassHeaderGenerator::ExportConstructorsMacros(FOutputDevice& OutGen
 	ExportStandardConstructorsMacro(StdMacro, ClassDef, APIArg, *ClassCPPName);
 	ExportEnhancedConstructorsMacro(EnhMacro, ClassDef, APIArg, *ClassCPPName);
 
-	if (!ClassDef.GetStructMetaData().bCustomVTableHelperConstructorDeclared)
+	if (!ClassDef.IsCustomVTableHelperConstructorDeclared())
 	{
 		Out.Logf(TEXT("\tDEFINE_VTABLE_PTR_HELPER_CTOR(%s);" LINE_TERMINATOR), *ClassCPPName);
 	}
@@ -4980,7 +4977,6 @@ void FNativeClassHeaderGenerator::ExportNativeFunctions(FOutputDevice& OutGenera
 	FNativeFunctionStringBuilder RuntimeStringBuilders;
 	FNativeFunctionStringBuilder EditorStringBuilders;
 
-	FStructMetaData& StructData = ClassDef.GetStructMetaData();
 	const FString ClassCPPName = ClassDef.GetAlternateNameCPP(ClassDef.HasAnyClassFlags(CLASS_Interface));
 
 	ClassDef.ValidateDefinitionRange();
@@ -5154,7 +5150,7 @@ void FNativeClassHeaderGenerator::ExportNativeFunctions(FOutputDevice& OutGenera
 
 	// static class data
 	{
-		FString MacroName = SourceFile.GetGeneratedMacroName(StructData, TEXT("_SPARSE_DATA"));
+		FString MacroName = SourceFile.GetGeneratedMacroName(ClassDef.GetGeneratedBodyLine(), TEXT("_SPARSE_DATA"));
 
 		WriteMacro(OutGeneratedHeaderText, MacroName, RuntimeStringBuilders.AutogeneratedStaticData + RuntimeStringBuilders.AutogeneratedStaticDataFuncs);
 		OutMacroCalls.Logf(TEXT("\t%s\r\n"), *MacroName);
@@ -5163,7 +5159,7 @@ void FNativeClassHeaderGenerator::ExportNativeFunctions(FOutputDevice& OutGenera
 
 	// Write runtime wrappers
 	{
-		FString MacroName = SourceFile.GetGeneratedMacroName(StructData, TEXT("_RPC_WRAPPERS"));
+		FString MacroName = SourceFile.GetGeneratedMacroName(ClassDef.GetGeneratedBodyLine(), TEXT("_RPC_WRAPPERS"));
 
 		// WriteMacro has an assumption about what will be at the end of this block that is no longer true due to splitting the
 		// definition and implementation, so add on a line terminator to satisfy it
@@ -5176,7 +5172,7 @@ void FNativeClassHeaderGenerator::ExportNativeFunctions(FOutputDevice& OutGenera
 		OutMacroCalls.Logf(TEXT("\t%s\r\n"), *MacroName);
 
 		// Put static checks before RPCWrappers to get proper messages from static asserts before compiler errors.
-		FString NoPureDeclsMacroName = SourceFile.GetGeneratedMacroName(StructData, TEXT("_RPC_WRAPPERS_NO_PURE_DECLS"));
+		FString NoPureDeclsMacroName = SourceFile.GetGeneratedMacroName(ClassDef.GetGeneratedBodyLine(), TEXT("_RPC_WRAPPERS_NO_PURE_DECLS"));
 		if (ClassDef.GetGeneratedCodeVersion() > EGeneratedCodeVersion::V1)
 		{
 			WriteMacro(OutGeneratedHeaderText, NoPureDeclsMacroName, RuntimeStringBuilders.RPCWrappers);
@@ -5196,7 +5192,7 @@ void FNativeClassHeaderGenerator::ExportNativeFunctions(FOutputDevice& OutGenera
 	{
 		OutGeneratedHeaderText.Log( BeginEditorOnlyGuard );
 
-		FString MacroName = SourceFile.GetGeneratedMacroName(StructData, TEXT("_EDITOR_ONLY_RPC_WRAPPERS"));
+		FString MacroName = SourceFile.GetGeneratedMacroName(ClassDef.GetGeneratedBodyLine(), TEXT("_EDITOR_ONLY_RPC_WRAPPERS"));
 
 		// WriteMacro has an assumption about what will be at the end of this block that is no longer true due to splitting the
 		// definition and implementation, so add on a line terminator to satisfy it
@@ -5209,7 +5205,7 @@ void FNativeClassHeaderGenerator::ExportNativeFunctions(FOutputDevice& OutGenera
 		OutMacroCalls.Logf(TEXT("\t%s\r\n"), *MacroName);
 
 		// Put static checks before RPCWrappers to get proper messages from static asserts before compiler errors.
-		FString NoPureDeclsMacroName = SourceFile.GetGeneratedMacroName(StructData, TEXT("_EDITOR_ONLY_RPC_WRAPPERS_NO_PURE_DECLS"));
+		FString NoPureDeclsMacroName = SourceFile.GetGeneratedMacroName(ClassDef.GetGeneratedBodyLine(), TEXT("_EDITOR_ONLY_RPC_WRAPPERS_NO_PURE_DECLS"));
 		if (ClassDef.GetGeneratedCodeVersion() > EGeneratedCodeVersion::V1)
 		{
 			WriteMacro(OutGeneratedHeaderText, NoPureDeclsMacroName, EditorStringBuilders.RPCWrappers);
diff --git a/Engine/Source/Programs/UnrealHeaderTool/Private/HeaderParser.cpp b/Engine/Source/Programs/UnrealHeaderTool/Private/HeaderParser.cpp
--- a/Engine/Source/Programs/UnrealHeaderTool/Private/HeaderParser.cpp
+++ b/Engine/Source/Programs/UnrealHeaderTool/Private/HeaderParser.cpp
@@ -4604,24 +4604,24 @@ bool FHeaderParser::CompileDeclaration(TArray<FUnrealFunctionDefinitionInfo*>& D
 		CompileVersionDeclaration(GetCurrentClassDef());
 		RequireSymbol(TEXT(')'), Token.Identifier);
 
-		FStructMetaData& StructData = GetCurrentClassData();
-		if (StructData.ParsedInterface == EParsedInterface::NotAnInterface)
+		FUnrealClassDefinitionInfo& ClassDef = GetCurrentClassDef();
+		if (ClassDef.GetParsedInterfaceState() == EParsedInterface::NotAnInterface)
 		{
-			FString CurrentClassName = GetCurrentClassDef().GetName();
+			FString CurrentClassName = ClassDef.GetName();
 			Throwf(TEXT("Could not find the associated 'U%s' class while parsing 'I%s' - it could be missing or malformed"), *CurrentClassName, *CurrentClassName);
 		}
 
-		if (StructData.ParsedInterface == EParsedInterface::ParsedIInterface)
+		if (ClassDef.GetParsedInterfaceState() == EParsedInterface::ParsedIInterface)
 		{
-			FString CurrentClassName = GetCurrentClassDef().GetName();
+			FString CurrentClassName = ClassDef.GetName();
 			Throwf(TEXT("Duplicate IInterface definition found while parsing 'I%s'"), *CurrentClassName);
 		}
 
-		check(StructData.ParsedInterface == EParsedInterface::ParsedUInterface);
+		check(ClassDef.GetParsedInterfaceState() == EParsedInterface::ParsedUInterface);
 
-		StructData.ParsedInterface = EParsedInterface::ParsedIInterface;
-		StructData.GeneratedBodyMacroAccessSpecifier = CurrentAccessSpecifier;
-		StructData.SetInterfaceGeneratedBodyLine(InputLine);
+		ClassDef.SetParsedInterfaceState(EParsedInterface::ParsedIInterface);
+		ClassDef.SetGeneratedBodyMacroAccessSpecifier(CurrentAccessSpecifier);
+		ClassDef.SetInterfaceGeneratedBodyLine(InputLine);
 
 		bClassHasGeneratedIInterfaceBody = true;
 
@@ -4647,10 +4647,10 @@ bool FHeaderParser::CompileDeclaration(TArray<FUnrealFunctionDefinitionInfo*>& D
 		CompileVersionDeclaration(GetCurrentClassDef());
 		RequireSymbol(TEXT(')'), Token.Identifier);
 
-		FStructMetaData& StructData = GetCurrentClassData();
+		FUnrealClassDefinitionInfo& ClassDef = GetCurrentClassDef();
 
-		StructData.GeneratedBodyMacroAccessSpecifier = CurrentAccessSpecifier;
-		StructData.SetGeneratedBodyLine(InputLine);
+		ClassDef.SetGeneratedBodyMacroAccessSpecifier(CurrentAccessSpecifier);
+		ClassDef.SetGeneratedBodyLine(InputLine);
 
 		bClassHasGeneratedUInterfaceBody = true;
 
@@ -4668,13 +4668,13 @@ bool FHeaderParser::CompileDeclaration(TArray<FUnrealFunctionDefinitionInfo*>& D
 			Throwf(TEXT("%s must occur inside the class definition"), Token.Identifier);
 		}
 
-		FStructMetaData& StructData = GetCurrentClassData();
+		FUnrealClassDefinitionInfo& ClassDef = GetCurrentClassDef();
 
 		if (Token.Matches(TEXT("GENERATED_BODY"), ESearchCase::CaseSensitive))
 		{
 			GetCurrentClassDef().MarkGeneratedBody();
 
-			StructData.GeneratedBodyMacroAccessSpecifier = CurrentAccessSpecifier;
+			ClassDef.SetGeneratedBodyMacroAccessSpecifier(CurrentAccessSpecifier);
 		}
 		else
 		{
@@ -4685,7 +4685,7 @@ bool FHeaderParser::CompileDeclaration(TArray<FUnrealFunctionDefinitionInfo*>& D
 		CompileVersionDeclaration(GetCurrentClassDef());
 		RequireSymbol(TEXT(')'), Token.Identifier);
 
-		StructData.SetGeneratedBodyLine(InputLine);
+		ClassDef.SetGeneratedBodyLine(InputLine);
 
 		bClassHasGeneratedBody = true;
 		return true;
@@ -4759,7 +4759,8 @@ bool FHeaderParser::CompileDeclaration(TArray<FUnrealFunctionDefinitionInfo*>& D
 
 	if (bEncounteredNewStyleClass_UnmatchedBrackets && Token.Matches(TEXT('}')))
 	{
-		GetCurrentClassDef().GetDefinitionRange().End = &Input[InputPos];
+		FUnrealClassDefinitionInfo& CurrentClassDef = GetCurrentClassDef();
+		CurrentClassDef.GetDefinitionRange().End = &Input[InputPos];
 		MatchSemi();
 
 		// Closing brace for class declaration
@@ -4767,8 +4768,6 @@ bool FHeaderParser::CompileDeclaration(TArray<FUnrealFunctionDefinitionInfo*>& D
 		// Instead, the whole statement-consumer loop should be in a nest
 		bEncounteredNewStyleClass_UnmatchedBrackets = false;
 
-		FUnrealClassDefinitionInfo& CurrentClassDef = UHTCastChecked<FUnrealClassDefinitionInfo>(GetCurrentClassDef());
-
 		// Pop nesting here to allow other non UClass declarations in the header file.
 		if (CurrentClassDef.HasAllClassFlags(CLASS_Interface))
 		{
@@ -4951,7 +4950,7 @@ bool FHeaderParser::CompileDeclaration(TArray<FUnrealFunctionDefinitionInfo*>& D
 						{
 							FString EnclosingDefine = CurrentCompilerDirective != 0 ? TEXT("WITH_EDITORONLY_DATA") : TEXT("");
 
-							FUnrealClassDefinitionInfo& ClassDef = UHTCastChecked<FUnrealClassDefinitionInfo>(GetCurrentClassDef());
+							FUnrealClassDefinitionInfo& ClassDef = GetCurrentClassDef();
 							ClassDef.AddArchiveType(ArchiveType);
 							ClassDef.SetEnclosingDefine(MoveTemp(EnclosingDefine));
 						}
@@ -5241,9 +5240,7 @@ FUnrealClassDefinitionInfo& FHeaderParser::CompileClassDeclaration()
 	ClassDef.MergeClassCategories();
 
 	// Class attributes.
-	FStructMetaData& StructData = ClassDef.GetStructMetaData();
-	StructData.SetPrologLine(PrologFinishLine);
-
+	ClassDef.SetPrologLine(PrologFinishLine);
 	ClassDef.MergeAndValidateClassFlags(DeclaredClassName);
 	ClassDef.SetInternalFlags(EInternalObjectFlags::Native);
 
@@ -5443,8 +5440,7 @@ void FHeaderParser::CompileInterfaceDeclaration()
 	}
 
 	// Try parsing metadata for the interface
-	FStructMetaData& StructData = InterfaceClassDef->GetStructMetaData();
-	StructData.SetPrologLine(PrologFinishLine);
+	InterfaceClassDef->SetPrologLine(PrologFinishLine);
 
 	// Register the metadata
 	AddModuleRelativePathToMetadata(*InterfaceClassDef, MetaData);
@@ -6099,7 +6095,7 @@ FUnrealFunctionDefinitionInfo& FHeaderParser::CompileFunctionDeclaration()
 	TArray<FPropertySpecifier> SpecifiersFound;
 	ReadSpecifierSetInsideMacro(SpecifiersFound, TEXT("Function"), MetaData);
 
-	FUnrealClassDefinitionInfo& OuterClassDef = UHTCastChecked<FUnrealClassDefinitionInfo>(GetCurrentClassDef());
+	FUnrealClassDefinitionInfo& OuterClassDef = GetCurrentClassDef();
 	if (!OuterClassDef.HasAnyClassFlags(CLASS_Native))
 	{
 		Throwf(TEXT("Should only be here for native classes!"));
@@ -6508,7 +6504,7 @@ FUnrealFunctionDefinitionInfo& FHeaderParser::CompileFunctionDeclaration()
 		// This is a final (prebinding, non-overridable) function
 		FuncDef.SetFunctionFlags(FUNC_Final);
 		FuncDef.SetFunctionExportFlags(FUNCEXPORT_Final);
-		if (UHTCastChecked<FUnrealClassDefinitionInfo>(GetCurrentClassDef()).HasAnyClassFlags(CLASS_Interface))
+		if (GetCurrentClassDef().HasAnyClassFlags(CLASS_Interface))
 		{
 			Throwf(TEXT("Interface functions cannot be declared 'final'"));
 		}
@@ -6565,7 +6561,7 @@ FUnrealFunctionDefinitionInfo& FHeaderParser::CompileFunctionDeclaration()
 	}
 
 	// perform documentation policy tests
-	CheckDocumentationPolicyForFunc(UHTCastChecked<FUnrealClassDefinitionInfo>(GetCurrentClassDef()), FuncDef);
+	CheckDocumentationPolicyForFunc(GetCurrentClassDef(), FuncDef);
 
 	return FuncDef;
 }
@@ -8125,18 +8121,17 @@ bool FHeaderParser::TryToMatchConstructorParameterList(FToken Token)
 		return false;
 	}
 
-	FStructMetaData& StructData = GetCurrentClassData();
+	FUnrealClassDefinitionInfo& ClassDef = GetCurrentClassDef();
 
 	bool bOICtor = false;
 	bool bVTCtor = false;
 
-	if (!StructData.bDefaultConstructorDeclared && MatchSymbol(TEXT(')')))
+	if (!ClassDef.IsDefaultConstructorDeclared() && MatchSymbol(TEXT(')')))
 	{
-		StructData.bDefaultConstructorDeclared = true;
+		ClassDef.MarkDefaultConstructorDeclared();
 	}
-	else if (!StructData.bObjectInitializerConstructorDeclared
-		|| !StructData.bCustomVTableHelperConstructorDeclared
-	)
+	else if (!ClassDef.IsObjectInitializerConstructorDeclared()
+		|| !ClassDef.IsCustomVTableHelperConstructorDeclared())
 	{
 		FToken ObjectInitializerParamParsingToken;
 
@@ -8207,11 +8202,20 @@ bool FHeaderParser::TryToMatchConstructorParameterList(FToken Token)
 			}
 		}
 
-		StructData.bObjectInitializerConstructorDeclared = StructData.bObjectInitializerConstructorDeclared || (bOICtor && bIsRef && bIsConst);
-		StructData.bCustomVTableHelperConstructorDeclared = StructData.bCustomVTableHelperConstructorDeclared || (bVTCtor && bIsRef);
+		if (bOICtor && bIsRef && bIsConst)
+		{
+			ClassDef.MarkObjectInitializerConstructorDeclared();
+		}
+		if (bVTCtor && bIsRef)
+		{
+			ClassDef.MarkCustomVTableHelperConstructorDeclared();
+		}
 	}
 
-	StructData.bConstructorDeclared = StructData.bConstructorDeclared || !bVTCtor;
+	if (!bVTCtor)
+	{
+		ClassDef.MarkConstructorDeclared();
+	}
 
 	// Optionally match semicolon.
 	if (!MatchSymbol(TEXT(';')))
@@ -8258,7 +8262,7 @@ void FHeaderParser::CompileVersionDeclaration(FUnrealStructDefinitionInfo& Struc
 
 void FHeaderParser::ResetClassData()
 {
-	FUnrealClassDefinitionInfo& CurrentClassDef = UHTCastChecked<FUnrealClassDefinitionInfo>(GetCurrentClassDef());
+	FUnrealClassDefinitionInfo& CurrentClassDef = GetCurrentClassDef();
 
 	// Set class flags and within.
 	CurrentClassDef.ClearClassFlags(CLASS_RecompilerClear);
diff --git a/Engine/Source/Programs/UnrealHeaderTool/Private/HeaderParser.h b/Engine/Source/Programs/UnrealHeaderTool/Private/HeaderParser.h
--- a/Engine/Source/Programs/UnrealHeaderTool/Private/HeaderParser.h
+++ b/Engine/Source/Programs/UnrealHeaderTool/Private/HeaderParser.h
@@ -180,35 +180,6 @@ protected:
 		return TopNest->GetScope();
 	}
 
-	/**
-	 * Gets current file scope.
-	 */
-	FFileScope* GetCurrentFileScope() const
-	{
-		int32 Index = 0;
-		if (!TopNest)
-		{
-			check(!NestLevel);
-			return nullptr;
-		}
-		while (TopNest[Index].NestType != ENestType::GlobalScope)
-		{
-			--Index;
-		}
-
-		return (FFileScope*)TopNest[Index].GetScope();
-	}
-
-	/**
-	 * Gets current class scope.
-	 */
-	FStructScope* GetCurrentClassScope() const
-	{
-		check(TopNest->NestType == ENestType::Class || TopNest->NestType == ENestType::Interface || TopNest->NestType == ENestType::NativeInterface);
-
-		return (FStructScope*)TopNest->GetScope();
-	}
-
 	/**
 	 * Tells if parser is currently in a class.
 	 */
@@ -231,17 +202,11 @@ protected:
 	/**
 	 * Gets current class definition.
 	 */
-	FUnrealStructDefinitionInfo& GetCurrentClassDef() const
+	FUnrealClassDefinitionInfo& GetCurrentClassDef() const
 	{
-		return GetCurrentClassScope()->GetStructDef();
-	}
+		check(TopNest->NestType == ENestType::Class || TopNest->NestType == ENestType::Interface || TopNest->NestType == ENestType::NativeInterface);
 
-	/**
-	 * Gets current class's metadata.
-	 */
-	FStructMetaData& GetCurrentClassData()
-	{
-		return GetCurrentClassDef().GetStructMetaData();
+		return UHTCastChecked<FUnrealClassDefinitionInfo>(static_cast<FStructScope*>(TopNest->GetScope())->GetStructDef());
 	}
 
 	// Information about all nesting levels.
diff --git a/Engine/Source/Programs/UnrealHeaderTool/Private/ParserHelper.h b/Engine/Source/Programs/UnrealHeaderTool/Private/ParserHelper.h
--- a/Engine/Source/Programs/UnrealHeaderTool/Private/ParserHelper.h
+++ b/Engine/Source/Programs/UnrealHeaderTool/Private/ParserHelper.h
@@ -662,114 +662,6 @@ struct FFuncInfo
 	void SetFunctionNames(FUnrealFunctionDefinitionInfo& FunctionDef);
 };
 
-enum class EParsedInterface
-{
-	NotAnInterface,
-	ParsedUInterface,
-	ParsedIInterface
-};
-
-/**
- * Class for storing compiler metadata about a class's properties.
- */
-class FStructMetaData
-{
-	/** The line of UCLASS/UINTERFACE macro in this class. */
-	int32 PrologLine;
-
-	/** The line of GENERATED_BODY/GENERATED_UCLASS_BODY macro in this class. */
-	int32 GeneratedBodyLine;
-
-	/** Same as above, but for interface class associated with this class. */
-	int32 InterfaceGeneratedBodyLine;
-
-public:
-	/** Default constructor */
-	FStructMetaData()
-		: PrologLine(-1)
-		, GeneratedBodyLine(-1)
-		, InterfaceGeneratedBodyLine(-1)
-		, bConstructorDeclared(false)
-		, bDefaultConstructorDeclared(false)
-		, bObjectInitializerConstructorDeclared(false)
-		, bCustomVTableHelperConstructorDeclared(false)
-		, GeneratedBodyMacroAccessSpecifier(ACCESS_NotAnAccessSpecifier)
-	{
-	}
-
-	/**
-	 * Gets prolog line number for this class.
-	 */
-	int32 GetPrologLine() const
-	{
-		check(PrologLine > 0);
-		return PrologLine;
-	}
-
-	/**
-	 * Gets generated body line number for this class.
-	 */
-	int32 GetGeneratedBodyLine() const
-	{
-		check(GeneratedBodyLine > 0);
-		return GeneratedBodyLine;
-	}
-
-	/**
-	 * Gets interface generated body line number for this class.
-	 */
-	int32 GetInterfaceGeneratedBodyLine() const
-	{
-		check(InterfaceGeneratedBodyLine > 0);
-		return InterfaceGeneratedBodyLine;
-	}
-
-	/**
-	 * Sets prolog line number for this class.
-	 */
-	void SetPrologLine(int32 Line)
-	{
-		check(Line > 0);
-		PrologLine = Line;
-	}
-
-	/**
-	 * Sets generated body line number for this class.
-	 */
-	void SetGeneratedBodyLine(int32 Line)
-	{
-		check(Line > 0);
-		GeneratedBodyLine = Line;
-	}
-
-	/**
-	 * Sets interface generated body line number for this class.
-	 */
-	void SetInterfaceGeneratedBodyLine(int32 Line)
-	{
-		check(Line > 0);
-		InterfaceGeneratedBodyLine = Line;
-	}
-
-	// Is constructor declared?
-	bool bConstructorDeclared;
-
-	// Is default constructor declared?
-	bool bDefaultConstructorDeclared;
-
-	// Is ObjectInitializer constructor (i.e. a constructor with only one parameter of type FObjectInitializer) declared?
-	bool bObjectInitializerConstructorDeclared;
-
-	// Is custom VTable helper constructor declared?
-	bool bCustomVTableHelperConstructorDeclared;
-
-	// GENERATED_BODY access specifier to preserve.
-	EAccessSpecifier GeneratedBodyMacroAccessSpecifier;
-
-	/** Parsed interface state */
-	EParsedInterface ParsedInterface = EParsedInterface::NotAnInterface;
-};
-
 /*-----------------------------------------------------------------------------
 	Retry points.
 -----------------------------------------------------------------------------*/
diff --git a/Engine/Source/Programs/UnrealHeaderTool/Private/Scope.h b/Engine/Source/Programs/UnrealHeaderTool/Private/Scope.h
--- a/Engine/Source/Programs/UnrealHeaderTool/Private/Scope.h
+++ b/Engine/Source/Programs/UnrealHeaderTool/Private/Scope.h
@@ -12,7 +12,6 @@ class UDelegateFunction;
 
 // Forward declarations.
 class UStruct;
-class FStructMetaData;
 class FUnrealSourceFile;
 class FFileScope;
 class FStructScope;
diff --git a/Engine/Source/Programs/UnrealHeaderTool/Private/UnrealSourceFile.cpp b/Engine/Source/Programs/UnrealHeaderTool/Private/UnrealSourceFile.cpp
--- a/Engine/Source/Programs/UnrealHeaderTool/Private/UnrealSourceFile.cpp
+++ b/Engine/Source/Programs/UnrealHeaderTool/Private/UnrealSourceFile.cpp
@@ -87,11 +87,6 @@ const FString& FUnrealSourceFile::GetStrippedFilename() const
 	return StrippedFilename;
 }
 
-FString FUnrealSourceFile::GetGeneratedMacroName(FStructMetaData& StructData, const TCHAR* Suffix) const
-{
-	return GetGeneratedMacroName(StructData.GetGeneratedBodyLine(), Suffix);
-}
-
 FString FUnrealSourceFile::GetGeneratedMacroName(int32 LineNumber, const TCHAR* Suffix) const
 {
 	if (Suffix != nullptr)
diff --git a/Engine/Source/Programs/UnrealHeaderTool/Private/UnrealSourceFile.h b/Engine/Source/Programs/UnrealHeaderTool/Private/UnrealSourceFile.h
--- a/Engine/Source/Programs/UnrealHeaderTool/Private/UnrealSourceFile.h
+++ b/Engine/Source/Programs/UnrealHeaderTool/Private/UnrealSourceFile.h
@@ -11,7 +11,6 @@
 
 class UPackage;
 class FArchive;
-class FStructMetaData;
 
 enum class ETopologicalState : uint8
 {
@@ -240,15 +239,7 @@ public:
 	 * @param LineNumber Number at which generated body macro is.
 	 * @param Suffix Suffix to add to generated body macro name.
 	 */
-	FString GetGeneratedMacroName(int32 LineNumber, const TCHAR* Suffix) const;
-
-	/**
-	 * Gets file-wise generated body macro name.
-	 *
-	 * @param StructData Struct metadata for which to get generated body macro name.
-	 * @param Suffix Suffix to add to generated body macro name.
-	 */
-	FString GetGeneratedMacroName(FStructMetaData& StructData, const TCHAR* Suffix = nullptr) const;
+	FString GetGeneratedMacroName(int32 LineNumber, const TCHAR* Suffix = nullptr) const;
 
 	/**
 	 * Gets scope for this file.
diff --git a/Engine/Source/Programs/UnrealHeaderTool/Private/UnrealTypeDefinitionInfo.cpp b/Engine/Source/Programs/UnrealHeaderTool/Private/UnrealTypeDefinitionInfo.cpp
--- a/Engine/Source/Programs/UnrealHeaderTool/Private/UnrealTypeDefinitionInfo.cpp
+++ b/Engine/Source/Programs/UnrealHeaderTool/Private/UnrealTypeDefinitionInfo.cpp
@@ -1833,7 +1833,7 @@ FUnrealClassDefinitionInfo::FUnrealClassDefinitionInfo(FUnrealSourceFile& InSour
 	if (bInIsInterface)
 	{
 		ClassFlags |= CLASS_Interface;
-		GetStructMetaData().ParsedInterface = EParsedInterface::ParsedUInterface;
+		ParsedInterface = EParsedInterface::ParsedUInterface;
 	}
 }
 
@@ -2045,7 +2045,7 @@ void FUnrealClassDefinitionInfo::PostParseFinalizeInternal(EPostParseFinalizePha
 			}
 		}
 
-		if (IsInterface() && GetStructMetaData().ParsedInterface == EParsedInterface::ParsedUInterface)
+		if (IsInterface() && ParsedInterface == EParsedInterface::ParsedUInterface)
 		{
 			FString UName = GetNameCPP();
 			FString IName = TEXT("I") + UName.RightChop(1);
diff --git a/Engine/Source/Programs/UnrealHeaderTool/Private/UnrealTypeDefinitionInfo.h b/Engine/Source/Programs/UnrealHeaderTool/Private/UnrealTypeDefinitionInfo.h
--- a/Engine/Source/Programs/UnrealHeaderTool/Private/UnrealTypeDefinitionInfo.h
+++ b/Engine/Source/Programs/UnrealHeaderTool/Private/UnrealTypeDefinitionInfo.h
@@ -71,6 +71,13 @@ enum class EFunctionType
 	SparseDelegate,
 };
 
+enum class EParsedInterface
+{
+	NotAnInterface,
+	ParsedUInterface,
+	ParsedIInterface
+};
+
 enum class ECreateEngineTypesPhase : uint8
 {
 	Phase1,				// Create the Class, ScriptStruct, and Enum types
@@ -2155,14 +2162,6 @@ public:
 		return Functions;
 	}
 
-	/**
-	 * Return the class meta data information
-	 */
-	FStructMetaData& GetStructMetaData()
-	{
-		return StructMetaData;
-	}
-
 	/**
 	 * Return the super struct
 	 */
@@ -2205,14 +2204,6 @@ public:
 		return bContainsDelegates;
 	}
 
-	/**
-	 * Sets contains delegates flag for this class.
-	 */
-	void MarkContainsDelegate()
-	{
-		bContainsDelegates = true;
-	}
-
 	const FString* FindMetaDataHierarchical(const FName& Key) const
 	{
 		if (Key == NAME_None)
@@ -2315,7 +2306,6 @@ private:
 	/** Functions of the structure */
 	TArray<TSharedRef<FUnrealFunctionDefinitionInfo>> Functions;
 
-	FStructMetaData StructMetaData;
 	FBaseStructInfo SuperStructInfo;
 	TArray<FBaseStructInfo> BaseStructInfos;
 
@@ -3121,6 +3111,156 @@ public:
 		}
 	}
 
+	/**
+	 * Gets prolog line number for this class.
+	 */
+	int32 GetPrologLine() const
+	{
+		check(PrologLine > 0);
+		return PrologLine;
+	}
+
+	/**
+	 * Sets prolog line number for this class.
+	 */
+	void SetPrologLine(int32 Line)
+	{
+		check(Line > 0);
+		PrologLine = Line;
+	}
+
+	/**
+	 * Gets generated body line number for this class.
+	 */
+	int32 GetGeneratedBodyLine() const
+	{
+		check(GeneratedBodyLine > 0);
+		return GeneratedBodyLine;
+	}
+
+	/**
+	 * Sets generated body line number for this class.
+	 */
+	void SetGeneratedBodyLine(int32 Line)
+	{
+		check(Line > 0);
+		GeneratedBodyLine = Line;
+	}
+
+	/**
+	 * Gets interface generated body line number for this class.
+	 */
+	int32 GetInterfaceGeneratedBodyLine() const
+	{
+		check(InterfaceGeneratedBodyLine > 0);
+		return InterfaceGeneratedBodyLine;
+	}
+
+	/**
+	 * Sets interface generated body line number for this class.
+	 */
+	void SetInterfaceGeneratedBodyLine(int32 Line)
+	{
+		check(Line > 0);
+		InterfaceGeneratedBodyLine = Line;
+	}
+
+	/**
+	 * Return true if the constructor is declared
+	 */
+	bool IsConstructorDeclared() const
+	{
+		return bConstructorDeclared;
+	}
+
+	/**
+	 * Mark that the constructor has been declared
+	 */
+	void MarkConstructorDeclared()
+	{
+		bConstructorDeclared = true;
+	}
+
+	/**
+	 * Return true if the default constructor is declared
+	 */
+	bool IsDefaultConstructorDeclared() const
+	{
+		return bDefaultConstructorDeclared;
+	}
+
+	/**
+	 * Mark that the default constructor has been declared
+	 */
+	void MarkDefaultConstructorDeclared()
+	{
+		bDefaultConstructorDeclared = true;
+	}
+
+	/**
+	 * Return true if the object initializer constructor is declared
+	 */
+	bool IsObjectInitializerConstructorDeclared() const
+	{
+		return bObjectInitializerConstructorDeclared;
+	}
+
+	/**
+	 * Mark that the object initializer constructor has been declared
+	 */
+	void MarkObjectInitializerConstructorDeclared()
+	{
+		bObjectInitializerConstructorDeclared = true;
+	}
+
+	/**
+	 * Return true if the custom vtable helper constructor is declared
+	 */
+	bool IsCustomVTableHelperConstructorDeclared() const
+	{
+		return bCustomVTableHelperConstructorDeclared;
+	}
+
+	/**
+	 * Mark that the custom vtable helper constructor has been declared
+	 */
+	void MarkCustomVTableHelperConstructorDeclared()
+	{
+		bCustomVTableHelperConstructorDeclared = true;
+	}
+
+	/**
+	 * Get the generated body access specifier
+	 */
+	EAccessSpecifier GetGeneratedBodyMacroAccessSpecifier() const
+	{
+		return GeneratedBodyMacroAccessSpecifier;
+	}
+
+	/**
+	 * Set the generated body access specifier
+	 */
+	EAccessSpecifier SetGeneratedBodyMacroAccessSpecifier(EAccessSpecifier InGeneratedBodyMacroAccessSpecifier)
+	{
+		return GeneratedBodyMacroAccessSpecifier = InGeneratedBodyMacroAccessSpecifier;
+	}
+
+	/**
+	 * Get the parsed interface state 
+	 */
+	EParsedInterface GetParsedInterfaceState() const
+	{
+		return ParsedInterface;
+	}
+
+	/**
+	 * Set the parsed interface state
+	 */
+	void SetParsedInterfaceState(EParsedInterface InParsedInterface)
+	{
+		ParsedInterface = InParsedInterface;
+	}
+
 	/**
 	 * Test to see if the class implements a specific interface
 	 */
@@ -3176,8 +3316,37 @@ private:
 	EClassCastFlags ClassCastFlags = CASTCLASS_None;
 	FUnrealClassDefinitionInfo* ClassWithin = nullptr;
 	ESerializerArchiveType ArchiveType = ESerializerArchiveType::None;
+
+	/** GENERATED_BODY access specifier to preserve. */
+	EAccessSpecifier GeneratedBodyMacroAccessSpecifier = ACCESS_NotAnAccessSpecifier;
+
+	/** Parsed interface state */
+	EParsedInterface ParsedInterface = EParsedInterface::NotAnInterface;
+
+	/** The line of UCLASS/UINTERFACE macro in this class. */
+	int32 PrologLine = -1;
+
+	/** The line of GENERATED_BODY/GENERATED_UCLASS_BODY macro in this class. */
+	int32 GeneratedBodyLine = -1;
+
+	/** Same as above, but for interface class associated with this class. */
+	int32 InterfaceGeneratedBodyLine = -1;
+
 	int32 FirstOwnedClassRep = 0;
 	int32 PropertiesSize = 0;
+
+	/** Is constructor declared? */
+	bool bConstructorDeclared = false;
+
+	/** Is default constructor declared? */
+	bool bDefaultConstructorDeclared = false;
+
+	/* Is ObjectInitializer constructor (i.e. a constructor with only one parameter of type FObjectInitializer) declared? */
+	bool bObjectInitializerConstructorDeclared = false;
+
+	/* Is custom VTable helper constructor declared? */
+	bool bCustomVTableHelperConstructorDeclared = false;
+
 	bool bIsInterface = false;
 	bool bWantsToBePlaceable = false;
 };
