Fix issue where /* /* */ comment would cause UHT to crash.

diff --git a/Engine/Source/Programs/UnrealHeaderTool/Private/HeaderParser.cpp b/Engine/Source/Programs/UnrealHeaderTool/Private/HeaderParser.cpp
--- a/Engine/Source/Programs/UnrealHeaderTool/Private/HeaderParser.cpp
+++ b/Engine/Source/Programs/UnrealHeaderTool/Private/HeaderParser.cpp
@@ -7693,7 +7693,7 @@ void FHeaderParser::SimplifiedClassParse(FUnrealSourceFile& SourceFile, const TC
 
 	// now start over go look for the class
 
-	int32 CommentDim  = 0;
+	bool bInComment = false;
 	CurrentLine = 0;
 	Buffer      = *ClassHeaderTextStrippedOfCppText;
 
@@ -7705,7 +7705,7 @@ void FHeaderParser::SimplifiedClassParse(FUnrealSourceFile& SourceFile, const TC
 		CurrentLine++;
 
 		const TCHAR* Str = *StrLine;
-		bool bProcess = CommentDim <= 0;	// for skipping nested multi-line comments
+		bool bProcess = !bInComment;	// for skipping nested multi-line comments
 
 		int32 BraceCount = 0;
 		if( bProcess && FParse::Command(&Str,TEXT("#if")) )
@@ -7816,14 +7816,15 @@ void FHeaderParser::SimplifiedClassParse(FUnrealSourceFile& SourceFile, const TC
 							{
 								StrLine = StrLine.Left(Pos) + StrLine.Mid(EndPos + 2);
 								EndPos = INDEX_NONE;
+								bInComment = false;
 							}
 							else
 							{
 								StrLine.LeftInline(Pos, false);
-								CommentDim++;
+								bInComment = true;
 							}
 						}
-						bProcess = CommentDim <= 1;
+						bProcess = !bInComment;
 					}
 
 					if (EndPos >= 0)
@@ -7831,10 +7832,9 @@ void FHeaderParser::SimplifiedClassParse(FUnrealSourceFile& SourceFile, const TC
 						if (StrBegin == INDEX_NONE || EndPos < StrBegin || EndPos > StrEnd)
 						{
 							StrLine.MidInline(EndPos + 2, MAX_int32, false);
-							CommentDim--;
+							bInComment = false;
 						}
-
-						bProcess = CommentDim <= 0;
+						bProcess = !bInComment;
 					}
 				}
 			}
