Fix issue where USTRUCTs and UDELEGATES don't generate an error if the H file doesn't include the generated H file.  There us currently no requirement for UENUM.

diff --git a/Engine/Source/Programs/UnrealHeaderTool/Private/HeaderParser.cpp b/Engine/Source/Programs/UnrealHeaderTool/Private/HeaderParser.cpp
--- a/Engine/Source/Programs/UnrealHeaderTool/Private/HeaderParser.cpp
+++ b/Engine/Source/Programs/UnrealHeaderTool/Private/HeaderParser.cpp
@@ -1693,7 +1693,6 @@ FUnrealScriptStructDefinitionInfo& FHeaderParser::CompileStructDeclaration()
 	bool IsNative = false;
 	bool IsExport = false;
 	bool IsTransient = false;
-	int32 StructFlags = STRUCT_Native;
 	TMap<FName, FString> MetaData;
 
 	// Get the struct specifier list
@@ -1717,55 +1716,19 @@ FUnrealScriptStructDefinitionInfo& FHeaderParser::CompileStructDeclaration()
 	// Read the struct name
 	ParseNameWithPotentialAPIMacroPrefix(/*out*/ StructNameInScript, /*out*/ RequiredAPIMacroIfPresent, TEXT("struct"));
 
-	// Record that this struct is RequiredAPI if the CORE_API style macro was present
-	if (!RequiredAPIMacroIfPresent.IsEmpty())
-	{
-		StructFlags |= STRUCT_RequiredAPI;
-	}
-
 	StructNameStripped = GetClassNameWithPrefixRemoved(StructNameInScript);
 
 	// Effective struct name
 	const FString EffectiveStructName = *StructNameStripped;
 
-	// Create.
+	// Locate the structure
 	FUnrealScriptStructDefinitionInfo& StructDef = GTypeDefinitionInfoMap.FindByNameChecked<FUnrealScriptStructDefinitionInfo>(*StructNameStripped);
 
-	// Process the list of specifiers
-	for (const FPropertySpecifier& Specifier : SpecifiersFound)
+	if (StructDef.HasAnyStructFlags(STRUCT_Immutable))
 	{
-		switch ((EStructSpecifier)Algo::FindSortedStringCaseInsensitive(*Specifier.Key, GStructSpecifierStrings))
+		if (!FPaths::IsSamePath(Filename, GUObjectDef->GetUnrealSourceFile().GetFilename()))
 		{
-			default:
-			{
-				Throwf(TEXT("Unknown struct specifier '%s'"), *Specifier.Key);
-			}
-			break;
-
-			case EStructSpecifier::NoExport:
-			{
-				//UE_LOG_WARNING_UHT(TEXT("Struct named %s in %s is still marked noexport"), *EffectiveStructName, *(Class->GetName()));//@TODO: UCREMOVAL: Debug printing
-				StructFlags &= ~STRUCT_Native;
-				StructFlags |= STRUCT_NoExport;
-			}
-			break;
-
-			case EStructSpecifier::Atomic:
-			{
-				StructFlags |= STRUCT_Atomic;
-			}
-			break;
-
-			case EStructSpecifier::Immutable:
-			{
-				StructFlags |= STRUCT_Immutable | STRUCT_Atomic;
-
-				if (!FPaths::IsSamePath(Filename, GUObjectDef->GetUnrealSourceFile().GetFilename()))
-				{
-					LogError(TEXT("Immutable is being phased out in favor of SerializeNative, and is only legal on the mirror structs declared in UObject"));
-				}
-			}
-			break;
+			LogError(TEXT("Immutable is being phased out in favor of SerializeNative, and is only legal on the mirror structs declared in UObject"));
 		}
 	}
 
@@ -1830,7 +1793,7 @@ FUnrealScriptStructDefinitionInfo& FHeaderParser::CompileStructDeclaration()
 		}
 
 		StructDef.GetSuperStructInfo().Struct = BaseStructDef->AsStruct();
-		StructFlags |= (BaseStructDef->GetStructFlags() & STRUCT_Inherit);
+		StructDef.SetStructFlags(EStructFlags(BaseStructDef->GetStructFlags() & STRUCT_Inherit));
 		if (StructDef.GetScriptStructSafe() != nullptr && BaseStructDef->GetScriptStructSafe())
 		{
 			StructDef.GetScriptStructSafe()->SetSuperStruct(BaseStructDef->GetScriptStructSafe());
@@ -1861,8 +1824,6 @@ FUnrealScriptStructDefinitionInfo& FHeaderParser::CompileStructDeclaration()
 		Throwf(TEXT("Struct '%s' is missing a valid Unreal prefix, expecting '%s'"), *StructNameInScript, *ExpectedStructName);
 	}
 
-	StructDef.SetStructFlags(EStructFlags(StructFlags));
-
 	AddFormattedPrevCommentAsTooltipMetaData(MetaData);
 
 	// Register the metadata
@@ -2039,8 +2000,7 @@ FUnrealScriptStructDefinitionInfo& FHeaderParser::CompileStructDeclaration()
 
 	// Validation
 	bool bStructBodyFound = StructDef.GetMacroDeclaredLineNumber() != INDEX_NONE;
-	bool bExported        = !!(StructFlags & STRUCT_Native);
-	if (!bStructBodyFound && bExported)
+	if (!bStructBodyFound && StructDef.HasAnyStructFlags(STRUCT_Native))
 	{
 		// Roll the line number back to the start of the struct body and error out
 		InputLine = SavedLineNumber;
@@ -7697,8 +7657,8 @@ void FHeaderParser::SimplifiedClassParse(FUnrealSourceFile& SourceFile, const TC
 	CurrentLine = 0;
 	Buffer      = *ClassHeaderTextStrippedOfCppText;
 
-	bool         bFoundGeneratedInclude = false;
-	bool         bFoundExportedClasses  = false;
+	bool bFoundGeneratedInclude = false;
+	bool bGeneratedIncludeRequired = false;
 
 	for (const TCHAR* StartOfLine = Buffer; FParse::Line(&Buffer, StrLine, true); StartOfLine = Buffer)
 	{
@@ -7858,7 +7818,7 @@ void FHeaderParser::SimplifiedClassParse(FUnrealSourceFile& SourceFile, const TC
 					}
 
 					TSharedRef<FUnrealTypeDefinitionInfo> ClassDecl = Parser.ParseClassDeclaration(StartOfLine + (UInterfaceMacroDecl - Str), CurrentLine, true, TEXT("UINTERFACE"));
-					bFoundExportedClasses |= !ClassDecl->AsClassChecked().HasAnyClassFlags(CLASS_NoExport);
+					bGeneratedIncludeRequired |= !ClassDecl->AsClassChecked().HasAnyClassFlags(CLASS_NoExport);
 					SourceFile.AddDefinedClass(MoveTemp(ClassDecl));
 				}
 			}
@@ -7873,7 +7833,7 @@ void FHeaderParser::SimplifiedClassParse(FUnrealSourceFile& SourceFile, const TC
 					}
 
 					TSharedRef<FUnrealTypeDefinitionInfo> ClassDecl = Parser.ParseClassDeclaration(StartOfLine + (UClassMacroDecl - Str), CurrentLine, false, TEXT("UCLASS"));
-					bFoundExportedClasses |= !ClassDecl->AsClassChecked().HasAnyClassFlags(CLASS_NoExport);
+					bGeneratedIncludeRequired |= !ClassDecl->AsClassChecked().HasAnyClassFlags(CLASS_NoExport);
 					SourceFile.AddDefinedClass(MoveTemp(ClassDecl));
 				}
 			}
@@ -7892,6 +7852,19 @@ void FHeaderParser::SimplifiedClassParse(FUnrealSourceFile& SourceFile, const TC
 				}
 			}
 
+			else if (const TCHAR* UDelegateMacroDecl = FCString::Strfind(Str, TEXT("UDELEGATE")))
+			{
+				if (UDelegateMacroDecl == FCString::Strspn(Str, TEXT("\t ")) + Str)
+				{
+					if (UDelegateMacroDecl[9] != TEXT('('))
+					{
+						FUHTMessage(SourceFile, CurrentLine).Throwf(TEXT("Missing open parenthesis after UDELEGATE"));
+					}
+					// We don't preparse the delegates, but we need to make sure the include is there
+					bGeneratedIncludeRequired = true;
+				}
+			}
+
 			else if (const TCHAR* UStructMacroDecl = FCString::Strfind(Str, TEXT("USTRUCT")))
 			{
 				if (UStructMacroDecl == FCString::Strspn(Str, TEXT("\t ")) + Str)
@@ -7902,13 +7875,14 @@ void FHeaderParser::SimplifiedClassParse(FUnrealSourceFile& SourceFile, const TC
 					}
 
 					TSharedRef<FUnrealTypeDefinitionInfo> StructDecl = Parser.ParseStructDeclaration(StartOfLine + (UStructMacroDecl - Str), CurrentLine);
+					bGeneratedIncludeRequired |= !StructDecl->AsScriptStructChecked().HasAnyStructFlags(STRUCT_NoExport);
 					SourceFile.AddDefinedStruct(MoveTemp(StructDecl));
 				}
 			}
 		}
 	}
 
-	if (bFoundExportedClasses && !bFoundGeneratedInclude)
+	if (bGeneratedIncludeRequired && !bFoundGeneratedInclude)
 	{
 		Parser.Throwf(TEXT("No #include found for the .generated.h file - the .generated.h file should always be the last #include in a header"));
 	}
@@ -8104,6 +8078,48 @@ TSharedRef<FUnrealTypeDefinitionInfo> FHeaderPreParser::ParseStructDeclaration(c
 		}
 	}
 	);
+
+	// Initialize the structure flags
+	StructDef->SetStructFlags(STRUCT_Native);
+
+	// Record that this struct is RequiredAPI if the CORE_API style macro was present
+	if (!RequiredAPIMacroIfPresent.IsEmpty())
+	{
+		StructDef->SetStructFlags(STRUCT_RequiredAPI);
+	}
+
+	// Process the list of specifiers
+	for (const FPropertySpecifier& Specifier : SpecifiersFound)
+	{
+		switch ((EStructSpecifier)Algo::FindSortedStringCaseInsensitive(*Specifier.Key, GStructSpecifierStrings))
+		{
+		default:
+		{
+			Throwf(TEXT("Unknown struct specifier '%s'"), *Specifier.Key);
+		}
+		break;
+
+		case EStructSpecifier::NoExport:
+		{
+			StructDef->SetStructFlags(STRUCT_NoExport);
+			StructDef->ClearStructFlags(STRUCT_Native);
+		}
+		break;
+
+		case EStructSpecifier::Atomic:
+		{
+			StructDef->SetStructFlags(STRUCT_Atomic);
+		}
+		break;
+
+		case EStructSpecifier::Immutable:
+		{
+			StructDef->SetStructFlags(STRUCT_Immutable);
+			StructDef->SetStructFlags(STRUCT_Atomic);
+		}
+		break;
+		}
+	}
 	return StructDef;
 }
 
diff --git a/Engine/Source/Programs/UnrealHeaderTool/Private/UnrealTypeDefinitionInfo.h b/Engine/Source/Programs/UnrealHeaderTool/Private/UnrealTypeDefinitionInfo.h
--- a/Engine/Source/Programs/UnrealHeaderTool/Private/UnrealTypeDefinitionInfo.h
+++ b/Engine/Source/Programs/UnrealHeaderTool/Private/UnrealTypeDefinitionInfo.h
@@ -2391,6 +2391,21 @@ public:
 		StructFlags = EStructFlags(int32(StructFlags) | int32(FlagsToSet));
 	}
 
+	/**
+	 * Clears the given flags
+	 */
+	void ClearStructFlags(EStructFlags FlagsToClear)
+	{
+#if UHT_ENABLE_ENGINE_TYPE_CHECKS
+		check(GetScriptStructSafe() == nullptr || (GetScriptStructSafe()->StructFlags & ~STRUCT_ComputedFlags) == StructFlags);
+#endif
+		StructFlags = EStructFlags(int32(StructFlags) & int32(~FlagsToClear));
+		if (UScriptStruct* Struct = GetScriptStructSafe())
+		{
+			Struct->StructFlags = EStructFlags(int32(Struct->StructFlags) & int32(~FlagsToClear));
+		}
+	}
+
 	/**
 	 * Used to safely check whether the passed in flag is set.
 	 *
