Add the ability to disable concurrent parsing and code generation to make debugging a bit easier.

#rnx

diff --git a/Engine/Source/Programs/UnrealHeaderTool/Private/CodeGenerator.cpp b/Engine/Source/Programs/UnrealHeaderTool/Private/CodeGenerator.cpp
--- a/Engine/Source/Programs/UnrealHeaderTool/Private/CodeGenerator.cpp
+++ b/Engine/Source/Programs/UnrealHeaderTool/Private/CodeGenerator.cpp
@@ -317,6 +317,18 @@ void FGeneratedFileInfo::StartLoad(FString&& InFilename)
 	}
 }
 
+void FGeneratedFileInfo::Load(FString&& InFilename)
+{
+	ensureMsgf(Filename.IsEmpty(), TEXT("FPreloadHeaderFileInfo::StartLoad called twice with different paths."));
+	Filename = MoveTemp(InFilename);
+
+	if (bAllowSaveExportedHeaders)
+	{
+		SCOPE_SECONDS_COUNTER_UHT(LoadHeaderContentFromFile);
+		FFileHelper::LoadFileToString(OriginalContents, *Filename);
+	}
+}
+
 void FGeneratedFileInfo::GenerateBodyHash()
 {
 	GeneratedBodyHash = GenerateTextHash(*GeneratedBody);
@@ -5465,173 +5477,198 @@ const FString& FNativeClassHeaderGenerator::GetAPIString() const
 	return PackageDef.GetAPI();
 }
 
-void FNativeClassHeaderGenerator::GenerateSourceFiles(
-	TArray<FGeneratedCPP>& GeneratedCPPs
-)
+bool FNativeClassHeaderGenerator::LoadSourceFile(FGeneratedCPP& GeneratedCPP)
 {
-	TSet<FUnrealSourceFile*> Includes;
-	Includes.Reserve(GeneratedCPPs.Num());
-	FGraphEventArray TempTasks;
-	TempTasks.Reserve(3);
-	for (FGeneratedCPP& GeneratedCPP : GeneratedCPPs)
+	const FManifestModule& Module = GeneratedCPP.PackageDef.GetModule();
+	FUnrealSourceFile& SourceFile = GeneratedCPP.SourceFile;
+	if (!SourceFile.ShouldExport())
 	{
-		const FManifestModule& Module = GeneratedCPP.PackageDef.GetModule();
-		FUnrealSourceFile& SourceFile = GeneratedCPP.SourceFile;
-		if (!SourceFile.ShouldExport())
-		{
-			continue;
-		}
-
-		FString ModuleRelativeFilename = SourceFile.GetFilename();
-		ConvertToBuildIncludePath(Module, ModuleRelativeFilename);
+		return false;
+	}
 
-		FString StrippedName = FPaths::GetBaseFilename(MoveTemp(ModuleRelativeFilename));
-		FString HeaderPath = (Module.GeneratedIncludeDirectory / StrippedName) + TEXT(".generated.h");
-		FString GeneratedSourceFilename = (Module.GeneratedIncludeDirectory / StrippedName) + TEXT(".gen.cpp");
+	FString ModuleRelativeFilename = SourceFile.GetFilename();
+	ConvertToBuildIncludePath(Module, ModuleRelativeFilename);
 
-		GeneratedCPP.Header.StartLoad(MoveTemp(HeaderPath));
-		GeneratedCPP.Source.StartLoad(MoveTemp(GeneratedSourceFilename));
+	FString StrippedName = FPaths::GetBaseFilename(MoveTemp(ModuleRelativeFilename));
+	FString HeaderPath = (Module.GeneratedIncludeDirectory / StrippedName) + TEXT(".generated.h");
+	FString GeneratedSourceFilename = (Module.GeneratedIncludeDirectory / StrippedName) + TEXT(".gen.cpp");
 
-		TempTasks.Reset();
-		GeneratedCPP.Header.AddLoadTaskRef(TempTasks);
-		GeneratedCPP.Source.AddLoadTaskRef(TempTasks);
+#if UHT_ENABLE_CONCURRENT_CODE_GENERATION
+	GeneratedCPP.Header.StartLoad(MoveTemp(HeaderPath));
+	GeneratedCPP.Source.StartLoad(MoveTemp(GeneratedSourceFilename));
+#else
+	GeneratedCPP.Header.Load(MoveTemp(HeaderPath));
+	GeneratedCPP.Source.Load(MoveTemp(GeneratedSourceFilename));
+#endif
+	return true;
+}
 
-		auto GenerateSource = [&GeneratedCPP]()
+void FNativeClassHeaderGenerator::GenerateSourceFile(FGeneratedCPP& GeneratedCPP)
+{
+	FResults::Try([&GeneratedCPP]()
 		{
-			FResults::Try([&GeneratedCPP]()
-				{
-					FUnrealPackageDefinitionInfo& PackageDefLcl = GeneratedCPP.PackageDef;
-					const FManifestModule& Module = PackageDefLcl.GetModule();
-					const FNativeClassHeaderGenerator Generator(PackageDefLcl);
-					FUnrealSourceFile& SourceFile = GeneratedCPP.SourceFile;
-					FUHTStringBuilder& GeneratedFunctionDeclarations = GeneratedCPP.GeneratedFunctionDeclarations;
-					FUHTStringBuilder& GeneratedHeaderText = GeneratedCPP.Header.GetGeneratedBody();
-					FUHTStringBuilder& GeneratedCPPText = GeneratedCPP.Source.GetGeneratedBody();
-					FScopedDurationTimer SourceTimer(SourceFile.GetTime(ESourceFileTime::Generate));
-
-					FReferenceGatherers ReferenceGatherers(&GeneratedCPP.CrossModuleReferences, GeneratedCPP.ForwardDeclarations);
-
-					TArray<FUnrealFieldDefinitionInfo*> Types;
-					SourceFile.GetScope()->GatherTypes(Types);
-
-					// Make sure that all the types have a definition range.
-					for (FUnrealFieldDefinitionInfo* TypeDef : Types)
-					{
-						TypeDef->ValidateDefinitionRange();
-					}
+			FUnrealPackageDefinitionInfo& PackageDefLcl = GeneratedCPP.PackageDef;
+			const FManifestModule& Module = PackageDefLcl.GetModule();
+			const FNativeClassHeaderGenerator Generator(PackageDefLcl);
+			FUnrealSourceFile& SourceFile = GeneratedCPP.SourceFile;
+			FUHTStringBuilder& GeneratedFunctionDeclarations = GeneratedCPP.GeneratedFunctionDeclarations;
+			FUHTStringBuilder& GeneratedHeaderText = GeneratedCPP.Header.GetGeneratedBody();
+			FUHTStringBuilder& GeneratedCPPText = GeneratedCPP.Source.GetGeneratedBody();
+			FScopedDurationTimer SourceTimer(SourceFile.GetTime(ESourceFileTime::Generate));
 
-					// Sort by the end of the definition range.  Since classes can contain delegates, we can't use the definition line
-					// since that will place the class ahead of the delegate.
-					Types.StableSort([](const FUnrealFieldDefinitionInfo& Lhs, const FUnrealFieldDefinitionInfo& Rhs)
-						{
-							return Lhs.GetDefinitionRange().End < Rhs.GetDefinitionRange().End;
-						}
-					);
+			FReferenceGatherers ReferenceGatherers(&GeneratedCPP.CrossModuleReferences, GeneratedCPP.ForwardDeclarations);
 
-					TArray<FUnrealFieldDefinitionInfo*> Singletons;
-					Singletons.Reserve(Types.Num());
+			TArray<FUnrealFieldDefinitionInfo*> Types;
+			SourceFile.GetScope()->GatherTypes(Types);
 
-					const FString FileDefineName = SourceFile.GetFileDefineName();
-					const FString& StrippedFilename = SourceFile.GetStrippedFilename();
+			// Make sure that all the types have a definition range.
+			for (FUnrealFieldDefinitionInfo* TypeDef : Types)
+			{
+				TypeDef->ValidateDefinitionRange();
+			}
 
-					GeneratedHeaderText.Logf(
-						TEXT("#ifdef %s")																	LINE_TERMINATOR
-						TEXT("#error \"%s.generated.h already included, missing '#pragma once' in %s.h\"")	LINE_TERMINATOR
-						TEXT("#endif")																		LINE_TERMINATOR
-						TEXT("#define %s")																	LINE_TERMINATOR
-						LINE_TERMINATOR,
-						*FileDefineName, *StrippedFilename, *StrippedFilename, *FileDefineName);
+			// Sort by the end of the definition range.  Since classes can contain delegates, we can't use the definition line
+			// since that will place the class ahead of the delegate.
+			Types.StableSort([](const FUnrealFieldDefinitionInfo& Lhs, const FUnrealFieldDefinitionInfo& Rhs)
+				{
+					return Lhs.GetDefinitionRange().End < Rhs.GetDefinitionRange().End;
+				}
+			);
 
-					for (FUnrealFieldDefinitionInfo* FieldDef : Types)
-					{
-						if (FUnrealEnumDefinitionInfo* EnumDef = UHTCast<FUnrealEnumDefinitionInfo>(FieldDef))
-						{
-							// Is this ever not the case?
-							if (EnumDef->GetOuter()->GetObject()->IsA(UPackage::StaticClass()))
-							{
-								GeneratedFunctionDeclarations.Log(EnumDef->GetExternDecl(true));
-								Generator.ExportGeneratedEnumInitCode(GeneratedCPPText, ReferenceGatherers, SourceFile, *EnumDef);
-							}
-						}
-						else if (FUnrealScriptStructDefinitionInfo* ScriptStructDef = UHTCast<FUnrealScriptStructDefinitionInfo>(FieldDef))
-						{
-							if (ScriptStructDef->HasAnyStructFlags(STRUCT_NoExport) && !HasDynamicOuter(*ScriptStructDef))
-							{
-								Singletons.Add(ScriptStructDef);
-							}
-							GeneratedFunctionDeclarations.Log(ScriptStructDef->GetExternDecl(true));
-							Generator.ExportGeneratedStructBodyMacros(GeneratedHeaderText, GeneratedCPPText, ReferenceGatherers, SourceFile, *ScriptStructDef);
-						}
-						else if (FUnrealFunctionDefinitionInfo* FunctionDef = UHTCast<FUnrealFunctionDefinitionInfo>(FieldDef))
-						{
-							if (!HasDynamicOuter(*FunctionDef))
-							{
-								Singletons.Add(FunctionDef);
-							}
-							GeneratedFunctionDeclarations.Log(FunctionDef->GetExternDecl(true));
-							Generator.ExportDelegateDeclaration(GeneratedCPPText, ReferenceGatherers, SourceFile, *FunctionDef);
-							Generator.ExportDelegateDefinition(GeneratedHeaderText, ReferenceGatherers, SourceFile, *FunctionDef);
-						}
-						else if (FUnrealClassDefinitionInfo* ClassDef = UHTCast<FUnrealClassDefinitionInfo>(FieldDef))
-						{
-							if (!ClassDef->HasAnyClassFlags(CLASS_Intrinsic))
-							{
-								Generator.ExportClassFromSourceFileInner(GeneratedHeaderText, GeneratedCPPText, GeneratedFunctionDeclarations, ReferenceGatherers, *ClassDef, SourceFile, GeneratedCPP.ExportFlags);
-							}
-						}
-					}
+			TArray<FUnrealFieldDefinitionInfo*> Singletons;
+			Singletons.Reserve(Types.Num());
 
-					GeneratedHeaderText.Log(TEXT("#undef CURRENT_FILE_ID\r\n"));
-					GeneratedHeaderText.Logf(TEXT("#define CURRENT_FILE_ID %s\r\n\r\n\r\n"), *SourceFile.GetFileId());
+			const FString FileDefineName = SourceFile.GetFileDefineName();
+			const FString& StrippedFilename = SourceFile.GetStrippedFilename();
 
-					for (FUnrealFieldDefinitionInfo* FieldDef : Types)
+			GeneratedHeaderText.Logf(
+				TEXT("#ifdef %s")																	LINE_TERMINATOR
+				TEXT("#error \"%s.generated.h already included, missing '#pragma once' in %s.h\"")	LINE_TERMINATOR
+				TEXT("#endif")																		LINE_TERMINATOR
+				TEXT("#define %s")																	LINE_TERMINATOR
+				LINE_TERMINATOR,
+				*FileDefineName, *StrippedFilename, *StrippedFilename, *FileDefineName);
+
+			for (FUnrealFieldDefinitionInfo* FieldDef : Types)
+			{
+				if (FUnrealEnumDefinitionInfo* EnumDef = UHTCast<FUnrealEnumDefinitionInfo>(FieldDef))
+				{
+					// Is this ever not the case?
+					if (EnumDef->GetOuter()->GetObject()->IsA(UPackage::StaticClass()))
 					{
-						if (FUnrealEnumDefinitionInfo* EnumDef = UHTCast<FUnrealEnumDefinitionInfo>(FieldDef))
-						{
-							Generator.ExportEnum(GeneratedHeaderText, *EnumDef);
-						}
+						GeneratedFunctionDeclarations.Log(EnumDef->GetExternDecl(true));
+						Generator.ExportGeneratedEnumInitCode(GeneratedCPPText, ReferenceGatherers, SourceFile, *EnumDef);
 					}
-
-					if (Singletons.Num())
+				}
+				else if (FUnrealScriptStructDefinitionInfo* ScriptStructDef = UHTCast<FUnrealScriptStructDefinitionInfo>(FieldDef))
+				{
+					if (ScriptStructDef->HasAnyStructFlags(STRUCT_NoExport) && !HasDynamicOuter(*ScriptStructDef))
 					{
-						SourceFile.GetSingletons().Append(MoveTemp(Singletons));
+						Singletons.Add(ScriptStructDef);
 					}
+					GeneratedFunctionDeclarations.Log(ScriptStructDef->GetExternDecl(true));
+					Generator.ExportGeneratedStructBodyMacros(GeneratedHeaderText, GeneratedCPPText, ReferenceGatherers, SourceFile, *ScriptStructDef);
 				}
-			);
-		};
-
-		auto WriteGenerated = [&GeneratedCPP]()
-		{
-			FResults::Try([&GeneratedCPP]()
+				else if (FUnrealFunctionDefinitionInfo* FunctionDef = UHTCast<FUnrealFunctionDefinitionInfo>(FieldDef))
 				{
-					FUnrealPackageDefinitionInfo& PackageDefLcl = GeneratedCPP.PackageDef;
-					const FManifestModule& Module = PackageDefLcl.GetModule();
-					FUnrealSourceFile& SourceFile = GeneratedCPP.SourceFile;
-					FUHTStringBuilder& GeneratedHeaderText = GeneratedCPP.Header.GetGeneratedBody();
-					FUHTStringBuilder& GeneratedCPPText = GeneratedCPP.Source.GetGeneratedBody();
-					FScopedDurationTimer SourceTimer(SourceFile.GetTime(ESourceFileTime::Generate));
-
-					GeneratedCPP.Source.GenerateBodyHash();
-
-					TSet<FString> AdditionalHeaders;
-					if (EnumHasAnyFlags(GeneratedCPP.ExportFlags, EExportClassOutFlags::NeedsPushModelHeaders))
+					if (!HasDynamicOuter(*FunctionDef))
+					{
+						Singletons.Add(FunctionDef);
+					}
+					GeneratedFunctionDeclarations.Log(FunctionDef->GetExternDecl(true));
+					Generator.ExportDelegateDeclaration(GeneratedCPPText, ReferenceGatherers, SourceFile, *FunctionDef);
+					Generator.ExportDelegateDefinition(GeneratedHeaderText, ReferenceGatherers, SourceFile, *FunctionDef);
+				}
+				else if (FUnrealClassDefinitionInfo* ClassDef = UHTCast<FUnrealClassDefinitionInfo>(FieldDef))
+				{
+					if (!ClassDef->HasAnyClassFlags(CLASS_Intrinsic))
 					{
-						AdditionalHeaders.Add(FString(TEXT("Net/Core/PushModel/PushModelMacros.h")));
+						Generator.ExportClassFromSourceFileInner(GeneratedHeaderText, GeneratedCPPText, GeneratedFunctionDeclarations, ReferenceGatherers, *ClassDef, SourceFile, GeneratedCPP.ExportFlags);
 					}
+				}
+			}
 
-					bool bHasChanged = WriteHeader(GeneratedCPP.Header, GeneratedHeaderText, AdditionalHeaders, GeneratedCPP.ForwardDeclarations);
-					WriteSource(Module, GeneratedCPP.Source, GeneratedCPPText, &SourceFile, GeneratedCPP.CrossModuleReferences);
+			GeneratedHeaderText.Log(TEXT("#undef CURRENT_FILE_ID\r\n"));
+			GeneratedHeaderText.Logf(TEXT("#define CURRENT_FILE_ID %s\r\n\r\n\r\n"), *SourceFile.GetFileId());
 
-					SourceFile.SetGeneratedFilename(MoveTemp(GeneratedCPP.Header.GetFilename()));
-					SourceFile.SetHasChanged(bHasChanged);
+			for (FUnrealFieldDefinitionInfo* FieldDef : Types)
+			{
+				if (FUnrealEnumDefinitionInfo* EnumDef = UHTCast<FUnrealEnumDefinitionInfo>(FieldDef))
+				{
+					Generator.ExportEnum(GeneratedHeaderText, *EnumDef);
 				}
-			);
+			}
+
+			if (Singletons.Num())
+			{
+				SourceFile.GetSingletons().Append(MoveTemp(Singletons));
+			}
+		}
+	);
+}
+
+void FNativeClassHeaderGenerator::WriteSourceFile(FGeneratedCPP& GeneratedCPP)
+{
+	FResults::Try([&GeneratedCPP]()
+		{
+			FUnrealPackageDefinitionInfo& PackageDefLcl = GeneratedCPP.PackageDef;
+			const FManifestModule& Module = PackageDefLcl.GetModule();
+			FUnrealSourceFile& SourceFile = GeneratedCPP.SourceFile;
+			FUHTStringBuilder& GeneratedHeaderText = GeneratedCPP.Header.GetGeneratedBody();
+			FUHTStringBuilder& GeneratedCPPText = GeneratedCPP.Source.GetGeneratedBody();
+			FScopedDurationTimer SourceTimer(SourceFile.GetTime(ESourceFileTime::Generate));
+
+			GeneratedCPP.Source.GenerateBodyHash();
+
+			TSet<FString> AdditionalHeaders;
+			if (EnumHasAnyFlags(GeneratedCPP.ExportFlags, EExportClassOutFlags::NeedsPushModelHeaders))
+			{
+				AdditionalHeaders.Add(FString(TEXT("Net/Core/PushModel/PushModelMacros.h")));
+			}
+
+			bool bHasChanged = WriteHeader(GeneratedCPP.Header, GeneratedHeaderText, AdditionalHeaders, GeneratedCPP.ForwardDeclarations);
+			WriteSource(Module, GeneratedCPP.Source, GeneratedCPPText, &SourceFile, GeneratedCPP.CrossModuleReferences);
+
+			SourceFile.SetGeneratedFilename(MoveTemp(GeneratedCPP.Header.GetFilename()));
+			SourceFile.SetHasChanged(bHasChanged);
+		}
+	);
+}
+
+void FNativeClassHeaderGenerator::GenerateSourceFiles(
+	TArray<FGeneratedCPP>& GeneratedCPPs
+)
+{
+#if UHT_ENABLE_CONCURRENT_CODE_GENERATION
+	TSet<FUnrealSourceFile*> Includes;
+	Includes.Reserve(GeneratedCPPs.Num());
+	FGraphEventArray TempTasks;
+	TempTasks.Reserve(3);
+	for (FGeneratedCPP& GeneratedCPP : GeneratedCPPs)
+	{
+		if (!LoadSourceFile(GeneratedCPP))
+		{
+			continue;
+		}
+
+		TempTasks.Reset();
+		GeneratedCPP.Header.AddLoadTaskRef(TempTasks);
+		GeneratedCPP.Source.AddLoadTaskRef(TempTasks);
+
+		auto GenerateSource = [&GeneratedCPP]()
+		{
+			GenerateSourceFile(GeneratedCPP);
+		};
+
+		auto WriteGenerated = [&GeneratedCPP]()
+		{
+			WriteSourceFile(GeneratedCPP);
 		};
 
 		Includes.Reset();
-		for (FHeaderProvider& Header : SourceFile.GetIncludes())
+		for (FHeaderProvider& Header : GeneratedCPP.SourceFile.GetIncludes())
 		{
-			if (FUnrealSourceFile* Include = Header.Resolve(SourceFile))
+			if (FUnrealSourceFile* Include = Header.Resolve(GeneratedCPP.SourceFile))
 			{
 				Includes.Add(Include);
 			}
@@ -5662,6 +5699,18 @@ void FNativeClassHeaderGenerator::GenerateSourceFiles(
 		GeneratedCPP.AddExportTaskRef(ExportSourceTasks);
 	}
 	FTaskGraphInterface::Get().WaitUntilTasksComplete(ExportSourceTasks);
+#else
+	for (FGeneratedCPP& GeneratedCPP : GeneratedCPPs)
+	{
+		if (!LoadSourceFile(GeneratedCPP))
+		{
+			continue;
+		}
+
+		GenerateSourceFile(GeneratedCPP);
+		WriteSourceFile(GeneratedCPP);
+	}
+#endif
 	FResults::WaitForErrorTasks();
 }
 
@@ -6410,8 +6459,41 @@ void PrepareModules(TArray<FUnrealPackageDefinitionInfo*>& PackageDefs, const FS
 	GUnrealSourceFilesMap.Freeze();
 }
 
+void LoadSource(FUnrealSourceFile& SourceFile, const FString& ModuleInfoPath)
+{
+	FResults::Try([&SourceFile, &ModuleInfoPath]()
+		{
+			FScopedDurationTimer SourceTimer(SourceFile.GetTime(ESourceFileTime::Load));
+
+			const FString FullFilename = FPaths::ConvertRelativePathToFull(ModuleInfoPath, *SourceFile.GetFilename());
+
+			FString Content;
+			if (!FFileHelper::LoadFileToString(Content, *FullFilename))
+			{
+				FUHTMessage(SourceFile).Throwf(TEXT("UnrealHeaderTool was unable to load source file '%s'"), *FullFilename);
+			}
+			SourceFile.SetContent(MoveTemp(Content));
+		}
+	);
+}
+
+void PreparseSource(FUnrealSourceFile& SourceFile)
+{
+	FResults::Try([&SourceFile]()
+		{
+			FScopedDurationTimer SourceTimer(SourceFile.GetTime(ESourceFileTime::PreParse));
+
+			// Parse the header to extract the information needed
+			FUHTStringBuilder ClassHeaderTextStrippedOfCppText;
+			FHeaderParser::SimplifiedClassParse(SourceFile, *SourceFile.GetContent(), /*out*/ ClassHeaderTextStrippedOfCppText);
+			SourceFile.SetContent(MoveTemp(ClassHeaderTextStrippedOfCppText));
+		}
+	);
+}
+
 void PreparseSources(TArray<FUnrealPackageDefinitionInfo*>& PackageDefs, const FString& ModuleInfoPath)
 {
+#if UHT_ENABLE_CONCURRENT_PREPARSING
 	FGraphEventArray LoadTasks;
 	LoadTasks.Reserve(1024); // Fairly arbitrary number
 	for (FUnrealPackageDefinitionInfo* PackageDef : PackageDefs)
@@ -6422,35 +6504,13 @@ void PreparseSources(TArray<FUnrealPackageDefinitionInfo*>& PackageDefs, const F
 			// Phase #1: Load the file
 			auto LoadLambda = [&SourceFile = *SourceFile, &ModuleInfoPath]()
 			{
-				FResults::Try([&SourceFile, &ModuleInfoPath]()
-					{
-						FScopedDurationTimer SourceTimer(SourceFile.GetTime(ESourceFileTime::Load));
-
-						const FString FullFilename = FPaths::ConvertRelativePathToFull(ModuleInfoPath, *SourceFile.GetFilename());
-
-						FString Content;
-						if (!FFileHelper::LoadFileToString(Content, *FullFilename))
-						{
-							FUHTMessage(SourceFile).Throwf(TEXT("UnrealHeaderTool was unable to load source file '%s'"), *FullFilename);
-						}
-						SourceFile.SetContent(MoveTemp(Content));
-					}
-				);
+				LoadSource(SourceFile, ModuleInfoPath);
 			};
 
 			// Phase #2: Perform simplified class parse (can run concurrenrtly)
 			auto PreProcessLambda = [&SourceFile = *SourceFile]()
 			{
-				FResults::Try([&SourceFile]()
-					{
-						FScopedDurationTimer SourceTimer(SourceFile.GetTime(ESourceFileTime::PreParse));
-
-						// Parse the header to extract the information needed
-						FUHTStringBuilder ClassHeaderTextStrippedOfCppText;
-						FHeaderParser::SimplifiedClassParse(SourceFile, *SourceFile.GetContent(), /*out*/ ClassHeaderTextStrippedOfCppText);
-						SourceFile.SetContent(MoveTemp(ClassHeaderTextStrippedOfCppText));
-					}
-				);
+				PreparseSource(SourceFile);
 			};
 
 			FGraphEventRef LoadTask = FFunctionGraphTask::CreateAndDispatchWhenReady(MoveTemp(LoadLambda), TStatId());
@@ -6461,6 +6521,16 @@ void PreparseSources(TArray<FUnrealPackageDefinitionInfo*>& PackageDefs, const F
 
 	// Wait for all the loading and preparsing to complete
 	FTaskGraphInterface::Get().WaitUntilTasksComplete(LoadTasks);
+#else
+	for (FUnrealPackageDefinitionInfo* PackageDef : PackageDefs)
+	{
+		for (TSharedRef<FUnrealSourceFile>& SourceFile : PackageDef->GetAllSourceFiles())
+		{
+			LoadSource(*SourceFile, ModuleInfoPath);
+			PreparseSource(*SourceFile);
+		}
+	}
+#endif
 	FResults::WaitForErrorTasks();
 }
 
diff --git a/Engine/Source/Programs/UnrealHeaderTool/Private/NativeClassExporter.h b/Engine/Source/Programs/UnrealHeaderTool/Private/NativeClassExporter.h
--- a/Engine/Source/Programs/UnrealHeaderTool/Private/NativeClassExporter.h
+++ b/Engine/Source/Programs/UnrealHeaderTool/Private/NativeClassExporter.h
@@ -90,6 +90,11 @@ struct FGeneratedFileInfo
 	 */
 	void StartLoad(FString&& InFilename);
 
+	/**
+	 * Load the existing version of the file immediately
+	 */
+	void Load(FString&& InFilename);
+
 	/**
 	 * Get the output file name.
 	 */
@@ -725,6 +730,28 @@ public:
 	 */
 	static void GenerateSourceFiles(TArray<FGeneratedCPP>& GeneratedCPPs);
 
+	/**
+	 * Load a single source file
+	 *
+	 * @parm	GeneratedCPP		Single source file to load
+	 * @return  True if the source file is to be exported, false if not.
+	 */
+	static bool LoadSourceFile(FGeneratedCPP& GeneratedCPP);
+
+	/**
+	 * Generate a single source file
+	 *
+	 * @parm	GeneratedCPP		Single source file to generate
+	 */
+	static void GenerateSourceFile(FGeneratedCPP& GeneratedCPP);
+
+	/**
+	 * Write the generated files
+	 *
+	 * @parm	GeneratedCPP		Single source file to write
+	 */
+	static void WriteSourceFile(FGeneratedCPP& GeneratedCPP);
+
 	/**
 	 * Generate all the extra output files for the given package.
 	 *
diff --git a/Engine/Source/Programs/UnrealHeaderTool/Private/UnrealHeaderTool.h b/Engine/Source/Programs/UnrealHeaderTool/Private/UnrealHeaderTool.h
--- a/Engine/Source/Programs/UnrealHeaderTool/Private/UnrealHeaderTool.h
+++ b/Engine/Source/Programs/UnrealHeaderTool/Private/UnrealHeaderTool.h
@@ -13,9 +13,15 @@
 // we might never need this tracking since it is always by pointer.
 #define UHT_ENABLE_PTR_PROPERTY_TAG 0
 
+// Enable concurrent pre-parsing of header files
+#define UHT_ENABLE_CONCURRENT_PREPARSING 1
+
 // Enable concurrent parsing of header files
 #define UHT_ENABLE_CONCURRENT_PARSING 1
 
+// Enable concurrent code generation
+#define UHT_ENABLE_CONCURRENT_CODE_GENERATION 1
+
 struct FHeaderParserNames
 {
 	static const FName NAME_IsConversionRoot;
