Removed the FToken from the FFuncInfo structure.
Saved around 100MB of memory on some larger games.

#rnx

diff --git a/Engine/Source/Programs/UnrealHeaderTool/Private/HeaderParser.cpp b/Engine/Source/Programs/UnrealHeaderTool/Private/HeaderParser.cpp
--- a/Engine/Source/Programs/UnrealHeaderTool/Private/HeaderParser.cpp
+++ b/Engine/Source/Programs/UnrealHeaderTool/Private/HeaderParser.cpp
@@ -5921,7 +5921,8 @@ FUnrealFunctionDefinitionInfo& FHeaderParser::CompileDelegateDeclaration(const T
 	FuncInfo.InputPos = InputPos;
 
 	// Get the delegate name
-	if (!GetIdentifier(FuncInfo.Function))
+	FToken FuncNameToken;
+	if (!GetIdentifier(FuncNameToken))
 	{
 		Throwf(TEXT("Missing name for %s"), CurrentScopeName );
 	}
@@ -5932,7 +5933,7 @@ FUnrealFunctionDefinitionInfo& FHeaderParser::CompileDelegateDeclaration(const T
 		//@TODO: UCREMOVAL: Eventually this mangling shouldn't occur
 
 		// Remove the leading F
-		FString Name(FuncInfo.Function.Identifier);
+		FString Name(FuncNameToken.Identifier);
 
 		if (!Name.StartsWith(TEXT("F"), ESearchCase::CaseSensitive))
 		{
@@ -5945,12 +5946,11 @@ FUnrealFunctionDefinitionInfo& FHeaderParser::CompileDelegateDeclaration(const T
 		Name += HEADER_GENERATED_DELEGATE_SIGNATURE_SUFFIX;
 
 		// Replace the name
-		FCString::Strcpy( FuncInfo.Function.Identifier, *Name );
+		FCString::Strcpy(FuncNameToken.Identifier, *Name );
 	}
 
-	FUnrealFunctionDefinitionInfo& DelegateSignatureFunctionDef = CreateFunction(MoveTemp(FuncInfo), bIsSparse ? EFunctionType::SparseDelegate : EFunctionType::Delegate);
+	FUnrealFunctionDefinitionInfo& DelegateSignatureFunctionDef = CreateFunction(FuncNameToken.Identifier, MoveTemp(FuncInfo), bIsSparse ? EFunctionType::SparseDelegate : EFunctionType::Delegate);
 	DelegateSignatureFunctionDef.GetDefinitionRange().Start = StartPos;
-	const FFuncInfo& FuncDefFuncInfo = DelegateSignatureFunctionDef.GetFunctionData();
 
 	// determine whether this function should be 'const'
 	if (bDeclaredConst)
@@ -6034,7 +6034,7 @@ FUnrealFunctionDefinitionInfo& FHeaderParser::CompileDelegateDeclaration(const T
 		FUnrealFunctionDefinitionInfo* TestFuncDef = *FunctionIterator;
 		if (TestFuncDef != &DelegateSignatureFunctionDef && TestFuncDef->GetFName() == DelegateSignatureFunctionDef.GetFName())
 		{
-			Throwf(TEXT("Can't override delegate signature function '%s'"), FuncDefFuncInfo.Function.Identifier);
+			Throwf(TEXT("Can't override delegate signature function '%s'"), *DelegateSignatureFunctionDef.GetNameCPP());
 		}
 	}
 
@@ -6301,7 +6301,8 @@ FUnrealFunctionDefinitionInfo& FHeaderParser::CompileFunctionDeclaration()
 	FuncInfo.InputPos = InputPos;
 
 	// Get function or operator name.
-	if (!GetIdentifier(FuncInfo.Function))
+	FToken FuncNameToken;
+	if (!GetIdentifier(FuncNameToken))
 	{
 		Throwf(TEXT("Missing %s name"), TypeOfFunction);
 	}
@@ -6328,7 +6329,7 @@ FUnrealFunctionDefinitionInfo& FHeaderParser::CompileFunctionDeclaration()
 				Throwf(TEXT("Function %s already uses identifier %d"), **ExistingFunc, FuncInfo.RPCId);
 			}
 
-			UsedRPCIds.Add(FuncInfo.RPCId, FuncInfo.Function.Identifier);
+			UsedRPCIds.Add(FuncInfo.RPCId, FuncNameToken.Identifier);
 			if (FuncInfo.FunctionFlags & FUNC_NetResponse)
 			{
 				// Look for another function expecting this response
@@ -6347,12 +6348,12 @@ FUnrealFunctionDefinitionInfo& FHeaderParser::CompileFunctionDeclaration()
 			if (ExistingFunc == NULL)
 			{
 				// If this list isn't empty at end of class, throw error
-				RPCsNeedingHookup.Add(FuncInfo.RPCResponseId, FuncInfo.Function.Identifier);
+				RPCsNeedingHookup.Add(FuncInfo.RPCResponseId, FuncNameToken.Identifier);
 			}
 		}
 	}
 
-	FUnrealFunctionDefinitionInfo& FuncDef = CreateFunction(MoveTemp(FuncInfo), EFunctionType::Function);
+	FUnrealFunctionDefinitionInfo& FuncDef = CreateFunction(FuncNameToken.Identifier, MoveTemp(FuncInfo), EFunctionType::Function);
 	FuncDef.GetDefinitionRange().Start = StartPos;
 	const FFuncInfo& FuncDefFuncInfo = FuncDef.GetFunctionData();
 	FUnrealFunctionDefinitionInfo* SuperFuncDef = FuncDef.GetSuperFunction();
@@ -6389,10 +6390,10 @@ FUnrealFunctionDefinitionInfo& FHeaderParser::CompileFunctionDeclaration()
 	}
 	if (SuperStructDef)
 	{
-		if (FUnrealFunctionDefinitionInfo* OverriddenFunctionDef = FindFunction(*SuperStructDef, FuncDefFuncInfo.Function.Identifier, true, nullptr))
+		if (FUnrealFunctionDefinitionInfo* OverriddenFunctionDef = FindFunction(*SuperStructDef, *FuncDef.GetNameCPP(), true, nullptr))
 		{
 			// Native function overrides should be done in CPP text, not in a UFUNCTION() declaration (you can't change flags, and it'd otherwise be a burden to keep them identical)
-			LogError(TEXT("%s: Override of UFUNCTION in parent class (%s) cannot have a UFUNCTION() declaration above it; it will use the same parameters as the original declaration."), FuncDefFuncInfo.Function.Identifier, *OverriddenFunctionDef->GetOuter()->GetName());
+			LogError(TEXT("%s: Override of UFUNCTION in parent class (%s) cannot have a UFUNCTION() declaration above it; it will use the same parameters as the original declaration."), *FuncDef.GetNameCPP(), *OverriddenFunctionDef->GetOuter()->GetName());
 		}
 	}
 
@@ -8763,7 +8764,7 @@ void FHeaderParser::ConditionalLogPointerUsage(EPointerMemberBehavior PointerMem
 	}
 }
 
-FUnrealFunctionDefinitionInfo& FHeaderParser::CreateFunction(FFuncInfo&& FuncInfo, EFunctionType InFunctionType) const
+FUnrealFunctionDefinitionInfo& FHeaderParser::CreateFunction(const TCHAR* FuncName, FFuncInfo&& FuncInfo, EFunctionType InFunctionType) const
 {
 	FScope* CurrentScope = GetCurrentScope();
 	FUnrealObjectDefinitionInfo& OuterDef = IsInAClass() ? static_cast<FUnrealObjectDefinitionInfo&>(GetCurrentClassDef()) : SourceFile.GetPackageDef();
@@ -8776,14 +8777,14 @@ FUnrealFunctionDefinitionInfo& FHeaderParser::CreateFunction(FFuncInfo&& FuncInf
 		while (TypeIterator.MoveNext())
 		{
 			FUnrealFieldDefinitionInfo* Type = *TypeIterator;
-			if (Type->GetFName() == FuncInfo.Function.Identifier)
+			if (Type->GetFName() == FuncName)
 			{
-				Throwf(TEXT("'%s' conflicts with '%s'"), FuncInfo.Function.Identifier, *Type->GetFullName());
+				Throwf(TEXT("'%s' conflicts with '%s'"), FuncName, *Type->GetFullName());
 			}
 		}
 	}
 
-	TSharedRef<FUnrealFunctionDefinitionInfo> FuncDefRef = MakeShared<FUnrealFunctionDefinitionInfo>(SourceFile, InputLine, FString(FuncInfo.Function.Identifier), FName(FuncInfo.Function.Identifier, FNAME_Add), OuterDef, MoveTemp(FuncInfo), InFunctionType);
+	TSharedRef<FUnrealFunctionDefinitionInfo> FuncDefRef = MakeShared<FUnrealFunctionDefinitionInfo>(SourceFile, InputLine, FString(FuncName), FName(FuncName, FNAME_Add), OuterDef, MoveTemp(FuncInfo), InFunctionType);
 	FUnrealFunctionDefinitionInfo& FuncDef = *FuncDefRef;
 	FuncDef.GetFunctionData().SetFunctionNames(FuncDef);
 
diff --git a/Engine/Source/Programs/UnrealHeaderTool/Private/HeaderParser.h b/Engine/Source/Programs/UnrealHeaderTool/Private/HeaderParser.h
--- a/Engine/Source/Programs/UnrealHeaderTool/Private/HeaderParser.h
+++ b/Engine/Source/Programs/UnrealHeaderTool/Private/HeaderParser.h
@@ -387,7 +387,7 @@ protected:
 	/**
 	 * Create new function object based on given info structure.
 	 */
-	FUnrealFunctionDefinitionInfo& CreateFunction(FFuncInfo&& FuncInfo, EFunctionType InFunctionType) const;
+	FUnrealFunctionDefinitionInfo& CreateFunction(const TCHAR* FuncName, FFuncInfo&& FuncInfo, EFunctionType InFunctionType) const;
 
 	FUnrealClassDefinitionInfo& CompileClassDeclaration();
 	FUnrealFunctionDefinitionInfo& CompileDelegateDeclaration(const TCHAR* DelegateIdentifier, EDelegateSpecifierAction::Type SpecifierAction = EDelegateSpecifierAction::DontParse);
diff --git a/Engine/Source/Programs/UnrealHeaderTool/Private/ParserHelper.h b/Engine/Source/Programs/UnrealHeaderTool/Private/ParserHelper.h
--- a/Engine/Source/Programs/UnrealHeaderTool/Private/ParserHelper.h
+++ b/Engine/Source/Programs/UnrealHeaderTool/Private/ParserHelper.h
@@ -607,14 +607,12 @@ struct FFuncInfo
 {
 	/** @name Variables */
 	//@{
-	/** Name of the function or operator. */
-	FToken		Function;
 	/** Function flags. */
-	EFunctionFlags	FunctionFlags;
+	EFunctionFlags	FunctionFlags = FUNC_None;
 	/** Function flags which are only required for exporting */
-	uint32		FunctionExportFlags;
+	uint32		FunctionExportFlags = 0;
 	/** Number of parameters expected for operator. */
-	int32		ExpectParms;
+	int32		ExpectParms = 0;
 	/** Name of the wrapper function that marshalls the arguments and does the indirect call **/
 	FString		MarshallAndCallName;
 	/** Name of the actual implementation **/
@@ -626,38 +624,20 @@ struct FFuncInfo
 	/** Endpoint name */
 	FString		EndpointName;
 	/** Identifier for an RPC call to a platform service */
-	uint16		RPCId;
+	uint16		RPCId = 0;
 	/** Identifier for an RPC call expecting a response */
-	uint16		RPCResponseId;
+	uint16		RPCResponseId = 0;
 	/** Delegate macro line in header. */
-	int32		MacroLine;
+	int32		MacroLine = -1;
 	/** Position in file where this function was declared. Points to first char of function name. */
-	int32 InputPos;
+	int32		InputPos = -1;
 	/** Whether this function represents a sealed event */
-	bool		bSealedEvent;
+	bool		bSealedEvent = false;
 	/** TRUE if the function is being forced to be considered as impure by the user */
-	bool bForceBlueprintImpure;
+	bool		bForceBlueprintImpure = false;
 
 	//@}
 
-	/** Constructor. */
-	FFuncInfo()
-		: Function()
-		, FunctionFlags(FUNC_None)
-		, FunctionExportFlags(0)
-		, ExpectParms(0)
-		, RPCId(0)
-		, RPCResponseId(0)
-		, MacroLine(-1)
-		, InputPos(-1)
-		, bSealedEvent(false)
-		, bForceBlueprintImpure(false)
-	{}
-
-	FFuncInfo(const FFuncInfo& Other) = default;
-	FFuncInfo(FFuncInfo&& Other) = default;
-	FFuncInfo& operator=(FFuncInfo&& Other) = default;
-
 	/** Set the internal function names based on flags **/
 	void SetFunctionNames(FUnrealFunctionDefinitionInfo& FunctionDef);
 };
