Removed all lookups from UObjects to UHT definition objects.

#rnx

diff --git a/Engine/Source/Programs/UnrealHeaderTool/Private/ClassMaps.h b/Engine/Source/Programs/UnrealHeaderTool/Private/ClassMaps.h
--- a/Engine/Source/Programs/UnrealHeaderTool/Private/ClassMaps.h
+++ b/Engine/Source/Programs/UnrealHeaderTool/Private/ClassMaps.h
@@ -32,84 +32,13 @@ protected:
 };
 
 // Wrapper class around TypeDefinition map so we can maintain a parallel by name map
-struct FTypeDefinitionInfoMap : public FFreezableContainer
+struct FTypeDefinitionInfoMap
 {
-
-	//NOTE: UObjects are frozen after the preparsing phase
-	void Add(UObject* Object, TSharedRef<FUnrealTypeDefinitionInfo>&& Definition)
-	{
-		check(!bFrozen);
-		AddNameLookup(UHTCastChecked<FUnrealObjectDefinitionInfo>(Definition));
-		AddObjectLookup(Object, MoveTemp(Definition));
-	}
-
-	void Add(UObject* Object, TSharedRef<FUnrealTypeDefinitionInfo>& Definition)
-	{
-		check(!bFrozen);
-		AddNameLookup(UHTCastChecked<FUnrealObjectDefinitionInfo>(Definition));
-		AddObjectLookup(Object, Definition);
-	}
-
-	void AddObjectLookup(UObject* Object, TSharedRef<FUnrealTypeDefinitionInfo>&& Definition)
-	{
-		//check(!bFrozen);
-		DefinitionsByUObject.Add(Object, Definition);
-	}
-
-	void AddObjectLookup(UObject* Object, TSharedRef<FUnrealTypeDefinitionInfo>& Definition)
-	{
-		//check(!bFrozen);
-		DefinitionsByUObject.Add(Object, Definition);
-	}
-
 	void AddNameLookup(FUnrealObjectDefinitionInfo& Definition)
 	{
-		check(!bFrozen);
 		DefinitionsByName.Add(Definition.GetFName(), Definition.AsShared());
 	}
 
-	bool Contains(const UObject* Object)
-	{
-		check(bFrozen); 
-		return DefinitionsByUObject.Contains(Object); 
-	}
-
-	TSharedRef<FUnrealTypeDefinitionInfo>* Find(const UObject* Object)
-	{
-		check(bFrozen);
-		return DefinitionsByUObject.Find(Object);
-	}
-
-	template<typename To>
-	To* Find(const UObject* Object)
-	{
-		check(bFrozen);
-		return UHTCast<To>(DefinitionsByUObject.Find(Object));
-	}
-
-	TSharedRef<FUnrealTypeDefinitionInfo>& operator[](const UObject* Object)
-	{
-		check(bFrozen);
-		return DefinitionsByUObject[Object]; 
-	}
-
-	FUnrealTypeDefinitionInfo& FindChecked(const UObject* Object)
-	{
-		check(bFrozen);
-		TSharedRef<FUnrealTypeDefinitionInfo>* TypeDef = DefinitionsByUObject.Find(Object);
-		check(TypeDef);
-		return **TypeDef;
-	}
-
-	template<typename To>
-	To& FindChecked(const UObject* Object)
-	{
-		check(bFrozen);
-		To* TypeDef = Find<To>(Object);
-		check(TypeDef);
-		return *TypeDef;
-	}
-
 	// Finding by name must be done on the stripped name for classes and script structs
 	TSharedRef<FUnrealTypeDefinitionInfo>* FindByName(const TCHAR* Name)
 	{
@@ -124,13 +53,11 @@ struct FTypeDefinitionInfoMap : public FFreezableContainer
 	template<typename To>
 	To* FindByName(const TCHAR* Name)
 	{
-		check(bFrozen);
 		return UHTCast<To>(FindByName(Name));
 	}
 
 	FUnrealTypeDefinitionInfo& FindByNameChecked(const TCHAR* Name)
 	{
-		check(bFrozen);
 		TSharedRef<FUnrealTypeDefinitionInfo>* TypeDef = FindByName(Name);
 		check(TypeDef);
 		return **TypeDef;
@@ -139,7 +66,6 @@ struct FTypeDefinitionInfoMap : public FFreezableContainer
 	template <typename To>
 	To& FindByNameChecked(const TCHAR* Name)
 	{
-		check(bFrozen);
 		To* TypeDef = FindByName<To>(Name);
 		check(TypeDef);
 		return *TypeDef;
@@ -154,60 +80,7 @@ struct FTypeDefinitionInfoMap : public FFreezableContainer
 		}
 	}
 
-	//NOTE: Currently UFunctions are created during the parsing phase and can not be frozen
-	void Add(UFunction* Object, TSharedRef<FUnrealTypeDefinitionInfo>&& Definition)
-	{
-		DefinitionsByUObject.Add(Object, Definition);
-		// At this point, do not add the name
-		//DefinitionsByName.Add(Object->GetFName(), Definition);
-	}
-
-	//NOTE: FFields (properties) are not frozen since they are added during the parsing phase
-	void Add(FField* Field, TSharedRef<FUnrealTypeDefinitionInfo>&& Definition)
-	{
-		DefinitionsByFField.Add(Field, MoveTemp(Definition));
-	}
-
-	bool Contains(const FField* Field) 
-	{
-		return DefinitionsByFField.Contains(Field);
-	}
-
-	TSharedRef<FUnrealTypeDefinitionInfo>* Find(const FField* Field)
-	{
-		return DefinitionsByFField.Find(Field);
-	}
-
-	template <typename To>
-	To* Find(const FField* Field)
-	{
-		return UHTCast<To>(DefinitionsByFField.Find(Field));
-	}
-
-	FUnrealTypeDefinitionInfo& FindChecked(const FField* Field)
-	{
-		TSharedRef<FUnrealTypeDefinitionInfo>* TypeDef = DefinitionsByFField.Find(Field);
-		check(TypeDef);
-		return **TypeDef;
-	}
-
-	template <typename To>
-	To& FindChecked(const FField* Field)
-	{
-		To* TypeDef = Find<To>(Field);
-		check(TypeDef);
-		return *TypeDef;
-	}
-
-	TSharedRef<FUnrealTypeDefinitionInfo>& operator[](const FField* Field)
-	{ 
-		return DefinitionsByFField[Field];
-	}
-
 private:
-
-	TMap<UObject*, TSharedRef<FUnrealTypeDefinitionInfo>> DefinitionsByUObject;
-	TMap<FField*, TSharedRef<FUnrealTypeDefinitionInfo>> DefinitionsByFField;
 	TMap<FName, TSharedRef<FUnrealTypeDefinitionInfo>> DefinitionsByName;
 };
 
diff --git a/Engine/Source/Programs/UnrealHeaderTool/Private/CodeGenerator.cpp b/Engine/Source/Programs/UnrealHeaderTool/Private/CodeGenerator.cpp
--- a/Engine/Source/Programs/UnrealHeaderTool/Private/CodeGenerator.cpp
+++ b/Engine/Source/Programs/UnrealHeaderTool/Private/CodeGenerator.cpp
@@ -1944,11 +1944,8 @@ void FNativeClassHeaderGenerator::ExportGeneratedPackageInitCode(FOutputDevice&
 
 void FNativeClassHeaderGenerator::ExportNativeGeneratedInitCode(FOutputDevice& Out, FOutputDevice& OutDeclarations, FReferenceGatherers& OutReferenceGatherers, const FUnrealSourceFile& SourceFile, FUnrealClassDefinitionInfo& ClassDef, FUHTStringBuilder& OutFriendText) const
 {
-	UClass* Class = ClassDef.GetClass();
 	check(!OutFriendText.Len());
 
-	UE_CLOG(Class->ClassGeneratedBy, LogCompile, Fatal, TEXT("For intrinsic and compiled-in classes, ClassGeneratedBy should always be null"));
-
 	const bool   bIsNoExport  = ClassDef.HasAnyClassFlags(CLASS_NoExport);
 	const bool   bIsDynamic   = ClassDef.IsDynamic();
 	const FString ClassNameCPP = ClassDef.GetAlternateNameCPP();
@@ -2113,33 +2110,41 @@ void FNativeClassHeaderGenerator::ExportNativeGeneratedInitCode(FOutputDevice& O
 
 		const TCHAR* InterfaceArray;
 		const TCHAR* InterfaceCount;
-		if (Class->Interfaces.Num() > 0)
+
+		// Check to see if we have any interfaces
+		bool bHasInterfaces = false;
+		for (FUnrealStructDefinitionInfo::FBaseStructInfo& BaseStruct : ClassDef.GetBaseStructInfos())
+		{
+			if (BaseStruct.Struct != nullptr && UHTCastChecked<FUnrealClassDefinitionInfo>(BaseStruct.Struct).IsInterface())
+			{
+				bHasInterfaces = true;
+				break;
+			}
+		}
+
+		if (bHasInterfaces)
 		{
 			GeneratedClassRegisterFunctionText.Log(TEXT("\t\tstatic const UECodeGen_Private::FImplementedInterfaceParams InterfaceParams[];\r\n"));
 
 			StaticDefinitions.Logf(TEXT("\t\tconst UECodeGen_Private::FImplementedInterfaceParams %s::InterfaceParams[] = {\r\n"), *StaticsStructName);
-			for (const FImplementedInterface& Inter : Class->Interfaces)
+			for (FUnrealStructDefinitionInfo::FBaseStructInfo& BaseStruct : ClassDef.GetBaseStructInfos())
 			{
-				check(Inter.Class);
-				FUnrealFieldDefinitionInfo& InterClassDef = GTypeDefinitionInfoMap.FindChecked<FUnrealFieldDefinitionInfo>(Inter.Class);
-				InterClassDef.AddCrossModuleReference(OutReferenceGatherers.UniqueCrossModuleReferences, false);
-				FString OffsetString;
-				if (Inter.PointerOffset)
-				{
-					OffsetString = FString::Printf(TEXT("(int32)VTABLE_OFFSET(%s, %s)"), *ClassNameCPP, *FNameLookupCPP::GetNameCPP(Inter.Class, true));
-				}
-				else
+				if (BaseStruct.Struct == nullptr || !UHTCastChecked<FUnrealClassDefinitionInfo>(BaseStruct.Struct).IsInterface())
 				{
-					OffsetString = TEXT("0");
+					continue;
 				}
 
+				FUnrealClassDefinitionInfo& InterClassDef = UHTCastChecked<FUnrealClassDefinitionInfo>(BaseStruct.Struct);
+				InterClassDef.AddCrossModuleReference(OutReferenceGatherers.UniqueCrossModuleReferences, false);
+				FString OffsetString = FString::Printf(TEXT("(int32)VTABLE_OFFSET(%s, %s)"), *ClassNameCPP, *InterClassDef.GetAlternateNameCPP(true));
+
 				FUHTStringBuilder IntHash;
 				InterClassDef.GetHashTag(ClassDef, IntHash);
 				StaticDefinitions.Logf(
 					TEXT("\t\t\t{ %s, %s, %s }, %s\r\n"),
 					*InterClassDef.GetSingletonNameChopped(false),
 					*OffsetString,
-					Inter.bImplementedByK2 ? TEXT("true") : TEXT("false"),
+					TEXT("false"),
 					*IntHash
 				);
 			}
@@ -2320,14 +2325,14 @@ void FNativeClassHeaderGenerator::ExportNativeGeneratedInitCode(FOutputDevice& O
 		FUHTStringBuilder ValidationBuilder;
 		ValidationBuilder.Log(TEXT("\t\tconst bool bIsValid = true"));
 
-		for (int32 i = Class->FirstOwnedClassRep; i < Class->ClassReps.Num(); ++i)
+		for (int32 i = ClassDef.GetFirstOwnedClassRep(); i < ClassDef.GetClassReps().Num(); ++i)
 		{
-			const FProperty* const Property = Class->ClassReps[i].Property;
-			const FString PropertyName = Property->GetName();
+			const FUnrealPropertyDefinitionInfo* PropertyDef = ClassDef.GetClassReps()[i];
+			const FString PropertyName = PropertyDef->GetName();
 
-			NameBuilder.Logf(TEXT("\t\tstatic const FName Name_%s(TEXT(\"%s\"));\r\n"), *PropertyName, *FNativeClassHeaderGenerator::GetOverriddenName(Property));
+			NameBuilder.Logf(TEXT("\t\tstatic const FName Name_%s(TEXT(\"%s\"));\r\n"), *PropertyName, *FNativeClassHeaderGenerator::GetOverriddenName(PropertyDef));
 
-			if (Property->ArrayDim == 1)
+			if (!PropertyDef->IsStaticArray())
 			{
 				ValidationBuilder.Logf(TEXT("\r\n\t\t\t&& Name_%s == ClassReps[(int32)ENetFields_Private::%s].Property->GetFName()"), *PropertyName, *PropertyName);
 			}
@@ -6361,7 +6366,7 @@ void PrepareModules(TArray<FUnrealPackageDefinitionInfo*>& PackageDefs, const FS
 		// Create the package definition
 		TSharedRef<FUnrealPackageDefinitionInfo> PackageDefRef = MakeShared<FUnrealPackageDefinitionInfo>(Module, Package);
 		FUnrealPackageDefinitionInfo& PackageDef = *PackageDefRef;
-		GTypeDefinitionInfoMap.Add(Package, PackageDefRef);
+		GTypeDefinitionInfoMap.AddNameLookup(PackageDef);
 		PackageDefs.Add(&PackageDef);
 
 		TArray<TSharedRef<FUnrealSourceFile>>& AllSourceFiles = PackageDef.GetAllSourceFiles();
@@ -6582,7 +6587,6 @@ void DefineTypes(TArray<FUnrealPackageDefinitionInfo*>& PackageDefs)
 		}
 	}
 	FResults::WaitForErrorTasks();
-	GTypeDefinitionInfoMap.Freeze();
 }
 
 void ResolveParents(TArray<FUnrealPackageDefinitionInfo*>& PackageDefs)
@@ -6874,34 +6878,34 @@ void Export(TArray<FUnrealPackageDefinitionInfo*>& PackageDefs, TArray<FUnrealSo
 }
 
 // Exports the class to all available plugins
-void ExportClassToScriptPlugins(UClass* Class, const FManifestModule& Module, IScriptGeneratorPluginInterface& ScriptPlugin)
+void ExportClassToScriptPlugins(const TMap<UClass*, FUnrealSourceFile*>& SourceFileLookup, UClass* Class, const FManifestModule& Module, IScriptGeneratorPluginInterface& ScriptPlugin)
 {
-	TSharedRef<FUnrealTypeDefinitionInfo>* DefinitionInfoRef = GTypeDefinitionInfoMap.Find(Class);
-	if (DefinitionInfoRef == nullptr)
+	check(SourceFileLookup.Find(Class) != nullptr);
+	FUnrealSourceFile* const * SourceFile = SourceFileLookup.Find(Class);
+	if (SourceFile == nullptr)
 	{
 		const FString Empty = TEXT("");
 		ScriptPlugin.ExportClass(Class, Empty, Empty, false);
 	}
 	else
 	{
-		FUnrealSourceFile& SourceFile = (*DefinitionInfoRef)->GetUnrealSourceFile();
-		ScriptPlugin.ExportClass(Class, SourceFile.GetFilename(), SourceFile.GetGeneratedFilename(), SourceFile.HasChanged());
+		ScriptPlugin.ExportClass(Class, (*SourceFile)->GetFilename(), (*SourceFile)->GetGeneratedFilename(), (*SourceFile)->HasChanged());
 	}
 }
 
 // Exports class tree to all available plugins
-void ExportClassTreeToScriptPlugins(const FClassTree* Node, const FManifestModule& Module, IScriptGeneratorPluginInterface& ScriptPlugin)
+void ExportClassTreeToScriptPlugins(const TMap<UClass*, FUnrealSourceFile*>& SourceFileLookup, const FClassTree* Node, const FManifestModule& Module, IScriptGeneratorPluginInterface& ScriptPlugin)
 {
 	for (int32 ChildIndex = 0; ChildIndex < Node->NumChildren(); ++ChildIndex)
 	{
 		const FClassTree* ChildNode = Node->GetChild(ChildIndex);
-		ExportClassToScriptPlugins(ChildNode->GetClass(), Module, ScriptPlugin);
+		ExportClassToScriptPlugins(SourceFileLookup, ChildNode->GetClass(), Module, ScriptPlugin);
 	}
 
 	for (int32 ChildIndex = 0; ChildIndex < Node->NumChildren(); ++ChildIndex)
 	{
 		const FClassTree* ChildNode = Node->GetChild(ChildIndex);
-		ExportClassTreeToScriptPlugins(ChildNode, Module, ScriptPlugin);
+		ExportClassTreeToScriptPlugins(SourceFileLookup, ChildNode, Module, ScriptPlugin);
 	}
 }
 
@@ -6918,6 +6922,19 @@ void ExportToScriptPlugins(TArray<FUnrealPackageDefinitionInfo*>& PackageDefs, F
 		return;
 	}
 
+	TMap<UClass*, FUnrealSourceFile*> SourceFileLookup;
+	for (FUnrealPackageDefinitionInfo* PackageDef : PackageDefs)
+	{
+		for (TSharedRef<FUnrealTypeDefinitionInfo> TypeDef : PackageDef->GetAllClasses())
+		{
+			if (TypeDef->HasSource())
+			{
+				UClass* Class = UHTCastChecked<FUnrealClassDefinitionInfo>(TypeDef).GetClass();
+				SourceFileLookup.Add(Class, &TypeDef->GetUnrealSourceFile());
+			}
+		}
+	}
+
 	for (FUnrealPackageDefinitionInfo* PackageDef : PackageDefs)
 	{
 		const FManifestModule& Module = PackageDef->GetModule();
@@ -6933,8 +6950,8 @@ void ExportToScriptPlugins(TArray<FUnrealPackageDefinitionInfo*>& PackageDefs, F
 		{
 			if (Plugin->ShouldExportClassesForModule(Module.Name, Module.ModuleType, Module.GeneratedIncludeDirectory))
 			{
-				ExportClassToScriptPlugins(ClassTree.GetClass(), Module, *Plugin);
-				ExportClassTreeToScriptPlugins(&ClassTree, Module, *Plugin);
+				ExportClassToScriptPlugins(SourceFileLookup, ClassTree.GetClass(), Module, *Plugin);
+				ExportClassTreeToScriptPlugins(SourceFileLookup, &ClassTree, Module, *Plugin);
 			}
 		}
 	}
@@ -7003,6 +7020,15 @@ ECompilationResult::Type UnrealHeaderTool_Main(const FString& ModuleInfoFilename
 	// The meta data keywords must be initialized prior to going wide
 	FBaseParser::InitMetadataKeywords();
 
+	// Load the manifest file, giving a list of all modules to be processed, pre-sorted by dependency ordering
+	FResults::Try([&ModuleInfoFilename]() { GManifest = FManifest::LoadFromFile(ModuleInfoFilename); });
+
+	TArray<FUnrealSourceFile*> OrderedSourceFiles;
+	TArray<FUnrealPackageDefinitionInfo*> PackageDefs;
+	PackageDefs.Reserve(GManifest.Modules.Num());
+
+	double TotalPrepareModuleTime = FResults::TimedTry([&PackageDefs, &ModuleInfoPath]() { PrepareModules(PackageDefs, ModuleInfoPath); });
+
 	// TEMPORARY!!!
 	// At this time, there is a collection of classes that aren't listed in the NoExport
 	// file.  This means that we don't have an associated type definition for them.
@@ -7016,14 +7042,6 @@ ECompilationResult::Type UnrealHeaderTool_Main(const FString& ModuleInfoFilename
 		GEngineClasses.Add(ClassDefRef);
 	}
 
-	// Load the manifest file, giving a list of all modules to be processed, pre-sorted by dependency ordering
-	FResults::Try([&ModuleInfoFilename]() { GManifest = FManifest::LoadFromFile(ModuleInfoFilename); });
-
-	TArray<FUnrealSourceFile*> OrderedSourceFiles;
-	TArray<FUnrealPackageDefinitionInfo*> PackageDefs;
-	PackageDefs.Reserve(GManifest.Modules.Num());
-
-	double TotalPrepareModuleTime = FResults::TimedTry([&PackageDefs, &ModuleInfoPath]() { PrepareModules(PackageDefs, ModuleInfoPath); });
 	double TotalPreparseTime = FResults::TimedTry([&PackageDefs, &ModuleInfoPath]() { PreparseSources(PackageDefs, ModuleInfoPath); });
 	double TotalDefineTypesTime = FResults::TimedTry([&PackageDefs]() { DefineTypes(PackageDefs); });
 	double TotalResolveParentsTime = FResults::TimedTry([&PackageDefs]() { ResolveParents(PackageDefs); });
diff --git a/Engine/Source/Programs/UnrealHeaderTool/Private/ParserHelper.h b/Engine/Source/Programs/UnrealHeaderTool/Private/ParserHelper.h
--- a/Engine/Source/Programs/UnrealHeaderTool/Private/ParserHelper.h
+++ b/Engine/Source/Programs/UnrealHeaderTool/Private/ParserHelper.h
@@ -668,26 +668,6 @@ struct FScriptLocation
 	int32 InputLine;
 };
 
-/////////////////////////////////////////////////////
-// FNameLookupCPP
-
-/**
- * Helper class used to cache UClass* -> TCHAR* name lookup for finding the named used for C++ declaration.
- */
-struct FNameLookupCPP
-{
-	/**
-	 * Returns the name used for declaring the passed in struct in C++
-	 *
-	 * @param	Struct	UStruct to obtain C++ name for
-	 * @return	Name used for C++ declaration
-	 */
-	static FString GetNameCPP(UStruct* Struct, bool bForceInterface = false)
-	{
-		return FString::Printf(TEXT("%s%s"), (bForceInterface ? TEXT("I") : Struct->GetPrefixCPP()), *Struct->GetName());
-	}
-};
-
 /////////////////////////////////////////////////////
 // FAdvancedDisplayParameterHandler - used by FHeaderParser::ParseParameterList, to check if a property if a function parameter has 'AdvancedDisplay' flag
 
diff --git a/Engine/Source/Programs/UnrealHeaderTool/Private/PropertyTypes.cpp b/Engine/Source/Programs/UnrealHeaderTool/Private/PropertyTypes.cpp
--- a/Engine/Source/Programs/UnrealHeaderTool/Private/PropertyTypes.cpp
+++ b/Engine/Source/Programs/UnrealHeaderTool/Private/PropertyTypes.cpp
@@ -2717,9 +2717,6 @@ FProperty* FPropertyTraits::CreateEngineType(TSharedRef<FUnrealPropertyDefinitio
 		Property->Next = nullptr;
 		*Prev = Property;
 	}
-
-	// Add the property to the lookup map
-	GTypeDefinitionInfoMap.Add(PropDef.GetProperty(), PropDefRef);
 	return Property;
 }
 
diff --git a/Engine/Source/Programs/UnrealHeaderTool/Private/UnrealTypeDefinitionInfo.cpp b/Engine/Source/Programs/UnrealHeaderTool/Private/UnrealTypeDefinitionInfo.cpp
--- a/Engine/Source/Programs/UnrealHeaderTool/Private/UnrealTypeDefinitionInfo.cpp
+++ b/Engine/Source/Programs/UnrealHeaderTool/Private/UnrealTypeDefinitionInfo.cpp
@@ -936,7 +936,15 @@ FUnrealPackageDefinitionInfo& FUnrealPropertyDefinitionInfo::GetPackageDef() con
 	{
 		return GetUnrealSourceFile().GetPackageDef();
 	}
-	return GTypeDefinitionInfoMap.FindChecked<FUnrealPackageDefinitionInfo>(GetProperty()->GetOutermost());
+	for (FUnrealTypeDefinitionInfo* TempOuter = GetOuter(); TempOuter; TempOuter = TempOuter->GetOuter())
+	{
+		if (FUnrealPackageDefinitionInfo* PackageDef = UHTCast<FUnrealPackageDefinitionInfo>(TempOuter))
+		{
+			return *PackageDef;
+		}
+	}
+	checkf(false, TEXT("Object does not have a package outer"));
+	return *(FUnrealPackageDefinitionInfo*)(nullptr);
 }
 
 bool FUnrealPropertyDefinitionInfo::SameType(const FUnrealPropertyDefinitionInfo& Other) const
@@ -957,7 +965,15 @@ FUnrealPackageDefinitionInfo& FUnrealObjectDefinitionInfo::GetPackageDef() const
 	{
 		return GetUnrealSourceFile().GetPackageDef();
 	}
-	return GTypeDefinitionInfoMap.FindChecked<FUnrealPackageDefinitionInfo>(GetObject()->GetPackage());
+	for (FUnrealTypeDefinitionInfo* TempOuter = GetOuter(); TempOuter; TempOuter = TempOuter->GetOuter())
+	{
+		if (FUnrealPackageDefinitionInfo* PackageDef = UHTCast<FUnrealPackageDefinitionInfo>(TempOuter))
+		{
+			return *PackageDef;
+		}
+	}
+	checkf(false, TEXT("Object does not have a package outer"));
+	return *(FUnrealPackageDefinitionInfo*)(nullptr);
 }
 
 FUnrealObjectDefinitionInfo::FUnrealObjectDefinitionInfo(UObject* InObject)
@@ -1542,7 +1558,6 @@ void FUnrealEnumDefinitionInfo::CreateUObjectEngineTypesInternal(ECreateEngineTy
 		Enum->CppType = CppType;
 		Enum->GetPackage()->GetMetaData()->SetObjectValues(Enum, GetMetaDataMap());
 		SetObject(Enum);
-		GTypeDefinitionInfoMap.AddObjectLookup(Enum, SharedThis(this));
 		break;
 	}
 
@@ -1762,7 +1777,6 @@ void FUnrealScriptStructDefinitionInfo::CreateUObjectEngineTypesInternal(ECreate
 			ScriptStruct->SetSuperStruct(SuperStructDef->GetStruct());
 		}
 		SetObject(ScriptStruct);
-		GTypeDefinitionInfoMap.AddObjectLookup(ScriptStruct, SharedThis(this));
 
 		// The structure needs to be prepared at this point.  
 		ScriptStruct->PrepareCppStructOps();
@@ -1871,24 +1885,7 @@ const TCHAR* FUnrealClassDefinitionInfo::GetPrefixCPP() const
 
 FUnrealClassDefinitionInfo* FUnrealClassDefinitionInfo::FindClass(const TCHAR* ClassName)
 {
-	FUnrealClassDefinitionInfo* ClassDef = nullptr;
-	{
-		ClassDef = GTypeDefinitionInfoMap.FindByName<FUnrealClassDefinitionInfo>(ClassName);
-	}
-
-	if (ClassDef != nullptr)
-	{
-		if (UObjectRedirector* RenamedClassRedirector = FindObject<UObjectRedirector>(ANY_PACKAGE, ClassName))
-		{
-			UClass* RedierClass = CastChecked<UClass>(RenamedClassRedirector->DestinationObject);
-			if (RedierClass != nullptr)
-			{
-				ClassDef = &GTypeDefinitionInfoMap.FindChecked<FUnrealClassDefinitionInfo>(RedierClass);
-			}
-		}
-	}
-
-	return ClassDef;
+	return GTypeDefinitionInfoMap.FindByName<FUnrealClassDefinitionInfo>(ClassName);
 }
 
 FUnrealClassDefinitionInfo* FUnrealClassDefinitionInfo::FindScriptClassOrThrow(const FHeaderParser& Parser, const FString& InClassName)
@@ -2003,9 +2000,6 @@ void FUnrealClassDefinitionInfo::CreateUObjectEngineTypesInternal(ECreateEngineT
 			}
 			SetObject(Class);
 		}
-
-		// Even if we already had an object, add it back to the types regardless
-		GTypeDefinitionInfoMap.AddObjectLookup(GetObjectSafe(), SharedThis(this));
 		break;
 	}
 
@@ -2136,6 +2130,10 @@ void FUnrealClassDefinitionInfo::InitializeFromExistingUObject(UClass* Class)
 	PropertiesSize = Class->PropertiesSize;
 	ClassConfigName = Class->ClassConfigName;
 	SetInternalFlags(Class->GetInternalFlags());
+
+	UPackage* Package = Class->GetOutermost();
+	FUnrealPackageDefinitionInfo& PackageDef = GTypeDefinitionInfoMap.FindByNameChecked<FUnrealPackageDefinitionInfo>(*Package->GetName());
+	SetOuter(&PackageDef);
 }
 
 void FUnrealClassDefinitionInfo::ParseClassProperties(TArray<FPropertySpecifier>&& InClassSpecifiers, const FString& InRequiredAPIMacroIfPresent)
@@ -2760,7 +2758,6 @@ void FUnrealFunctionDefinitionInfo::CreateUObjectEngineTypesInternal(ECreateEngi
 		Function->GetPackage()->GetMetaData()->SetObjectValues(Function, GetMetaDataMap());
 
 		SetObject(Function);
-		GTypeDefinitionInfoMap.Add(Function, SharedThis(this));
 
 		if (FUnrealStructDefinitionInfo* StructDef = UHTCast<FUnrealStructDefinitionInfo>(GetOuter()))
 		{
diff --git a/Engine/Source/Programs/UnrealHeaderTool/Private/UnrealTypeDefinitionInfo.h b/Engine/Source/Programs/UnrealHeaderTool/Private/UnrealTypeDefinitionInfo.h
--- a/Engine/Source/Programs/UnrealHeaderTool/Private/UnrealTypeDefinitionInfo.h
+++ b/Engine/Source/Programs/UnrealHeaderTool/Private/UnrealTypeDefinitionInfo.h
@@ -604,6 +604,11 @@ protected:
 	 */
 	virtual void PostParseFinalizeInternal(EPostParseFinalizePhase Phase) {}
 
+	void SetOuter(FUnrealTypeDefinitionInfo* InOuter)
+	{
+		Outer = InOuter;
+	}
+
 private:
 	FString NameCPP;
 	FUnrealTypeDefinitionInfo* Outer = nullptr;
