Fix UHT build issue on Linux

#rnx

diff --git a/Engine/Source/Programs/UnrealHeaderTool/Private/UnrealTypeDefinitionInfo.cpp b/Engine/Source/Programs/UnrealHeaderTool/Private/UnrealTypeDefinitionInfo.cpp
--- a/Engine/Source/Programs/UnrealHeaderTool/Private/UnrealTypeDefinitionInfo.cpp
+++ b/Engine/Source/Programs/UnrealHeaderTool/Private/UnrealTypeDefinitionInfo.cpp
@@ -936,15 +936,18 @@ FUnrealPackageDefinitionInfo& FUnrealPropertyDefinitionInfo::GetPackageDef() con
 	{
 		return GetUnrealSourceFile().GetPackageDef();
 	}
+
+	FUnrealPackageDefinitionInfo* Out = nullptr;
 	for (FUnrealTypeDefinitionInfo* TempOuter = GetOuter(); TempOuter; TempOuter = TempOuter->GetOuter())
 	{
 		if (FUnrealPackageDefinitionInfo* PackageDef = UHTCast<FUnrealPackageDefinitionInfo>(TempOuter))
 		{
-			return *PackageDef;
+			Out = PackageDef;
+			break;
 		}
 	}
-	checkf(false, TEXT("Object does not have a package outer"));
-	return *(FUnrealPackageDefinitionInfo*)(nullptr);
+	checkf(Out != nullptr, TEXT("Object does not have a package outer"));
+	return *Out;
 }
 
 bool FUnrealPropertyDefinitionInfo::SameType(const FUnrealPropertyDefinitionInfo& Other) const
@@ -965,15 +968,18 @@ FUnrealPackageDefinitionInfo& FUnrealObjectDefinitionInfo::GetPackageDef() const
 	{
 		return GetUnrealSourceFile().GetPackageDef();
 	}
+
+	FUnrealPackageDefinitionInfo* Out = nullptr;
 	for (FUnrealTypeDefinitionInfo* TempOuter = GetOuter(); TempOuter; TempOuter = TempOuter->GetOuter())
 	{
 		if (FUnrealPackageDefinitionInfo* PackageDef = UHTCast<FUnrealPackageDefinitionInfo>(TempOuter))
 		{
-			return *PackageDef;
+			Out = PackageDef;
+			break;
 		}
 	}
-	checkf(false, TEXT("Object does not have a package outer"));
-	return *(FUnrealPackageDefinitionInfo*)(nullptr);
+	checkf(Out != nullptr, TEXT("Object does not have a package outer"));
+	return *Out;
 }
 
 FUnrealObjectDefinitionInfo::FUnrealObjectDefinitionInfo(UObject* InObject)
