Change the Token to use a StringView instead of copying the text directly into two different buffers in the token.  Improves performance of the tokenizer and provides a much smaller structure that can be cached for other uses.

#rnx

diff --git a/Engine/Source/Programs/UnrealHeaderTool/Private/BaseParser.cpp b/Engine/Source/Programs/UnrealHeaderTool/Private/BaseParser.cpp
--- a/Engine/Source/Programs/UnrealHeaderTool/Private/BaseParser.cpp
+++ b/Engine/Source/Programs/UnrealHeaderTool/Private/BaseParser.cpp
@@ -131,10 +131,10 @@ namespace
 		}
 	}
 
-	FMetadataKeyword* GetMetadataKeyword(const TCHAR* Keyword)
+	FMetadataKeyword* GetMetadataKeyword(const FToken& Token)
 	{
 		check(MetadataDictionary.Num() > 0);
-		return MetadataDictionary.Find(Keyword);
+		return MetadataDictionary.Find(FString(Token.Value));
 	}
 }
 
@@ -212,75 +212,49 @@ void FBaseParser::ResetParser(const TCHAR* SourceBuffer, int32 StartingLineNumbe
 //
 TCHAR FBaseParser::GetChar(bool bLiteral)
 {
-	bool bInsideComment = false;
-
 	PrevPos = InputPos;
 	PrevLine = InputLine;
 
 Loop:
-	const TCHAR c = Input[InputPos++];
-	if (bInsideComment)
-	{
-		// Record the character as a comment.
-		PrevComment += c;
+	TCHAR Char = Input[InputPos++];
+	if (Char == 0)
+	{ 
+		--InputPos;
 	}
-
-	if (c == TEXT('\n'))
+	else if (Char == '\n')
 	{
-		InputLine++;
+		++InputLine;
 	}
-	else if (!bLiteral)
+	else if (Char == '/')
 	{
-		const TCHAR NextChar = PeekChar();
-		if ( c==TEXT('/') && NextChar==TEXT('*') )
-		{
-			if (!bInsideComment)
-			{
-				ClearComment();
-				// Record the slash and star.
-				PrevComment += c;
-				PrevComment += NextChar;
-				bInsideComment = true;
-
-				// Move past the star. Do it only when not in comment,
-				// otherwise end of comment might be missed e.g.
-				// /*/ Comment /*/
-				// ~~~~~~~~~~~~~^ Will report second /* as beginning of comment
-				// And throw error that end of file is found in comment.
-				InputPos++;
-			}
-
-			goto Loop;
-		}
-		else if( c==TEXT('*') && NextChar==TEXT('/') )
+		if (!bLiteral && Input[InputPos] == '*')
 		{
-			if (!bInsideComment)
+			ClearComment();
+			const TCHAR* CommentStart = &Input[InputPos - 1];
+			const TCHAR* CommentEnd = CommentStart + 2;
+			for (;;)
 			{
-				ClearComment();
-				Throwf(TEXT("Unexpected '*/' outside of comment") );
+				TCHAR CommentChar = *CommentEnd++;
+				if (CommentChar == 0)
+				{
+					ClearComment();
+					Throwf(TEXT("End of class header encountered inside comment"));
+				}
+				else if (CommentChar == '\n')
+				{
+					++InputLine;
+				}
+				else if (CommentChar == '*' && *CommentEnd == '/')
+				{
+					int32 Length = int32(CommentEnd - CommentStart);
+					PrevComment.AppendChars(CommentStart, Length + 1);
+					InputPos += Length;
+					goto Loop;
+				}
 			}
-
-			/** Asterisk and slash always end comment. */
-			bInsideComment = false;
-
-			// Star already recorded; record the slash.
-			PrevComment += Input[InputPos];
-
-			InputPos++;
-			goto Loop;
 		}
 	}
-
-	if (bInsideComment)
-	{
-		if (c == 0)
-		{
-			ClearComment();
-			Throwf(TEXT("End of class header encountered inside comment") );
-		}
-		goto Loop;
-	}
-	return c;
+	return Char;
 }
 
 //
@@ -301,69 +275,110 @@ TCHAR FBaseParser::PeekChar()
 	return (InputPos < InputLen) ? Input[InputPos] : 0;
 }
 
-//
-// Skip past all spaces and tabs in the input stream.
-//
-TCHAR FBaseParser::GetLeadingChar()
+void FBaseParser::SkipWhitespaceAndComments()
 {
-	TCHAR TrailingCommentNewline = 0;
+	const TCHAR* Pos = &Input[InputPos];
+	bool bGotNewlineBetweenComments = false;
+	bool bGotInlineComment = false;
 
 	for (;;)
 	{
-		bool MultipleNewlines = false;
-
-		TCHAR c;
-
-		// Skip blanks.
-		do
+		TCHAR Char = *Pos++;
+		if (Char == 0)
 		{
-			c = GetChar();
-
-			// Check if we've encountered another newline since the last one
-			if (c == TrailingCommentNewline)
-			{
-				MultipleNewlines = true;
-			}
-		} while (IsWhitespace(c));
-
-		if (c != TEXT('/') || PeekChar() != TEXT('/'))
+			--Pos;
+			break;
+		}
+		else if (Char == '\n')
 		{
-			return c;
+			bGotNewlineBetweenComments |= bGotInlineComment;
+			++InputLine;
 		}
-
-		// Clear the comment if we've encountered newlines since the last comment
-		if (MultipleNewlines)
+		else if (Char == '\r' || Char == '\t' || Char == ' ')
 		{
-			ClearComment();
 		}
-
-		// Record the first slash.  The first iteration of the loop will get the second slash.
-		PrevComment += c;
-
-		do
+		else if (Char == '/')
 		{
-			c = GetChar(true);
-			if (c == 0)
-				return c;
-			PrevComment += c;
-		} while (!IsEOL(c));
-
-		TrailingCommentNewline = c;
+			TCHAR NextChar = *Pos;
+			if (NextChar == '*')
+			{
+				ClearComment();
+				const TCHAR* CommentStart = Pos - 1;
+				++Pos;
+				for (;;)
+				{
+					TCHAR CommentChar = *Pos++;
+					if (CommentChar == 0)
+					{
+						ClearComment();
+						Throwf(TEXT("End of class header encountered inside comment"));
+					}
+					else if (CommentChar == '\n')
+					{
+						++InputLine;
+					}
+					else if (CommentChar == '*' && *Pos == '/')
+					{
+						++Pos;
+						break;
+					}
+				}
+				PrevComment.AppendChars(CommentStart, int32(Pos - CommentStart));
+			}
+			else if (NextChar == '/')
+			{
+				if (bGotNewlineBetweenComments)
+				{
+					bGotNewlineBetweenComments = false;
+					ClearComment();
+				}
+				bGotInlineComment = true;
+				const TCHAR* CommentStart = Pos - 1;
+				Pos++;
 
-		for (;;)
-		{
-			c = GetChar();
-			if (c == 0)
-				return c;
-			if (c == TrailingCommentNewline || !IsEOL(c))
+				// Scan to the end of the line
+				for (;;)
+				{
+					TCHAR CommentChar = *Pos++;
+					if (CommentChar == 0)
+					{
+						--Pos;
+						break;
+					}
+					if (CommentChar == '\r')
+					{
+					} 
+					else if (CommentChar == '\n')
+					{
+						++InputLine;
+						break;
+					}
+				}
+				PrevComment.AppendChars(CommentStart, int32(Pos - CommentStart));
+			}
+			else
 			{
-				UngetChar();
+				--Pos;
 				break;
 			}
-
-			PrevComment += c;
+		}
+		else
+		{
+			--Pos;
+			break;
 		}
 	}
+	InputPos = int32(Pos - Input);
+	return;
+}
+
+//
+// Skip past all spaces and tabs in the input stream.
+//
+TCHAR FBaseParser::GetLeadingChar()
+{
+	SkipWhitespaceAndComments();
+	return GetChar();
 }
 
 bool FBaseParser::IsEOL( TCHAR c )
@@ -381,148 +396,130 @@ bool FBaseParser::IsWhitespace( TCHAR c )
 -----------------------------------------------------------------------------*/
 
 // Gets the next token from the input stream, advancing the variables which keep track of the current input position and line.
-bool FBaseParser::GetToken( FToken& Token, bool bNoConsts/*=false*/, ESymbolParseOption bParseTemplateClosingBracket/*=ESymbolParseOption::Normal*/ )
+// NOTE: There is a slight difference between StartPos/StartLine handling here.  In the old version, the starting location would 
+// include any preceeding /**/ comments (but not // comments) while this version does not.  The only place in the code this
+// caused an issue was in the processing of default arguments.
+bool FBaseParser::GetToken(FToken& Token, bool bNoConsts/*=false*/, ESymbolParseOption ParseTemplateClosingBracket/*=ESymbolParseOption::Normal*/ )
 {
-	TCHAR c = GetLeadingChar();
-	if( c == 0 )
+	SkipWhitespaceAndComments();
+	PrevPos = InputPos;
+	PrevLine = InputLine;
+	if (Input[InputPos] == 0)
 	{
-		UngetChar();
-		return 0;
+		return false;
 	}
-	TCHAR p = PeekChar();
-	Token.StartPos		= PrevPos;
-	Token.StartLine		= PrevLine;
-	if( (c>='A' && c<='Z') || (c>='a' && c<='z') || (c=='_') )
+	Token.StartPos = InputPos;
+	Token.StartLine = InputLine;
+
+	const TCHAR* Start = &Input[InputPos];
+	const TCHAR* Pos = Start;
+	TCHAR c = *Pos++;
+
+	// Identifier
+	if ((c >= 'A' && c <= 'Z') || (c >= 'a' && c <= 'z') || c == '_')
 	{
-		// Alphanumeric token.
-		int32 Length=0;
-		do
+		for (;; ++Pos)
 		{
-			Token.Identifier[Length++] = c;
-			if( Length >= NAME_SIZE )
+			c = *Pos;
+			if (!((c >= 'A' && c <= 'Z') || (c >= 'a' && c <= 'z') || (c >= '0' && c <= '9') || c == '_'))
 			{
-				Throwf(TEXT("Identifer length exceeds maximum of %i"), (int32)NAME_SIZE);
-				Length = ((int32)NAME_SIZE) - 1;
 				break;
 			}
-			c = GetChar();
-		} while( ((c>='A')&&(c<='Z')) || ((c>='a')&&(c<='z')) || ((c>='0')&&(c<='9')) || (c=='_') );
-		UngetChar();
-		Token.Identifier[Length]=0;
-
-		// Assume this is an identifier unless we find otherwise.
-		Token.TokenType = TOKEN_Identifier;
+		}
 
 		// If const values are allowed, determine whether the identifier represents a constant
-		if ( !bNoConsts )
+		Token.TokenType = ETokenType::Identifier;
+		Token.Value = FStringView(Start, int32(Pos - Start));
+		if (!bNoConsts)
 		{
 			// See if the identifier is part of a vector, rotation or other struct constant.
 			// boolean true/false
-			if( Token.Matches(TEXT("true"), ESearchCase::IgnoreCase) )
+			int32 Length = int32(Pos - Start);
+			if (Length == 4 && Token.IsValue(TEXT("true"), ESearchCase::IgnoreCase))
 			{
-				Token.SetConstBool(true);
-				return true;
+				Token.TokenType = ETokenType::TrueConst;
 			}
-			else if( Token.Matches(TEXT("false"), ESearchCase::IgnoreCase) )
+			else if (Length == 5 && Token.IsValue(TEXT("false"), ESearchCase::IgnoreCase))
 			{
-				Token.SetConstBool(false);
-				return true;
+				Token.TokenType = ETokenType::FalseConst;
 			}
 		}
+
+		if (Token.Value.Len() >= NAME_SIZE)
+		{
+			Throwf(TEXT("Identifer length exceeds maximum of %i"), (int32)NAME_SIZE);
+		}
+
+		InputPos = int32(Pos - Input);
 		return true;
 	}
 
 	// if const values are allowed, determine whether the non-identifier token represents a const
-	else if ( !bNoConsts && ((c>='0' && c<='9') || ((c=='+' || c=='-') && (p>='0' && p<='9'))) )
+	else if (!bNoConsts && ((c >= '0' && c <= '9') || ((c == '+' || c == '-') && (*Pos >= '0' && *Pos <= '9'))))
 	{
 		// Integer or floating point constant.
+
 		bool  bIsFloat = 0;
-		int32 Length   = 0;
-		bool  bIsHex   = 0;
+		bool  bIsHex = 0;
 		do
 		{
-			if( c==TEXT('.') )
+			if (c == TEXT('.'))
 			{
 				bIsFloat = true;
 			}
-			if( c==TEXT('X') || c == TEXT('x') )
+			if (c == TEXT('X') || c == TEXT('x'))
 			{
 				bIsHex = true;
 			}
-
-			Token.Identifier[Length++] = c;
-			if( Length >= NAME_SIZE )
-			{
-				Throwf(TEXT("Number length exceeds maximum of %i "), (int32)NAME_SIZE );
-				Length = ((int32)NAME_SIZE) - 1;
-				break;
-			}
-			c = FChar::ToUpper(GetChar());
+			c = FChar::ToUpper(*Pos++);
 		} while ((c >= TEXT('0') && c <= TEXT('9')) || (!bIsFloat && c == TEXT('.')) || (!bIsHex && c == TEXT('X')) || (bIsHex && c >= TEXT('A') && c <= TEXT('F')));
 
-		Token.Identifier[Length]=0;
 		if (!bIsFloat || c != 'F')
 		{
-			UngetChar();
+			--Pos;
 		}
 
-		if (bIsFloat)
-		{
-			Token.SetConstFloat( FCString::Atof(Token.Identifier) );
-		}
-		else if (bIsHex)
-		{
-			TCHAR* End = Token.Identifier + FCString::Strlen(Token.Identifier);
-			Token.SetConstInt64( FCString::Strtoi64(Token.Identifier,&End,0) );
-		}
-		else
+		Token.TokenType = bIsFloat ? ETokenType::FloatConst : (bIsHex ? ETokenType::HexConst : ETokenType::DecimalConst);
+		Token.Value = FStringView(Start, int32(Pos - Start));
+
+		if (Token.Value.Len() >= NAME_SIZE)
 		{
-			Token.SetConstInt64( FCString::Atoi64(Token.Identifier) );
+			Throwf(TEXT("Number length exceeds maximum of %i "), (int32)NAME_SIZE);
 		}
+
+		InputPos = int32(Pos - Input);
 		return true;
 	}
+
+	// Escaped character constant
 	else if (c == '\'')
 	{
-		TCHAR ActualCharLiteral = GetChar(/*bLiteral=*/ true);
-
-		if (ActualCharLiteral == '\\')
+		if (*Pos++ == '\\')
 		{
-			ActualCharLiteral = GetChar(/*bLiteral=*/ true);
-			switch (ActualCharLiteral)
-			{
-			case TCHAR('t'):
-				ActualCharLiteral = '\t';
-				break;
-			case TCHAR('n'):
-				ActualCharLiteral = '\n';
-				break;
-			case TCHAR('r'):
-				ActualCharLiteral = '\r';
-				break;
-			}
+			*Pos++;
 		}
 
-		c = GetChar(/*bLiteral=*/ true);
-		if (c != '\'')
+		if (*Pos++ != '\'')
 		{
 			Throwf(TEXT("Unterminated character constant"));
-			UngetChar();
 		}
 
-		Token.SetConstChar(ActualCharLiteral);
+		Token.TokenType = ETokenType::CharConst;
+		Token.Value = FStringView(Start, int32(Pos - Start));
+		InputPos = int32(Pos - Input);
 		return true;
 	}
+
+	// String contant
 	else if (c == '"')
 	{
-		// String constant.
-		TCHAR Temp[MAX_STRING_CONST_SIZE];
-		int32 Length=0;
-		c = GetChar(/*bLiteral=*/ true);
+		c = *Pos++;
+		int32 EscapeCount = 0;
 		while( (c!='"') && !IsEOL(c) )
 		{
 			if( c=='\\' )
 			{
-				c = GetChar(/*bLiteral=*/ true);
+				c = *Pos++;
 				if( IsEOL(c) )
 				{
 					break;
@@ -532,40 +529,35 @@ bool FBaseParser::GetToken( FToken& Token, bool bNoConsts/*=false*/, ESymbolPars
 					// Newline escape sequence.
 					c = '\n';
 				}
+				++EscapeCount;
 			}
-			Temp[Length++] = c;
-			if( Length >= MAX_STRING_CONST_SIZE )
-			{
-				Throwf(TEXT("String constant exceeds maximum of %i characters"), (int32)MAX_STRING_CONST_SIZE );
-				c = TEXT('\"');
-				Length = ((int32)MAX_STRING_CONST_SIZE) - 1;
-				break;
-			}
-			c = GetChar(/*bLiteral=*/ true);
+			c = *Pos++;
 		}
-		Temp[Length]=0;
 
 		if( c != '"' )
 		{
-			Throwf(TEXT("Unterminated string constant: %s"), Temp);
-			UngetChar();
+			Throwf(TEXT("Unterminated string constant"));
 		}
 
-		Token.SetConstString(Temp);
+		Token.TokenType = ETokenType::StringConst;
+		Token.Value = FStringView(Start, int32(Pos - Start));
+
+		if (!Token.Value.Len() - EscapeCount >= MAX_STRING_CONST_SIZE)
+		{
+			Throwf(TEXT("String constant exceeds maximum of %i characters"), (int32)MAX_STRING_CONST_SIZE);
+		}
+
+		InputPos = int32(Pos - Input);
 		return true;
 	}
 	else
 	{
-		// Symbol.
-		int32 Length=0;
-		Token.Identifier[Length++] = c;
-
 		// Handle special 2-character symbols.
-		#define PAIR(cc,dd) ((c==cc)&&(d==dd)) /* Comparison macro for convenience */
-		TCHAR d = GetChar();
+#define PAIR(cc,dd) ((c==cc)&&(d==dd)) /* Comparison macro for convenience */
+		TCHAR d = *Pos++;
 		if
 		(	PAIR('<','<')
-		||	(PAIR('>','>') && (bParseTemplateClosingBracket != ESymbolParseOption::CloseTemplateBracket))
+		||	(PAIR('>','>') && (ParseTemplateClosingBracket != ESymbolParseOption::CloseTemplateBracket))
 		||	PAIR('!','=')
 		||	PAIR('<','=')
 		||	PAIR('>','=')
@@ -584,36 +576,27 @@ bool FBaseParser::GetToken( FToken& Token, bool bNoConsts/*=false*/, ESymbolPars
 		||	PAIR(':',':')
 		)
 		{
-			Token.Identifier[Length++] = d;
-			if( c=='>' && d=='>' )
+			if (c=='>' && d=='>' &&  *Pos == '>')
 			{
-				if (GetChar()=='>')
-				{
-					Token.Identifier[Length++] = '>';
-				}
-				else
-				{
-					UngetChar();
-				}
+				++Pos;
 			}
 		}
 		else
 		{
-			UngetChar();
+			--Pos;
 		}
 		#undef PAIR
 
-		Token.Identifier[Length] = 0;
-		Token.TokenType = TOKEN_Symbol;
-
+		Token.TokenType = ETokenType::Symbol;
+		Token.Value = FStringView(Start, int32(Pos - Start));
+		InputPos = int32(Pos - Input);
 		return true;
 	}
 }
 
-bool FBaseParser::GetRawTokenRespectingQuotes( FToken& Token, TCHAR StopChar /* = TCHAR('\n') */ )
+bool FBaseParser::GetRawStringRespectingQuotes(FTokenString& String, TCHAR StopChar /* = TCHAR('\n') */ )
 {
 	// Get token after whitespace.
-	TCHAR Temp[MAX_STRING_CONST_SIZE];
 	int32 Length=0;
 	TCHAR c = GetLeadingChar();
 
@@ -631,10 +614,10 @@ bool FBaseParser::GetRawTokenRespectingQuotes( FToken& Token, TCHAR StopChar /*
 			bInQuote = !bInQuote;
 		}
 
-		Temp[Length++] = c;
-		if( Length >= MAX_STRING_CONST_SIZE )
+		String.String[Length++] = c;
+		if( Length >= MAX_STRING_CONST_SIZE)
 		{
-			Throwf(TEXT("Identifier exceeds maximum of %i characters"), (int32)MAX_STRING_CONST_SIZE );
+			Throwf(TEXT("Identifier exceeds maximum of %i characters"), (int32)MAX_STRING_CONST_SIZE);
 			c = GetChar(true);
 			Length = ((int32)MAX_STRING_CONST_SIZE) - 1;
 			break;
@@ -649,13 +632,11 @@ bool FBaseParser::GetRawTokenRespectingQuotes( FToken& Token, TCHAR StopChar /*
 	}
 
 	// Get rid of trailing whitespace.
-	while( Length>0 && (Temp[Length-1]==' ' || Temp[Length-1]==9 ) )
+	while( Length>0 && (String.String[Length-1]==' ' || String.String[Length-1]==9 ) )
 	{
 		Length--;
 	}
-	Temp[Length]=0;
-
-	Token.SetConstString(Temp);
+	String.String[Length]=0;
 
 	return Length>0;
 }
@@ -669,10 +650,9 @@ bool FBaseParser::GetRawTokenRespectingQuotes( FToken& Token, TCHAR StopChar /*
  *
  * @return	the number of character parsed
  */
-bool FBaseParser::GetRawToken( FToken& Token, TCHAR StopChar /* = TCHAR('\n') */ )
+bool FBaseParser::GetRawString( FTokenString& String, TCHAR StopChar /* = TCHAR('\n') */ )
 {
 	// Get token after whitespace.
-	TCHAR Temp[MAX_STRING_CONST_SIZE];
 	int32 Length=0;
 	TCHAR c = GetLeadingChar();
 	while( !IsEOL(c) && c != StopChar )
@@ -681,7 +661,7 @@ bool FBaseParser::GetRawToken( FToken& Token, TCHAR StopChar /* = TCHAR('\n') */
 		{
 			break;
 		}
-		Temp[Length++] = c;
+		String.String[Length++] = c;
 		if( Length >= MAX_STRING_CONST_SIZE )
 		{
 			Throwf(TEXT("Identifier exceeds maximum of %i characters"), (int32)MAX_STRING_CONST_SIZE );
@@ -691,13 +671,11 @@ bool FBaseParser::GetRawToken( FToken& Token, TCHAR StopChar /* = TCHAR('\n') */
 	UngetChar();
 
 	// Get rid of trailing whitespace.
-	while( Length>0 && (Temp[Length-1]==' ' || Temp[Length-1]==9 ) )
+	while( Length>0 && (String.String[Length-1]==' ' || String.String[Length-1]==9 ) )
 	{
 		Length--;
 	}
-	Temp[Length]=0;
-
-	Token.SetConstString(Temp);
+	String.String[Length]=0;
 
 	return Length>0;
 }
@@ -705,33 +683,14 @@ bool FBaseParser::GetRawToken( FToken& Token, TCHAR StopChar /* = TCHAR('\n') */
 //
 // Get an identifier token, return 1 if gotten, 0 if not.
 //
-bool FBaseParser::GetIdentifier( FToken& Token, bool bNoConsts )
+bool FBaseParser::GetIdentifier(FToken& Token, bool bNoConsts)
 {
 	if (!GetToken(Token, bNoConsts))
 	{
 		return false;
 	}
 
-	if (Token.TokenType == TOKEN_Identifier)
-	{
-		return true;
-	}
-
-	UngetToken(Token);
-	return false;
-}
-
-//
-// Get a symbol token, return 1 if gotten, 0 if not.
-//
-bool FBaseParser::GetSymbol( FToken& Token )
-{
-	if (!GetToken(Token))
-	{
-		return false;
-	}
-
-	if( Token.TokenType == TOKEN_Symbol )
+	if (Token.IsIdentifier())
 	{
 		return true;
 	}
@@ -743,12 +702,12 @@ bool FBaseParser::GetSymbol( FToken& Token )
 // Modify token to fix redirected types if needed
 void FBaseParser::RedirectTypeIdentifier(FToken& Token) const
 {
-	check(Token.TokenType == TOKEN_Identifier);
+	check(Token.IsIdentifier());
 
-	const FString* FoundRedirect = UHTConfig.TypeRedirectMap.Find(Token.Identifier);
+	const FString* FoundRedirect = UHTConfig.TypeRedirectMap.Find(FString(Token.Value));
 	if (FoundRedirect)
 	{
-		Token.SetIdentifier(**FoundRedirect);
+		Token.Value = *FoundRedirect;
 	}
 }
 
@@ -804,7 +763,7 @@ bool FBaseParser::MatchSymbol( const TCHAR Match, ESymbolParseOption bParseTempl
 
 	if (GetToken(Token, /*bNoConsts=*/ true, bParseTemplateClosingBracket))
 	{
-		if (Token.TokenType==TOKEN_Symbol && Token.Identifier[0] == Match && Token.Identifier[1] == 0)
+		if (Token.IsSymbol(Match))
 		{
 			return true;
 		}
@@ -823,7 +782,7 @@ bool FBaseParser::MatchSymbol(const TCHAR* Match, ESymbolParseOption bParseTempl
 
 	if (GetToken(Token, /*bNoConsts=*/ true, bParseTemplateClosingBracket))
 	{
-		if (Token.TokenType==TOKEN_Symbol && !FCString::Strcmp(Token.Identifier, Match))
+		if (Token.IsSymbol(Match, ESearchCase::CaseSensitive))
 		{
 			return true;
 		}
@@ -841,7 +800,7 @@ bool FBaseParser::MatchIdentifier( const TCHAR* Match, ESearchCase::Type SearchC
 	FToken Token;
 	if (GetToken(Token))
 	{
-		if (Token.TokenType==TOKEN_Identifier && ((SearchCase == ESearchCase::CaseSensitive) ? !FCString::Strcmp(Token.Identifier, Match) : !FCString::Stricmp(Token.Identifier, Match)))
+		if (Token.IsIdentifier(Match, SearchCase))
 		{
 			return true;
 		}
@@ -859,7 +818,7 @@ bool FBaseParser::MatchConstInt( const TCHAR* Match )
 	FToken Token;
 	if (GetToken(Token))
 	{
-		if( Token.TokenType==TOKEN_Const && (Token.ConstantType == CPT_Int || Token.ConstantType == CPT_Int64) && FCString::Strcmp(Token.Identifier,Match)==0 )
+		if (Token.IsConstInt(Match))
 		{
 			return true;
 		}
@@ -877,7 +836,7 @@ bool FBaseParser::MatchAnyConstInt()
 	FToken Token;
 	if (GetToken(Token))
 	{
-		if( Token.TokenType==TOKEN_Const && (Token.ConstantType == CPT_Int || Token.ConstantType == CPT_Int64) )
+		if (Token.IsConstInt())
 		{
 			return true;
 		}
@@ -897,7 +856,7 @@ void FBaseParser::MatchSemi()
 		FToken Token;
 		if( GetToken(Token) )
 		{
-			Throwf(TEXT("Missing ';' before '%s'"), Token.Identifier );
+			Throwf(TEXT("Missing ';' before '%s'"), *FString(Token.Value));
 		}
 		else
 		{
@@ -919,7 +878,7 @@ bool FBaseParser::PeekSymbol( const TCHAR Match )
 	}
 	UngetToken(Token);
 
-	return Token.TokenType==TOKEN_Symbol && Token.Identifier[0] == Match && Token.Identifier[1] == 0;
+	return Token.IsSymbol(Match);
 }
 
 //
@@ -933,7 +892,7 @@ bool FBaseParser::PeekIdentifier( const TCHAR* Match, ESearchCase::Type SearchCa
 		return false;
 	}
 	UngetToken(Token);
-	return Token.TokenType==TOKEN_Identifier && ((SearchCase == ESearchCase::CaseSensitive) ? !FCString::Strcmp(Token.Identifier, Match) : !FCString::Stricmp(Token.Identifier, Match));
+	return Token.IsIdentifier(Match, SearchCase);
 }
 
 //
@@ -962,6 +921,14 @@ void FBaseParser::RequireSymbol( const TCHAR Match, const TCHAR* Tag, ESymbolPar
 	}
 }
 
+void FBaseParser::RequireSymbol(const TCHAR Match, const FStringView& Tag, ESymbolParseOption bParseTemplateClosingBracket/*=ESymbolParseOption::Normal*/)
+{
+	if (!MatchSymbol(Match, bParseTemplateClosingBracket))
+	{
+		Throwf(TEXT("Missing '%c' in %s"), Match, *FString(Tag));
+	}
+}
+
 void FBaseParser::RequireSymbol(const TCHAR Match, TFunctionRef<FString()> ErrorGetter, ESymbolParseOption bParseTemplateClosingBracket/*=ESymbolParseOption::Normal*/)
 {
 	if (!MatchSymbol(Match, bParseTemplateClosingBracket))
@@ -1018,24 +985,19 @@ FString FBaseParser::ReadNewStyleValue(const TCHAR* TypeOfSpecifier)
 
 	FString Result;
 
-	switch (ValueToken.TokenType)
+	if (ValueToken.IsIdentifier() || ValueToken.IsSymbol())
 	{
-		case TOKEN_Identifier:
-		case TOKEN_Symbol:
-
-			Result = ValueToken.Identifier;
+		Result = FString(ValueToken.Value);
 
-			if (MatchSymbol(TEXT('=')))
-			{
-				Result += TEXT("=");
-				Result += ReadNewStyleValue(TypeOfSpecifier);
-			}
-			break;
-
-		case TOKEN_Const:
-
-			Result = ValueToken.GetConstantValue();
-			break;
+		if (MatchSymbol(TEXT('=')))
+		{
+			Result += TEXT("=");
+			Result += ReadNewStyleValue(TypeOfSpecifier);
+		}
+	}
+	else
+	{
+		Result = ValueToken.GetConstantValue();
 	}
 
 	return Result;
@@ -1072,7 +1034,7 @@ void FBaseParser::ParseNameWithPotentialAPIMacroPrefix(FString& DeclaredName, FS
 	}
 
 	// Is the identifier the name or an DLL import/export API macro?
-	FString NameTokenStr = NameToken.Identifier;
+	FString NameTokenStr(NameToken.Value);
 	if (NameTokenStr.EndsWith(TEXT("_API"), ESearchCase::CaseSensitive))
 	{
 		RequiredAPIMacroIfPresent = MoveTemp(NameTokenStr);
@@ -1082,7 +1044,7 @@ void FBaseParser::ParseNameWithPotentialAPIMacroPrefix(FString& DeclaredName, FS
 		{
 			Throwf(TEXT("Missing %s name"), FailureMessage);
 		}
-		DeclaredName = NameToken.Identifier;
+		DeclaredName = FString(NameToken.Value);
 	}
 	else
 	{
@@ -1115,7 +1077,7 @@ void FBaseParser::ReadSpecifierSetInsideMacro(TArray<FPropertySpecifier>& Specif
 			Throwf(TEXT("Expected %s"), *ErrorMessageGetter());
 		}
 
-		if (Specifier.Matches(TEXT("meta"), ESearchCase::IgnoreCase))
+		if (Specifier.IsIdentifier(TEXT("meta"), ESearchCase::IgnoreCase))
 		{
 			RequireSymbol(TEXT('='), ErrorMessageGetter);
 			RequireSymbol(TEXT('('), ErrorMessageGetter);
@@ -1130,7 +1092,7 @@ void FBaseParser::ReadSpecifierSetInsideMacro(TArray<FPropertySpecifier>& Specif
 					Throwf(TEXT("Expected a metadata key"));
 				}
 
-				FString Key = MetaKeyToken.Identifier;
+				FString Key(MetaKeyToken.Value);
 
 				// Potentially read a value
 				FString Value;
@@ -1146,13 +1108,13 @@ void FBaseParser::ReadSpecifierSetInsideMacro(TArray<FPropertySpecifier>& Specif
 			RequireSymbol(TEXT(')'), ErrorMessageGetter);
 		}
 		// Look up specifier in metadata dictionary
-		else if (FMetadataKeyword* MetadataKeyword = GetMetadataKeyword(Specifier.Identifier))
+		else if (FMetadataKeyword* MetadataKeyword = GetMetadataKeyword(Specifier))
 		{
 			if (MatchSymbol(TEXT('=')))
 			{
 				if (MetadataKeyword->ValueArgument == EMetadataValueArgument::None)
 				{
-					Throwf(TEXT("Incorrect = after metadata specifier '%s'"), Specifier.Identifier);
+					Throwf(TEXT("Incorrect = after metadata specifier '%s'"), *Specifier.GetTokenValue());
 				}
 
 				FString Value = ReadNewStyleValue(TypeOfSpecifier);
@@ -1162,7 +1124,7 @@ void FBaseParser::ReadSpecifierSetInsideMacro(TArray<FPropertySpecifier>& Specif
 			{
 				if (MetadataKeyword->ValueArgument == EMetadataValueArgument::Required)
 				{
-					Throwf(TEXT("Missing = after metadata specifier '%s'"), Specifier.Identifier);
+					Throwf(TEXT("Missing = after metadata specifier '%s'"), *Specifier.GetTokenValue());
 				}
 
 				MetadataKeyword->ApplyToMetadata(*this, MetaData);
@@ -1171,7 +1133,7 @@ void FBaseParser::ReadSpecifierSetInsideMacro(TArray<FPropertySpecifier>& Specif
 		else
 		{
 			// Creating a new specifier
-			SpecifiersFound.Emplace(Specifier.Identifier);
+			SpecifiersFound.Emplace(FString(Specifier.Value));
 
 			// Look for a value for this specifier
 			if (MatchSymbol(TEXT('=')) || PeekSymbol(TEXT('(')))
@@ -1231,41 +1193,41 @@ EUnderlyingEnumType FBaseParser::ParseUnderlyingEnumType()
 			Throwf(TEXT("Missing enum base"));
 		}
 
-		if (!FCString::Strcmp(BaseToken.Identifier, TEXT("uint8")))
+		if (BaseToken.IsValue(TEXT("uint8"), ESearchCase::CaseSensitive))
 		{
 			UnderlyingType = EUnderlyingEnumType::uint8;
 		}
-		else if (!FCString::Strcmp(BaseToken.Identifier, TEXT("uint16")))
+		else if (BaseToken.IsValue(TEXT("uint16"), ESearchCase::CaseSensitive))
 		{
 			UnderlyingType = EUnderlyingEnumType::uint16;
 		}
-		else if (!FCString::Strcmp(BaseToken.Identifier, TEXT("uint32")))
+		else if (BaseToken.IsValue(TEXT("uint32"), ESearchCase::CaseSensitive))
 		{
 			UnderlyingType = EUnderlyingEnumType::uint32;
 		}
-		else if (!FCString::Strcmp(BaseToken.Identifier, TEXT("uint64")))
+		else if (BaseToken.IsValue(TEXT("uint64"), ESearchCase::CaseSensitive))
 		{
 			UnderlyingType = EUnderlyingEnumType::uint64;
 		}
-		else if (!FCString::Strcmp(BaseToken.Identifier, TEXT("int8")))
+		else if (BaseToken.IsValue(TEXT("int8"), ESearchCase::CaseSensitive))
 		{
 			UnderlyingType = EUnderlyingEnumType::int8;
 		}
-		else if (!FCString::Strcmp(BaseToken.Identifier, TEXT("int16")))
+		else if (BaseToken.IsValue(TEXT("int16"), ESearchCase::CaseSensitive))
 		{
 			UnderlyingType = EUnderlyingEnumType::int16;
 		}
-		else if (!FCString::Strcmp(BaseToken.Identifier, TEXT("int32")))
+		else if (BaseToken.IsValue(TEXT("int32"), ESearchCase::CaseSensitive))
 		{
 			UnderlyingType = EUnderlyingEnumType::int32;
 		}
-		else if (!FCString::Strcmp(BaseToken.Identifier, TEXT("int64")))
+		else if (BaseToken.IsValue(TEXT("int64"), ESearchCase::CaseSensitive))
 		{
 			UnderlyingType = EUnderlyingEnumType::int64;
 		}
 		else
 		{
-			Throwf(TEXT("Unsupported enum class base type: %s"), BaseToken.Identifier);
+			Throwf(TEXT("Unsupported enum class base type: %s"), *BaseToken.GetTokenValue());
 		}
 	}
 	else
diff --git a/Engine/Source/Programs/UnrealHeaderTool/Private/BaseParser.h b/Engine/Source/Programs/UnrealHeaderTool/Private/BaseParser.h
--- a/Engine/Source/Programs/UnrealHeaderTool/Private/BaseParser.h
+++ b/Engine/Source/Programs/UnrealHeaderTool/Private/BaseParser.h
@@ -8,8 +8,6 @@
 #include "ParserHelper.h"
 #include "Containers/UnrealString.h"
 
-class FToken;
-
 enum class EPointerMemberBehavior
 {
 	Disallow,
@@ -65,6 +63,317 @@ private:
 	FUHTConfig();
 };
 
+/////////////////////////////////////////////////////
+// Token types.
+
+enum class ETokenType
+{
+	None,			// No token.
+	Identifier,		// Alphanumeric identifier.
+	Symbol,			// Symbol.
+	TrueConst,		// True value via identifier
+	FalseConst,		// False value via identifier
+	FloatConst,		// Floating point constant
+	DecimalConst,	// Decimal Integer constant
+	HexConst,		// Hex integer constant
+	CharConst,		// Single character constant
+	StringConst,	// String constant
+};
+
+// Helper structure for strings.  Only valid when called on a StringConst or CharConst.
+// The surrounding quotes are removed and \n escape characters are processed.
+struct FTokenString
+{
+	TCHAR String[MAX_STRING_CONST_SIZE];
+
+	FTokenString() { String[0] = 0; }
+
+	const TCHAR* operator *() const
+	{
+		return String;
+	}
+};
+
+// Helper structure for null terminated value.  This will be the raw text from the token.
+struct FTokenValue
+{
+	static constexpr int32 MAX_SIZE = std::max(int32(MAX_STRING_CONST_SIZE), int32(NAME_SIZE));
+
+	TCHAR Value[MAX_SIZE];
+
+	FTokenValue() { Value[0] = 0; }
+
+	const TCHAR* operator *() const
+	{
+		return Value;
+	}
+};
+
+/////////////////////////////////////////////////////
+// FToken
+
+class FToken
+{
+public:
+	/** Type of token. */
+	ETokenType TokenType = ETokenType::None;
+
+	/** Starting position in script where this token came from. */
+	int32 StartPos = 0;
+
+	/** Starting line in script. */
+	int32 StartLine = 0;
+
+	/** The text of the token */
+	FStringView Value;
+
+	// Return true if the token is an identifier
+	bool IsIdentifier() const
+	{
+		return TokenType == ETokenType::Identifier;
+	}
+
+	// Test to see if the token value matches the given string
+	bool IsValue(const TCHAR* Str, ESearchCase::Type SearchCase) const
+	{
+		return Value.Equals(Str, SearchCase);
+	}
+
+	// Test to see if the token value begins with the given string
+	bool ValueStartsWith(const TCHAR* Str, ESearchCase::Type SearchCase) const
+	{
+		return Value.StartsWith(Str, SearchCase);
+	}
+
+	// Return true if the token is a specific identifier
+	bool IsIdentifier(const TCHAR* Str, ESearchCase::Type SearchCase) const
+	{
+		return IsIdentifier() && IsValue(Str, SearchCase);
+	}
+
+	// Return true if the token is a symbol
+	bool IsSymbol() const
+	{
+		return TokenType == ETokenType::Symbol;
+	}
+
+	// Return true if the token is a specific single character symbol
+	bool IsSymbol(const TCHAR Ch) const
+	{
+		return IsSymbol() && Value.Len() == 1 && Value[0] == Ch;
+	}
+
+	// Return true if the token is a specific symbol
+	bool IsSymbol(const TCHAR* Str, ESearchCase::Type SearchCase) const
+	{
+		return IsSymbol() && IsValue(Str, SearchCase);
+	}
+
+	// Return true if the token is a decimal or hexidecimal constant
+	bool IsConstInt() const
+	{
+		return TokenType == ETokenType::DecimalConst || TokenType == ETokenType::HexConst;
+	}
+
+	// Return true if the token is a specific constant integer
+	bool IsConstInt(const TCHAR* Str) const
+	{
+		return IsConstInt() && IsValue(Str, ESearchCase::CaseSensitive);
+	}
+
+	// Return true if the token is a floating point constant
+	bool IsConstFloat() const
+	{
+		return TokenType == ETokenType::FloatConst;
+	}
+
+	// Return true if the token is a string or character constant
+	bool IsConstString() const
+	{
+		return TokenType == ETokenType::StringConst || TokenType == ETokenType::CharConst;
+	}
+
+	// Get the 32 bit integer value of the token.  Only supported for decimal, hexidecimal, and floating point values
+	bool GetConstInt(int32& I) const
+	{
+		switch (TokenType)
+		{
+		case ETokenType::DecimalConst:
+			I = (int32)GetDecimalValue();
+			return true;
+		case ETokenType::HexConst:
+			I = (int32)GetHexValue();
+			return true;
+		case ETokenType::FloatConst:
+		{
+			float Float = GetFloatValue();
+			if (Float == FMath::TruncToInt(Float))
+			{
+				I = (int32)Float;
+				return true;
+			}
+			return false;
+		}
+		default:
+			return false;
+		}
+	}
+
+	// Get the 64 bit integer value of the token.  Only supported for decimal, hexidecimal, and floating point values
+	bool GetConstInt64(int64& I) const
+	{
+		switch (TokenType)
+		{
+		case ETokenType::DecimalConst:
+			I = GetDecimalValue();
+			return true;
+		case ETokenType::HexConst:
+			I = GetHexValue();
+			return true;
+		case ETokenType::FloatConst:
+		{
+			float Float = GetFloatValue();
+			if (Float == FMath::TruncToInt(Float))
+			{
+				I = (int32)Float;
+				return true;
+			}
+			return false;
+		}
+		default:
+			return false;
+		}
+	}
+
+	// Return a string representation of a constant value.  In the case of strings, the surrounding quotes and \n escapres are removed/converted
+	FString GetConstantValue() const
+	{
+		switch (TokenType)
+		{
+		case ETokenType::DecimalConst:
+			return FString::Printf(TEXT("%" INT64_FMT), GetDecimalValue());
+		case ETokenType::HexConst:
+			return FString::Printf(TEXT("%" INT64_FMT), GetHexValue());
+		case ETokenType::FloatConst:
+			return FString::Printf(TEXT("%f"), GetFloatValue());
+		case ETokenType::TrueConst:
+			return FName::GetEntry(NAME_TRUE)->GetPlainNameString();
+		case ETokenType::FalseConst:
+			return FName::GetEntry(NAME_FALSE)->GetPlainNameString();
+		case ETokenType::CharConst:
+		case ETokenType::StringConst:
+			return FString(*GetTokenString());
+		default:
+			return TEXT("NotConstant");
+		}
+	}
+
+	// Return the token value as a null terminated string.
+	void GetTokenValue(FTokenValue& TokenValue) const
+	{
+		Value.CopyString(TokenValue.Value, Value.Len(), 0);
+		TokenValue.Value[Value.Len()] = 0;
+	}
+
+	// Return the token value as a null terminated string.
+	FTokenValue GetTokenValue() const
+	{
+		FTokenValue TokenValue;
+		GetTokenValue(TokenValue);
+		return TokenValue;
+	}
+
+	// Return the token value for string contants
+	void GetTokenString(FTokenString& TokenString) const
+	{
+		switch (TokenType)
+		{
+		case ETokenType::StringConst:
+		{
+			TCHAR* Out = TokenString.String;;
+			const TCHAR* Pos = &Value[1];
+			TCHAR c = *Pos++;
+			while (c != '"')
+			{
+				if (c == '\\')
+				{
+					c = *Pos++;
+					if (c == 'n')
+					{
+						// Newline escape sequence.
+						c = '\n';
+					}
+				}
+				*Out++ = c;
+				c = *Pos++;
+			}
+			*Out++ = 0;
+			// If this fails, GetToken has a bug/mismatch
+			check(Out - TokenString.String < MAX_STRING_CONST_SIZE);
+			break;
+		}
+
+		case ETokenType::CharConst:
+		{
+			TCHAR ActualCharLiteral = Value[1];
+
+			if (ActualCharLiteral == '\\')
+			{
+				ActualCharLiteral = Value[2];
+				switch (ActualCharLiteral)
+				{
+				case TCHAR('t'):
+					ActualCharLiteral = '\t';
+					break;
+				case TCHAR('n'):
+					ActualCharLiteral = '\n';
+					break;
+				case TCHAR('r'):
+					ActualCharLiteral = '\r';
+					break;
+				}
+			}
+			TokenString.String[0] = ActualCharLiteral;
+			TokenString.String[1] = 0;
+			break;
+		}
+
+		default:
+			checkf(false, TEXT("Call to GetTokenString on token that isn't a string or char constant"));
+		}
+	}
+
+	// Return the token value for string contants
+	FTokenString GetTokenString() const
+	{
+		FTokenString TokenString;
+		GetTokenString(TokenString);
+		return TokenString;
+	}
+
+private:
+	float GetFloatValue() const
+	{
+		int32 Length = Value.Len();
+		TCHAR* Temp = reinterpret_cast<TCHAR*>(alloca((Length + 1) * sizeof(TCHAR)));
+		Value.CopyString(Temp, Length, 0);
+		Temp[Length] = 0;
+		return FCString::Atof(Temp);
+	}
+
+	int64 GetDecimalValue() const
+	{
+		TCHAR* End = const_cast<TCHAR*>(Value.begin()) + Value.Len();
+		return FCString::Strtoi64(Value.begin(), &End, 10);
+	}
+
+	int64 GetHexValue() const
+	{
+		TCHAR* End = const_cast<TCHAR*>(Value.begin()) + Value.Len();
+		return FCString::Strtoi64(Value.begin(), &End, 16);
+	}
+};
+
 /////////////////////////////////////////////////////
 // FBaseParser
 
@@ -153,6 +462,7 @@ public:
 	TCHAR PeekChar();
 	TCHAR GetLeadingChar();
 	void UngetChar();
+	void SkipWhitespaceAndComments();
 
 	/**
 	 * Tests if a character is an end-of-line character.
@@ -183,26 +493,25 @@ public:
 	 *
 	 * @return	true if a token was successfully processed, false otherwise.
 	 */
-	bool GetToken( FToken& Token, bool bNoConsts = false, ESymbolParseOption bParseTemplateClosingBracket = ESymbolParseOption::Normal );
+	bool GetToken( FToken& Token, bool bNoConsts = false, ESymbolParseOption ParseTemplateClosingBracket = ESymbolParseOption::Normal );
 
 	/**
 	 * Put all text from the current position up to either EOL or the StopToken
-	 * into Token.  Advances the compiler's current position.
+	 * into string.  Advances the compiler's current position.
 	 *
-	 * @param	Token	[out] will contain the text that was parsed
+	 * @param	String		[out] will contain the text that was parsed
 	 * @param	StopChar	stop processing when this character is reached
 	 *
 	 * @return	true if a token was parsed
 	 */
-	bool GetRawToken( FToken& Token, TCHAR StopChar = TCHAR('\n') );
+	bool GetRawString( FTokenString& String, TCHAR StopChar = TCHAR('\n') );
 
 	// Doesn't quit if StopChar is found inside a double-quoted string, but does not support quote escapes
-	bool GetRawTokenRespectingQuotes( FToken& Token, TCHAR StopChar = TCHAR('\n') );
+	bool GetRawStringRespectingQuotes(FTokenString& String, TCHAR StopChar = TCHAR('\n') );
 
 	void UngetToken( const FToken& Token );
 	void UngetToken(int32 StartLine, int32 StartPos);
 	bool GetIdentifier( FToken& Token, bool bNoConsts = false );
-	bool GetSymbol( FToken& Token );
 
 	// Modify token to fix redirected types if needed
 	void RedirectTypeIdentifier(FToken& Token) const;
@@ -226,7 +535,8 @@ public:
 
 	// Requiring predefined text.
 	void RequireIdentifier( const TCHAR* Match, ESearchCase::Type SearchCase, const TCHAR* Tag );
-	void RequireSymbol( const TCHAR Match, const TCHAR* Tag, ESymbolParseOption bParseTemplateClosingBracket = ESymbolParseOption::Normal );
+	void RequireSymbol(const TCHAR Match, const TCHAR* Tag, ESymbolParseOption bParseTemplateClosingBracket = ESymbolParseOption::Normal);
+	void RequireSymbol(const TCHAR Match, const FStringView& Tag, ESymbolParseOption bParseTemplateClosingBracket = ESymbolParseOption::Normal);
 	void RequireSymbol(const TCHAR Match, TFunctionRef<FString()> TagGetter, ESymbolParseOption bParseTemplateClosingBracket = ESymbolParseOption::Normal);
 	void RequireConstInt( const TCHAR* Match, const TCHAR* Tag );
 	void RequireAnyConstInt( const TCHAR* Tag );
@@ -299,7 +609,7 @@ void FBaseParser::ParseInheritance(const TCHAR* What, Lambda&& InLambda)
 			Throwf(TEXT("Missing %s name"), What);
 		}
 		RedirectTypeIdentifier(Token);
-		InLambda(Token.Identifier, true);
+		InLambda(*FString(Token.Value), true);
 	}
 
 	// Handle additional inherited interface classes
@@ -317,7 +627,7 @@ void FBaseParser::ParseInheritance(const TCHAR* What, Lambda&& InLambda)
 				Throwf(TEXT("Failed to get interface class identifier"));
 			}
 
-			InterfaceName += Token.Identifier;
+			InterfaceName += Token.Value;
 
 			// Handle templated native classes
 			if (MatchSymbol(TEXT('<')))
@@ -332,19 +642,16 @@ void FBaseParser::ParseInheritance(const TCHAR* What, Lambda&& InLambda)
 						Throwf(TEXT("Unexpected end of file"));
 					}
 
-					if (Token.TokenType == TOKEN_Symbol)
+					if (Token.IsSymbol(TEXT('<')))
+					{
+						++NestedScopes;
+					}
+					else if (Token.IsSymbol(TEXT('>')))
 					{
-						if (Token.Matches(TEXT('<')))
-						{
-							++NestedScopes;
-						}
-						else if (Token.Matches(TEXT('>')))
-						{
-							--NestedScopes;
-						}
+						--NestedScopes;
 					}
 
-					InterfaceName += Token.Identifier;
+					InterfaceName += Token.Value;
 				}
 			}
 
diff --git a/Engine/Source/Programs/UnrealHeaderTool/Private/GeneratedCodeVersion.h b/Engine/Source/Programs/UnrealHeaderTool/Private/GeneratedCodeVersion.h
--- a/Engine/Source/Programs/UnrealHeaderTool/Private/GeneratedCodeVersion.h
+++ b/Engine/Source/Programs/UnrealHeaderTool/Private/GeneratedCodeVersion.h
@@ -54,3 +54,23 @@ inline EGeneratedCodeVersion ToGeneratedCodeVersion(const FString& InString)
 
 	return EGeneratedCodeVersion::None;
 }
+
+inline EGeneratedCodeVersion ToGeneratedCodeVersion(const FStringView& InString)
+{
+	if (InString.Compare(TEXT("V1")) == 0)
+	{
+		return EGeneratedCodeVersion::V1;
+	}
+
+	if (InString.Compare(TEXT("V2")) == 0)
+	{
+		return EGeneratedCodeVersion::V2;
+	}
+
+	if (InString.Compare(TEXT("VLatest")) == 0)
+	{
+		return EGeneratedCodeVersion::VLatest;
+	}
+
+	return EGeneratedCodeVersion::None;
+}
diff --git a/Engine/Source/Programs/UnrealHeaderTool/Private/HeaderParser.cpp b/Engine/Source/Programs/UnrealHeaderTool/Private/HeaderParser.cpp
--- a/Engine/Source/Programs/UnrealHeaderTool/Private/HeaderParser.cpp
+++ b/Engine/Source/Programs/UnrealHeaderTool/Private/HeaderParser.cpp
@@ -87,8 +87,13 @@ TArray<FString> FHeaderParser::ReservedTypeNames = { TEXT("none") };
 
 namespace
 {
-	bool ProbablyAMacro(const TCHAR* Identifier)
+	bool ProbablyAMacro(const FStringView& Identifier)
 	{
+		if (Identifier.Len() == 0)
+		{
+			return false;
+
+		}
 		// Macros must start with a capitalized alphanumeric character or underscore
 		TCHAR FirstChar = Identifier[0];
 		if (FirstChar != TEXT('_') && (FirstChar < TEXT('A') || FirstChar > TEXT('Z')))
@@ -97,26 +102,15 @@ namespace
 		}
 
 		// Test for known delegate and event macros.
-		TCHAR MulticastDelegateStart[] = TEXT("DECLARE_MULTICAST_DELEGATE");
-		if (!FCString::Strncmp(Identifier, MulticastDelegateStart, UE_ARRAY_COUNT(MulticastDelegateStart) - 1))
-		{
-			return true;
-		}
-
-		TCHAR DelegateStart[] = TEXT("DECLARE_DELEGATE");
-		if (!FCString::Strncmp(Identifier, DelegateStart, UE_ARRAY_COUNT(DelegateStart) - 1))
-		{
-			return true;
-		}
-
-		TCHAR DelegateEvent[] = TEXT("DECLARE_EVENT");
-		if (!FCString::Strncmp(Identifier, DelegateEvent, UE_ARRAY_COUNT(DelegateEvent) - 1))
+		if (Identifier.StartsWith(TEXT("DECLARE_MULTICAST_DELEGATE")) ||
+			Identifier.StartsWith(TEXT("DECLARE_DELEGATE")) ||
+			Identifier.StartsWith(TEXT("DECLARE_EVENT")))
 		{
 			return true;
 		}
 
 		// Failing that, we'll guess about it being a macro based on it being a fully-capitalized identifier.
-		while (TCHAR Ch = *++Identifier)
+		for (TCHAR Ch : Identifier.RightChop(1))
 		{
 			if (Ch != TEXT('_') && (Ch < TEXT('A') || Ch > TEXT('Z')) && (Ch < TEXT('0') || Ch > TEXT('9')))
 			{
@@ -138,22 +132,21 @@ namespace
 	bool ProbablyAnUnknownObjectLikeMacro(FHeaderParser& HeaderParser, FToken Token)
 	{
 		// Non-identifiers are not macros
-		if (Token.TokenType != TOKEN_Identifier)
+		if (!Token.IsIdentifier())
 		{
 			return false;
 		}
 
 		// Macros must start with a capitalized alphanumeric character or underscore
-		TCHAR FirstChar = Token.Identifier[0];
+		TCHAR FirstChar = Token.Value[0];
 		if (FirstChar != TEXT('_') && (FirstChar < TEXT('A') || FirstChar > TEXT('Z')))
 		{
 			return false;
 		}
 
 		// We'll guess about it being a macro based on it being fully-capitalized with at least one underscore.
-		const TCHAR* IdentPtr = Token.Identifier;
 		int32 UnderscoreCount = 0;
-		while (TCHAR Ch = *++IdentPtr)
+		for (TCHAR Ch : Token.Value.RightChop(1))
 		{
 			if (Ch == TEXT('_'))
 			{
@@ -173,20 +166,15 @@ namespace
 		}
 
 		// Identifiers which end in _API are known
-		if (IdentPtr - Token.Identifier > 4 && IdentPtr[-4] == TEXT('_') && IdentPtr[-3] == TEXT('A') && IdentPtr[-2] == TEXT('P') && IdentPtr[-1] == TEXT('I'))
+		if (Token.Value.Len() > 4 && Token.Value.EndsWith(FStringView(TEXT("_API"), 4)))
 		{
 			return false;
 		}
 
 		// Ignore certain known macros or identifiers that look like macros.
-		// IMPORTANT: needs to be in lexicographical order.
-		static const TCHAR* Whitelist[] =
-		{
-			TEXT("FORCEINLINE_DEBUGGABLE"),
-			TEXT("FORCEINLINE_STATS"),
-			TEXT("SIZE_T")
-		};
-		if (Algo::FindSortedStringCaseInsensitive(Token.Identifier, Whitelist, UE_ARRAY_COUNT(Whitelist)) >= 0)
+		if (Token.IsValue(TEXT("FORCEINLINE_DEBUGGABLE"), ESearchCase::CaseSensitive) ||
+			Token.IsValue(TEXT("FORCEINLINE_STATS"), ESearchCase::CaseSensitive) ||
+			Token.IsValue(TEXT("SIZE_T"), ESearchCase::CaseSensitive))
 		{
 			return false;
 		}
@@ -201,8 +189,7 @@ namespace
 		HeaderParser.UngetToken(Token);
 		HeaderParser.GetToken(Token);
 
-		bool bResult = PossibleBracketToken.TokenType != TOKEN_Symbol || !PossibleBracketToken.Matches(TEXT('('));
-		return bResult;
+		return !PossibleBracketToken.IsSymbol(TEXT('('));
 	}
 
 	/**
@@ -628,26 +615,27 @@ namespace
 			return;
 		}
 
-		if (MacroToken.TokenType != TOKEN_Identifier || (FCString::Strcmp(MacroToken.Identifier, TEXT("DEPRECATED")) != 0 && FCString::Strcmp(MacroToken.Identifier, TEXT("UE_DEPRECATED")) != 0))
+		if (!MacroToken.IsIdentifier() || 
+			(!MacroToken.IsValue(TEXT("DEPRECATED"), ESearchCase::CaseSensitive) && !MacroToken.IsValue(TEXT("UE_DEPRECATED"), ESearchCase::CaseSensitive)))
 		{
 			Parser.UngetToken(MacroToken);
 			return;
 		}
 
-		auto ErrorMessageGetter = [&MacroToken]() { return FString::Printf(TEXT("%s macro"), MacroToken.Identifier); };
+		auto ErrorMessageGetter = [&MacroToken]() { return FString::Printf(TEXT("%s macro"), *MacroToken.GetTokenValue()); };
 
 		Parser.RequireSymbol(TEXT('('), ErrorMessageGetter);
 
 		FToken Token;
-		if (Parser.GetToken(Token) && (Token.ConstantType != CPT_Float || Token.TokenType != TOKEN_Const))
+		if (Parser.GetToken(Token) && !Token.IsConstFloat())
 		{
-			Parser.Throwf(TEXT("Expected engine version in %s macro"), MacroToken.Identifier);
+			Parser.Throwf(TEXT("Expected engine version in %s macro"), *MacroToken.GetTokenValue());
 		}
 
 		Parser.RequireSymbol(TEXT(','), ErrorMessageGetter);
-		if (Parser.GetToken(Token) && (Token.ConstantType != CPT_String || Token.TokenType != TOKEN_Const))
+		if (Parser.GetToken(Token) && !Token.IsConstString())
 		{
-			Parser.Throwf(TEXT("Expected deprecation message in %s macro"), MacroToken.Identifier);
+			Parser.Throwf(TEXT("Expected deprecation message in %s macro"), *MacroToken.GetTokenValue());
 		}
 
 		Parser.RequireSymbol(TEXT(')'), ErrorMessageGetter);
@@ -785,22 +773,16 @@ const FUHTConfig& FUHTConfig::Get()
 //
 FUnrealClassDefinitionInfo* FHeaderParser::GetQualifiedClass(const TCHAR* Thing)
 {
-	TCHAR ClassName[256]=TEXT("");
-
 	FToken Token;
 	if (GetIdentifier(Token))
 	{
 		RedirectTypeIdentifier(Token);
-
-		FCString::Strncat( ClassName, Token.Identifier, UE_ARRAY_COUNT(ClassName) );
+		return FUnrealClassDefinitionInfo::FindScriptClassOrThrow(*this, FString(Token.Value));
 	}
-
-	if (!ClassName[0])
+	else
 	{
 		Throwf(TEXT("%s: Missing class name"), Thing);
 	}
-
-	return FUnrealClassDefinitionInfo::FindScriptClassOrThrow(*this, ClassName);
 }
 
 /*-----------------------------------------------------------------------------
@@ -964,7 +946,7 @@ FUnrealEnumDefinitionInfo& FHeaderParser::CompileEnum()
 		Throwf(TEXT("Missing identifier after UENUM()") );
 	}
 
-	if (EnumToken.Matches(TEXT("namespace"), ESearchCase::CaseSensitive))
+	if (EnumToken.IsValue(TEXT("namespace"), ESearchCase::CaseSensitive))
 	{
 		CppForm = UEnum::ECppForm::Namespaced;
 
@@ -972,14 +954,14 @@ FUnrealEnumDefinitionInfo& FHeaderParser::CompileEnum()
 
 		bReadEnumName = GetIdentifier(EnumToken);
 	}
-	else if (EnumToken.Matches(TEXT("enum"), ESearchCase::CaseSensitive))
+	else if (EnumToken.IsValue(TEXT("enum"), ESearchCase::CaseSensitive))
 	{
 		if (!GetIdentifier(EnumToken))
 		{
 			Throwf(TEXT("Missing identifier after enum") );
 		}
 
-		if (EnumToken.Matches(TEXT("class"), ESearchCase::CaseSensitive) || EnumToken.Matches(TEXT("struct"), ESearchCase::CaseSensitive))
+		if (EnumToken.IsValue(TEXT("class"), ESearchCase::CaseSensitive) || EnumToken.IsValue(TEXT("struct"), ESearchCase::CaseSensitive))
 		{
 			CppForm = UEnum::ECppForm::EnumClass;
 		}
@@ -1006,10 +988,11 @@ FUnrealEnumDefinitionInfo& FHeaderParser::CompileEnum()
 		Throwf(TEXT("Missing enumeration name") );
 	}
 
-	ParseFieldMetaData(EnumMetaData, EnumToken.Identifier);
+	FTokenValue EnumIdentifier = EnumToken.GetTokenValue();
+	ParseFieldMetaData(EnumMetaData, *EnumIdentifier);
 
 	// Create enum definition.
-	FUnrealEnumDefinitionInfo& EnumDef = GTypeDefinitionInfoMap.FindByNameChecked<FUnrealEnumDefinitionInfo>(EnumToken.Identifier);
+	FUnrealEnumDefinitionInfo& EnumDef = GTypeDefinitionInfoMap.FindByNameChecked<FUnrealEnumDefinitionInfo>(*EnumIdentifier);
 
 	Scope->AddType(EnumDef);
 
@@ -1082,7 +1065,7 @@ FUnrealEnumDefinitionInfo& FHeaderParser::CompileEnum()
 				Throwf(TEXT("Missing enumeration name") );
 			}
 
-			EnumDef.SetCppType(FString::Printf(TEXT("%s::%s"), EnumToken.Identifier, InnerEnumToken.Identifier));
+			EnumDef.SetCppType(FString::Printf(TEXT("%s::%s"), *EnumIdentifier, *InnerEnumToken.GetTokenValue()));
 
 			RequireSymbol( TEXT('{'), TEXT("'Enum'") );
 		}
@@ -1091,7 +1074,7 @@ FUnrealEnumDefinitionInfo& FHeaderParser::CompileEnum()
 		case UEnum::ECppForm::Regular:
 		case UEnum::ECppForm::EnumClass:
 		{
-			EnumDef.SetCppType(EnumToken.Identifier);
+			EnumDef.SetCppType(*EnumIdentifier);
 		}
 		break;
 	}
@@ -1112,6 +1095,7 @@ FUnrealEnumDefinitionInfo& FHeaderParser::CompileEnum()
 	int64 CurrentEnumValue = 0;
 	while (GetIdentifier(TagToken))
 	{
+		FTokenValue TagIdentifier = TagToken.GetTokenValue();
 		AddFormattedPrevCommentAsTooltipMetaData(TagMetaData);
 
 		// Try to read an optional explicit enum value specification
@@ -1139,21 +1123,10 @@ FUnrealEnumDefinitionInfo& FHeaderParser::CompileEnum()
 					Throwf(TEXT("Enumerator: end of file encountered while parsing the initializer"));
 				}
 
-				if (InitToken.TokenType == TOKEN_Symbol)
+				if (InitToken.IsSymbol(TEXT(',')) || InitToken.IsSymbol(TEXT('}')) || InitToken.IsIdentifier(TEXT("UMETA"), ESearchCase::IgnoreCase))
 				{
-					if (InitToken.Matches(TEXT(',')) || InitToken.Matches(TEXT('}')))
-					{
-						UngetToken(InitToken);
-						break;
-					}
-				}
-				else if (InitToken.TokenType == TOKEN_Identifier)
-				{
-					if (FCString::Stricmp(InitToken.Identifier, TEXT("UMETA")) == 0)
-					{
-						UngetToken(InitToken);
-						break;
-					}
+					UngetToken(InitToken);
+					break;
 				}
 
 				// There are tokens after the initializer so it's not a standalone literal,
@@ -1171,13 +1144,13 @@ FUnrealEnumDefinitionInfo& FHeaderParser::CompileEnum()
 			case UEnum::ECppForm::Namespaced:
 			case UEnum::ECppForm::EnumClass:
 			{
-				NewTag = FName(*FString::Printf(TEXT("%s::%s"), EnumToken.Identifier, TagToken.Identifier), FNAME_Add);
+				NewTag = FName(*FString::Printf(TEXT("%s::%s"), *EnumIdentifier, *TagIdentifier), FNAME_Add);
 			}
 			break;
 
 			case UEnum::ECppForm::Regular:
 			{
-				NewTag = FName(TagToken.Identifier, FNAME_Add);
+				NewTag = FName(TagToken.Value, FNAME_Add);
 			}
 			break;
 		}
@@ -1195,11 +1168,11 @@ FUnrealEnumDefinitionInfo& FHeaderParser::CompileEnum()
 		EntryMetaData.Add(TagMetaData);
 
 		// check for metadata on this enum value
-		ParseFieldMetaData(TagMetaData, TagToken.Identifier);
+		ParseFieldMetaData(TagMetaData, *TagIdentifier);
 		if (TagMetaData.Num() > 0)
 		{
 			// special case for enum value metadata - we need to prepend the key name with the enum value name
-			const FString TokenString = TagToken.Identifier;
+			const FString TokenString = *TagIdentifier;
 			for (const auto& MetaData : TagMetaData)
 			{
 				FString KeyString = TokenString + TEXT(".") + MetaData.Key.ToString();
@@ -1218,7 +1191,7 @@ FUnrealEnumDefinitionInfo& FHeaderParser::CompileEnum()
 				Throwf(TEXT("UENUM: end of file encountered"));
 			}
 
-			if (ClosingBrace.TokenType == TOKEN_Symbol && ClosingBrace.Matches(TEXT('}')))
+			if (ClosingBrace.IsSymbol(TEXT('}')))
 			{
 				UngetToken(ClosingBrace);
 				break;
@@ -1662,13 +1635,16 @@ static const TCHAR* GetAccessSpecifierName(EAccessSpecifier AccessSpecifier)
 // Tries to parse the token as an access protection specifier (public:, protected:, or private:)
 EAccessSpecifier FHeaderParser::ParseAccessProtectionSpecifier(const FToken& Token)
 {
-	EAccessSpecifier ResultAccessSpecifier = ACCESS_NotAnAccessSpecifier;
+	if (!Token.IsIdentifier())
+	{
+		return ACCESS_NotAnAccessSpecifier;
+	}
 
 	for (EAccessSpecifier Test = EAccessSpecifier(ACCESS_NotAnAccessSpecifier + 1); Test != ACCESS_Num; Test = EAccessSpecifier(Test + 1))
 	{
-		if (Token.Matches(GetAccessSpecifierName(Test), ESearchCase::CaseSensitive))
+		if (Token.IsValue(GetAccessSpecifierName(Test), ESearchCase::CaseSensitive))
 		{
-			auto ErrorMessageGetter = [&Token]() { return FString::Printf(TEXT("after %s"), Token.Identifier);  };
+			auto ErrorMessageGetter = [&Token]() { return FString::Printf(TEXT("after %s"), *Token.GetTokenValue());  };
 
 			// Consume the colon after the specifier
 			RequireSymbol(TEXT(':'), ErrorMessageGetter);
@@ -1854,29 +1830,29 @@ FUnrealScriptStructDefinitionInfo& FHeaderParser::CompileStructDeclaration()
 		{
 			CurrentAccessSpecifier = AccessSpecifier;
 		}
-		else if (Token.Matches(TEXT("UPROPERTY"), ESearchCase::CaseSensitive))
+		else if (Token.IsIdentifier(TEXT("UPROPERTY"), ESearchCase::CaseSensitive))
 		{
 			CompileVariableDeclaration(StructDef);
 		}
-		else if (Token.Matches(TEXT("UFUNCTION"), ESearchCase::CaseSensitive))
+		else if (Token.IsIdentifier(TEXT("UFUNCTION"), ESearchCase::CaseSensitive))
 		{
 			Throwf(TEXT("USTRUCTs cannot contain UFUNCTIONs."));
 		}
-		else if (Token.Matches(TEXT("RIGVM_METHOD"), ESearchCase::CaseSensitive))
+		else if (Token.IsIdentifier(TEXT("RIGVM_METHOD"), ESearchCase::CaseSensitive))
 		{
 			CompileRigVMMethodDeclaration(StructDef);
 		}
-		else if (Token.Matches(TEXT("GENERATED_USTRUCT_BODY"), ESearchCase::CaseSensitive) || Token.Matches(TEXT("GENERATED_BODY"), ESearchCase::CaseSensitive))
+		else if (Token.IsIdentifier(TEXT("GENERATED_USTRUCT_BODY"), ESearchCase::CaseSensitive) || Token.IsIdentifier(TEXT("GENERATED_BODY"), ESearchCase::CaseSensitive))
 		{
 			// Match 'GENERATED_USTRUCT_BODY' '(' [StructName] ')' or 'GENERATED_BODY' '(' [StructName] ')'
 			if (CurrentAccessSpecifier != ACCESS_Public)
 			{
-				Throwf(TEXT("%s must be in the public scope of '%s', not private or protected."), Token.Identifier, *StructNameInScript);
+				Throwf(TEXT("%s must be in the public scope of '%s', not private or protected."), *Token.GetTokenValue(), *StructNameInScript);
 			}
 
 			if (StructDef.GetMacroDeclaredLineNumber() != INDEX_NONE)
 			{
-				Throwf(TEXT("Multiple %s declarations found in '%s'"), Token.Identifier, *StructNameInScript);
+				Throwf(TEXT("Multiple %s declarations found in '%s'"), *Token.GetTokenValue(), *StructNameInScript);
 			}
 
 			StructDef.SetMacroDeclaredLineNumber(InputLine);
@@ -1889,20 +1865,20 @@ FUnrealScriptStructDefinitionInfo& FHeaderParser::CompileStructDeclaration()
 			// Eat a semicolon if present (not required)
 			SafeMatchSymbol(TEXT(';'));
 		}
-		else if ( Token.Matches(TEXT('#')) && MatchIdentifier(TEXT("ifdef"), ESearchCase::CaseSensitive) )
+		else if ( Token.IsSymbol(TEXT('#')) && MatchIdentifier(TEXT("ifdef"), ESearchCase::CaseSensitive) )
 		{
 			PushCompilerDirective(ECompilerDirective::Insignificant);
 		}
-		else if ( Token.Matches(TEXT('#')) && MatchIdentifier(TEXT("ifndef"), ESearchCase::CaseSensitive) )
+		else if ( Token.IsSymbol(TEXT('#')) && MatchIdentifier(TEXT("ifndef"), ESearchCase::CaseSensitive) )
 		{
 			PushCompilerDirective(ECompilerDirective::Insignificant);
 		}
-		else if (Token.Matches(TEXT('#')) && MatchIdentifier(TEXT("endif"), ESearchCase::CaseSensitive))
+		else if (Token.IsSymbol(TEXT('#')) && MatchIdentifier(TEXT("endif"), ESearchCase::CaseSensitive))
 		{
 			PopCompilerDirective();
 			// Do nothing and hope that the if code below worked out OK earlier
 		}
-		else if ( Token.Matches(TEXT('#')) && MatchIdentifier(TEXT("if"), ESearchCase::CaseSensitive) )
+		else if ( Token.IsSymbol(TEXT('#')) && MatchIdentifier(TEXT("if"), ESearchCase::CaseSensitive) )
 		{
 			//@TODO: This parsing should be combined with CompileDirective and probably happen much much higher up!
 			bool bInvertConditional = MatchSymbol(TEXT('!'));
@@ -1963,7 +1939,7 @@ FUnrealScriptStructDefinitionInfo& FHeaderParser::CompileStructDeclaration()
 				}
 			}
 		}
-		else if (Token.Matches(TEXT('#')) && MatchIdentifier(TEXT("pragma"), ESearchCase::CaseSensitive))
+		else if (Token.IsSymbol(TEXT('#')) && MatchIdentifier(TEXT("pragma"), ESearchCase::CaseSensitive))
 		{
 			// skip it and skip over the text, it is not recorded or processed
 			TCHAR c;
@@ -1977,15 +1953,13 @@ FUnrealScriptStructDefinitionInfo& FHeaderParser::CompileStructDeclaration()
 		}
 		else
 		{
-			if (!Token.Matches( TEXT('}')))
+			if (!Token.IsSymbol( TEXT('}')))
 			{
 				// Skip declaration will destroy data in Token, so cache off the identifier in case we need to provfide an error
-				TCHAR FirstTokenIdentifier[NAME_SIZE];
-				FCString::Strncpy(FirstTokenIdentifier, Token.Identifier, NAME_SIZE);
-
+				FToken FirstToken = Token;
 				if (!SkipDeclaration(Token))
 				{
-					Throwf(TEXT("'struct': Unexpected '%s'"), FirstTokenIdentifier);
+					Throwf(TEXT("'struct': Unexpected '%s'"), *FirstToken.GetTokenValue());
 				}	
 			}
 			else
@@ -2607,18 +2581,18 @@ void FHeaderParser::CompileDirective()
 	{
 		Throwf(TEXT("Missing compiler directive after '#'") );
 	}
-	else if (Directive.Matches(TEXT("error"), ESearchCase::CaseSensitive))
+	else if (Directive.IsValue(TEXT("error"), ESearchCase::CaseSensitive))
 	{
 		Throwf(TEXT("#error directive encountered") );
 	}
-	else if (Directive.Matches(TEXT("pragma"), ESearchCase::CaseSensitive))
+	else if (Directive.IsValue(TEXT("pragma"), ESearchCase::CaseSensitive))
 	{
 		// Ignore all pragmas
 	}
-	else if (Directive.Matches(TEXT("linenumber"), ESearchCase::CaseSensitive))
+	else if (Directive.IsValue(TEXT("linenumber"), ESearchCase::CaseSensitive))
 	{
 		FToken Number;
-		if (!GetToken(Number) || (Number.TokenType != TOKEN_Const) || (Number.ConstantType != CPT_Int && Number.ConstantType != CPT_Int64))
+		if (!GetToken(Number) || !Number.IsConstInt())
 		{
 			Throwf(TEXT("Missing line number in line number directive"));
 		}
@@ -2629,19 +2603,20 @@ void FHeaderParser::CompileDirective()
 			InputLine = newInputLine;
 		}
 	}
-	else if (Directive.Matches(TEXT("include"), ESearchCase::CaseSensitive))
+	else if (Directive.IsValue(TEXT("include"), ESearchCase::CaseSensitive))
 	{
 		FToken IncludeName;
-		if (GetToken(IncludeName) && (IncludeName.TokenType == TOKEN_Const) && (IncludeName.ConstantType == CPT_String))
+		if (GetToken(IncludeName) && IncludeName.IsConstString())
 		{
+			FTokenString InludeNameString = IncludeName.GetTokenString();
 			const FString& ExpectedHeaderName = SourceFile.GetGeneratedHeaderFilename();
-			if (FCString::Stricmp(IncludeName.String, *ExpectedHeaderName) == 0)
+			if (FCString::Stricmp(*InludeNameString, *ExpectedHeaderName) == 0)
 			{
 				bSpottedAutogeneratedHeaderInclude = true;
 			}
 		}
 	}
-	else if (Directive.Matches(TEXT("if"), ESearchCase::CaseSensitive))
+	else if (Directive.IsValue(TEXT("if"), ESearchCase::CaseSensitive))
 	{
 		// Eat the ! if present
 		bool bNotDefined = MatchSymbol(TEXT('!'));
@@ -2660,48 +2635,48 @@ void FHeaderParser::CompileDirective()
 				Throwf(TEXT("Missing define name '#if'") );
 			}
 
-			if ( Define.Matches(TEXT("WITH_EDITORONLY_DATA"), ESearchCase::CaseSensitive) )
+			if ( Define.IsValue(TEXT("WITH_EDITORONLY_DATA"), ESearchCase::CaseSensitive) )
 			{
 				PushCompilerDirective(ECompilerDirective::WithEditorOnlyData);
 			}
-			else if ( Define.Matches(TEXT("WITH_EDITOR"), ESearchCase::CaseSensitive) )
+			else if ( Define.IsValue(TEXT("WITH_EDITOR"), ESearchCase::CaseSensitive) )
 			{
 				PushCompilerDirective(ECompilerDirective::WithEditor);
 			}
-			else if (Define.Matches(TEXT("WITH_HOT_RELOAD"), ESearchCase::CaseSensitive) || Define.Matches(TEXT("WITH_HOT_RELOAD_CTORS"), ESearchCase::CaseSensitive) || Define.Matches(TEXT('1')))
+			else if (Define.IsValue(TEXT("WITH_HOT_RELOAD"), ESearchCase::CaseSensitive) || Define.IsValue(TEXT("WITH_HOT_RELOAD_CTORS"), ESearchCase::CaseSensitive) || Define.IsSymbol(TEXT('1'))) // @TODO: This symbol test can never work after a GetIdentifier
 			{
 				PushCompilerDirective(ECompilerDirective::Insignificant);
 			}
-			else if ( Define.Matches(TEXT("CPP"), ESearchCase::CaseSensitive) && bNotDefined)
+			else if ( Define.IsValue(TEXT("CPP"), ESearchCase::CaseSensitive) && bNotDefined)
 			{
 				PushCompilerDirective(ECompilerDirective::Insignificant);
 			}
 			else
 			{
-				Throwf(TEXT("Unknown define '#if %s' in class or global scope"), Define.Identifier);
+				Throwf(TEXT("Unknown define '#if %s' in class or global scope"), *Define.GetTokenValue());
 			}
 		}
 	}
-	else if (Directive.Matches(TEXT("endif"), ESearchCase::CaseSensitive))
+	else if (Directive.IsValue(TEXT("endif"), ESearchCase::CaseSensitive))
 	{
 		PopCompilerDirective();
 	}
-	else if (Directive.Matches(TEXT("define"), ESearchCase::CaseSensitive))
+	else if (Directive.IsValue(TEXT("define"), ESearchCase::CaseSensitive))
 	{
 		// Ignore the define directive (can be multiline).
 		bDefineDirective = true;
 	}
-	else if (Directive.Matches(TEXT("ifdef"), ESearchCase::CaseSensitive) || Directive.Matches(TEXT("ifndef"), ESearchCase::CaseSensitive))
+	else if (Directive.IsValue(TEXT("ifdef"), ESearchCase::CaseSensitive) || Directive.IsValue(TEXT("ifndef"), ESearchCase::CaseSensitive))
 	{
 		PushCompilerDirective(ECompilerDirective::Insignificant);
 	}
-	else if (Directive.Matches(TEXT("undef"), ESearchCase::CaseSensitive) || Directive.Matches(TEXT("else"), ESearchCase::CaseSensitive))
+	else if (Directive.IsValue(TEXT("undef"), ESearchCase::CaseSensitive) || Directive.IsValue(TEXT("else"), ESearchCase::CaseSensitive))
 	{
 		// Ignore. UHT can only handle #if directive
 	}
 	else
 	{
-		Throwf(TEXT("Unrecognized compiler directive %s"), Directive.Identifier );
+		Throwf(TEXT("Unrecognized compiler directive %s"), *Directive.GetTokenValue() );
 	}
 
 	// Skip to end of line (or end of multiline #define).
@@ -3273,9 +3248,9 @@ void FHeaderParser::GetVarType(
 			Throwf(TEXT("%s: Missing variable type"), GetHintText(*this, VariableCategory));
 		}
 
-		if (FCString::Strcmp(InlineToken.Identifier, TEXT("inline")) != 0
-			&& FCString::Strcmp(InlineToken.Identifier, TEXT("FORCENOINLINE")) != 0
-			&& FCString::Strncmp(InlineToken.Identifier, TEXT("FORCEINLINE"), 11) != 0)
+		if (!InlineToken.IsValue(TEXT("inline"), ESearchCase::CaseSensitive) &&
+			!InlineToken.IsValue(TEXT("FORCENOINLINE"), ESearchCase::CaseSensitive) &&
+			!InlineToken.ValueStartsWith(TEXT("FORCEINLINE"), ESearchCase::CaseSensitive))
 		{
 			UngetToken(InlineToken);
 		}
@@ -3309,9 +3284,10 @@ void FHeaderParser::GetVarType(
 		FToken LayoutToken;
 		if (GetToken(LayoutToken))
 		{
-			if (LayoutToken.TokenType == TOKEN_Identifier)
+			if (LayoutToken.IsIdentifier())
 			{
-				LayoutMacroType = (ELayoutMacroType)Algo::FindSortedStringCaseInsensitive(LayoutToken.Identifier, GLayoutMacroNames, UE_ARRAY_COUNT(GLayoutMacroNames));
+				FTokenValue Name = LayoutToken.GetTokenValue();
+				LayoutMacroType = (ELayoutMacroType)Algo::FindSortedStringCaseInsensitive(*Name, GLayoutMacroNames, UE_ARRAY_COUNT(GLayoutMacroNames));
 				if (LayoutMacroType != ELayoutMacroType::None)
 				{
 					RequireSymbol(TEXT('('), GLayoutMacroNames[(int32)LayoutMacroType]);
@@ -3363,68 +3339,68 @@ void FHeaderParser::GetVarType(
 
 	//
 	FToken VarType;
-	if ( !GetIdentifier(VarType,1) )
+	if ( !GetIdentifier(VarType, true) )
 	{
 		Throwf(TEXT("%s: Missing variable type"), GetHintText(*this, VariableCategory));
 	}
 
 	RedirectTypeIdentifier(VarType);
 
-	if ( VariableCategory == EVariableCategory::Return && VarType.Matches(TEXT("void"), ESearchCase::CaseSensitive) )
+	if ( VariableCategory == EVariableCategory::Return && VarType.IsValue(TEXT("void"), ESearchCase::CaseSensitive) )
 	{
 		VarProperty = FPropertyBase(CPT_None);
 	}
-	else if ( VarType.Matches(TEXT("int8"), ESearchCase::CaseSensitive) )
+	else if ( VarType.IsValue(TEXT("int8"), ESearchCase::CaseSensitive) )
 	{
 		VarProperty = FPropertyBase(CPT_Int8);
 	}
-	else if ( VarType.Matches(TEXT("int16"), ESearchCase::CaseSensitive) )
+	else if ( VarType.IsValue(TEXT("int16"), ESearchCase::CaseSensitive) )
 	{
 		VarProperty = FPropertyBase(CPT_Int16);
 	}
-	else if ( VarType.Matches(TEXT("int32"), ESearchCase::CaseSensitive) )
+	else if ( VarType.IsValue(TEXT("int32"), ESearchCase::CaseSensitive) )
 	{
 		VarProperty = FPropertyBase(CPT_Int);
 	}
-	else if ( VarType.Matches(TEXT("int64"), ESearchCase::CaseSensitive) )
+	else if ( VarType.IsValue(TEXT("int64"), ESearchCase::CaseSensitive) )
 	{
 		VarProperty = FPropertyBase(CPT_Int64);
 	}
-	else if ( VarType.Matches(TEXT("uint64"), ESearchCase::CaseSensitive) && IsBitfieldProperty(LayoutMacroType) )
+	else if ( VarType.IsValue(TEXT("uint64"), ESearchCase::CaseSensitive) && IsBitfieldProperty(LayoutMacroType) )
 	{
 		// 64-bit bitfield (bool) type, treat it like 8 bit type
 		VarProperty = FPropertyBase(CPT_Bool8);
 	}
-	else if ( VarType.Matches(TEXT("uint32"), ESearchCase::CaseSensitive) && IsBitfieldProperty(LayoutMacroType) )
+	else if ( VarType.IsValue(TEXT("uint32"), ESearchCase::CaseSensitive) && IsBitfieldProperty(LayoutMacroType) )
 	{
 		// 32-bit bitfield (bool) type, treat it like 8 bit type
 		VarProperty = FPropertyBase(CPT_Bool8);
 	}
-	else if ( VarType.Matches(TEXT("uint16"), ESearchCase::CaseSensitive) && IsBitfieldProperty(LayoutMacroType) )
+	else if ( VarType.IsValue(TEXT("uint16"), ESearchCase::CaseSensitive) && IsBitfieldProperty(LayoutMacroType) )
 	{
 		// 16-bit bitfield (bool) type, treat it like 8 bit type.
 		VarProperty = FPropertyBase(CPT_Bool8);
 	}
-	else if ( VarType.Matches(TEXT("uint8"), ESearchCase::CaseSensitive) && IsBitfieldProperty(LayoutMacroType) )
+	else if ( VarType.IsValue(TEXT("uint8"), ESearchCase::CaseSensitive) && IsBitfieldProperty(LayoutMacroType) )
 	{
 		// 8-bit bitfield (bool) type
 		VarProperty = FPropertyBase(CPT_Bool8);
 	}
-	else if ( VarType.Matches(TEXT("int"), ESearchCase::CaseSensitive) )
+	else if ( VarType.IsValue(TEXT("int"), ESearchCase::CaseSensitive) )
 	{
 		VarProperty = FPropertyBase(CPT_Int, EIntType::Unsized);
 	}
-	else if ( VarType.Matches(TEXT("signed"), ESearchCase::CaseSensitive) )
+	else if ( VarType.IsValue(TEXT("signed"), ESearchCase::CaseSensitive) )
 	{
 		MatchIdentifier(TEXT("int"), ESearchCase::CaseSensitive);
 		VarProperty = FPropertyBase(CPT_Int, EIntType::Unsized);
 	}
-	else if (VarType.Matches(TEXT("unsigned"), ESearchCase::CaseSensitive))
+	else if (VarType.IsValue(TEXT("unsigned"), ESearchCase::CaseSensitive))
 	{
 		MatchIdentifier(TEXT("int"), ESearchCase::CaseSensitive);
 		VarProperty = FPropertyBase(CPT_UInt32, EIntType::Unsized);
 	}
-	else if ( VarType.Matches(TEXT("bool"), ESearchCase::CaseSensitive) )
+	else if ( VarType.IsValue(TEXT("bool"), ESearchCase::CaseSensitive) )
 	{
 		if (IsBitfieldProperty(LayoutMacroType))
 		{
@@ -3433,34 +3409,34 @@ void FHeaderParser::GetVarType(
 		// C++ bool type
 		VarProperty = FPropertyBase(CPT_Bool);
 	}
-	else if ( VarType.Matches(TEXT("uint8"), ESearchCase::CaseSensitive) )
+	else if ( VarType.IsValue(TEXT("uint8"), ESearchCase::CaseSensitive) )
 	{
 		// Intrinsic Byte type.
 		VarProperty = FPropertyBase(CPT_Byte);
 	}
-	else if ( VarType.Matches(TEXT("uint16"), ESearchCase::CaseSensitive) )
+	else if ( VarType.IsValue(TEXT("uint16"), ESearchCase::CaseSensitive) )
 	{
 		VarProperty = FPropertyBase(CPT_UInt16);
 	}
-	else if ( VarType.Matches(TEXT("uint32"), ESearchCase::CaseSensitive) )
+	else if ( VarType.IsValue(TEXT("uint32"), ESearchCase::CaseSensitive) )
 	{
 		VarProperty = FPropertyBase(CPT_UInt32);
 	}
-	else if ( VarType.Matches(TEXT("uint64"), ESearchCase::CaseSensitive) )
+	else if ( VarType.IsValue(TEXT("uint64"), ESearchCase::CaseSensitive) )
 	{
 		VarProperty = FPropertyBase(CPT_UInt64);
 	}
-	else if ( VarType.Matches(TEXT("float"), ESearchCase::CaseSensitive) )
+	else if ( VarType.IsValue(TEXT("float"), ESearchCase::CaseSensitive) )
 	{
 		// Intrinsic single precision floating point type.
 		VarProperty = FPropertyBase(CPT_Float);
 	}
-	else if ( VarType.Matches(TEXT("double"), ESearchCase::CaseSensitive) )
+	else if ( VarType.IsValue(TEXT("double"), ESearchCase::CaseSensitive) )
 	{
 		// Intrinsic double precision floating point type type.
 		VarProperty = FPropertyBase(CPT_Double);
 	}
-	else if (VarType.Matches(TEXT("FLargeWorldCoordinatesReal"), ESearchCase::CaseSensitive))		// LWC_TODO: Remove with UE_LARGE_WORLD_COORDINATES_DISABLED toggle.
+	else if (VarType.IsValue(TEXT("FLargeWorldCoordinatesReal"), ESearchCase::CaseSensitive))		// LWC_TODO: Remove with UE_LARGE_WORLD_COORDINATES_DISABLED toggle.
 	{
 		if (!SourceFile.GetFilename().EndsWith(TEXT("NoExportTypes.h")))
 		{
@@ -3468,12 +3444,12 @@ void FHeaderParser::GetVarType(
 		}
 		VarProperty = FPropertyBase(CPT_FLargeWorldCoordinatesReal);
 	}
-	else if ( VarType.Matches(TEXT("FName"), ESearchCase::CaseSensitive) )
+	else if ( VarType.IsValue(TEXT("FName"), ESearchCase::CaseSensitive) )
 	{
 		// Intrinsic Name type.
 		VarProperty = FPropertyBase(CPT_Name);
 	}
-	else if ( VarType.Matches(TEXT("TArray"), ESearchCase::CaseSensitive) )
+	else if ( VarType.IsValue(TEXT("TArray"), ESearchCase::CaseSensitive) )
 	{
 		RequireSymbol( TEXT('<'), TEXT("'tarray'") );
 
@@ -3497,12 +3473,12 @@ void FHeaderParser::GetVarType(
 			Throwf(TEXT("Missing token while parsing TArray."));
 		}
 
-		if (CloseTemplateToken.TokenType != TOKEN_Symbol || !CloseTemplateToken.Matches(TEXT('>')))
+		if (!CloseTemplateToken.IsSymbol(TEXT('>')))
 		{
 			// If we didn't find a comma, report it
-			if (!CloseTemplateToken.Matches(TEXT(',')))
+			if (!CloseTemplateToken.IsSymbol(TEXT(',')))
 			{
-				Throwf(TEXT("Expected '>' but found '%s'"), CloseTemplateToken.Identifier);
+				Throwf(TEXT("Expected '>' but found '%s'"), *CloseTemplateToken.GetTokenValue());
 			}
 
 			// If we found a comma, read the next thing, assume it's an allocator, and report that
@@ -3512,12 +3488,12 @@ void FHeaderParser::GetVarType(
 				Throwf(TEXT("Unexpected end of file when parsing TArray allocator."));
 			}
 
-			if (AllocatorToken.TokenType != TOKEN_Identifier)
+			if (!AllocatorToken.IsIdentifier())
 			{
-				Throwf(TEXT("Found '%s' - expected a '>' or ','."), AllocatorToken.Identifier);
+				Throwf(TEXT("Found '%s' - expected a '>' or ','."), *AllocatorToken.GetTokenValue());
 			}
 
-			if (FCString::Strcmp(AllocatorToken.Identifier, TEXT("FMemoryImageAllocator")) == 0)
+			if (AllocatorToken.IsValue(TEXT("FMemoryImageAllocator"), ESearchCase::CaseSensitive))
 			{
 				if (EnumHasAnyFlags(Flags, CPF_Net))
 				{
@@ -3528,7 +3504,7 @@ void FHeaderParser::GetVarType(
 
 				VarProperty.AllocatorType = EAllocatorType::MemoryImage;
 			}
-			else if (FCString::Strcmp(AllocatorToken.Identifier, TEXT("TMemoryImageAllocator")) == 0)
+			else if (AllocatorToken.IsValue(TEXT("TMemoryImageAllocator"), ESearchCase::CaseSensitive))
 			{
 				if (EnumHasAnyFlags(Flags, CPF_Net))
 				{
@@ -3545,7 +3521,7 @@ void FHeaderParser::GetVarType(
 						Throwf(TEXT("Unexpected end of file when parsing TMemoryImageAllocator template arguments."));
 					}
 
-					if (SkipToken.TokenType == TOKEN_Symbol && FCString::Strcmp(SkipToken.Identifier, TEXT(">")) == 0)
+					if (SkipToken.IsSymbol(TEXT('>')))
 					{
 						RequireSymbol(TEXT('>'), TEXT("TArray template arguments"), ESymbolParseOption::CloseTemplateBracket);
 						VarProperty.AllocatorType = EAllocatorType::MemoryImage;
@@ -3555,11 +3531,11 @@ void FHeaderParser::GetVarType(
 			}
 			else
 			{
-				Throwf(TEXT("Found '%s' - explicit allocators are not supported in TArray properties."), AllocatorToken.Identifier);
+				Throwf(TEXT("Found '%s' - explicit allocators are not supported in TArray properties."), *AllocatorToken.GetTokenValue());
 			}
 		}
 	}
-	else if ( VarType.Matches(TEXT("TMap"), ESearchCase::CaseSensitive) )
+	else if ( VarType.IsValue(TEXT("TMap"), ESearchCase::CaseSensitive) )
 	{
 		RequireSymbol( TEXT('<'), TEXT("'tmap'") );
 
@@ -3587,7 +3563,7 @@ void FHeaderParser::GetVarType(
 		}
 
 		FToken CommaToken;
-		if (!GetToken(CommaToken, /*bNoConsts=*/ true) || CommaToken.TokenType != TOKEN_Symbol || !CommaToken.Matches(TEXT(',')))
+		if (!GetToken(CommaToken, /*bNoConsts=*/ true) || !CommaToken.IsSymbol(TEXT(',')))
 		{
 			Throwf(TEXT("Missing value type while parsing TMap."));
 		}
@@ -3608,12 +3584,12 @@ void FHeaderParser::GetVarType(
 			Throwf(TEXT("Missing token while parsing TMap."));
 		}
 
-		if (CloseTemplateToken.TokenType != TOKEN_Symbol || !CloseTemplateToken.Matches(TEXT('>')))
+		if (!CloseTemplateToken.IsSymbol(TEXT('>')))
 		{
 			// If we didn't find a comma, report it
-			if (!CloseTemplateToken.Matches(TEXT(',')))
+			if (!CloseTemplateToken.IsSymbol(TEXT(',')))
 			{
-				Throwf(TEXT("Expected '>' but found '%s'"), CloseTemplateToken.Identifier);
+				Throwf(TEXT("Expected '>' but found '%s'"), *CloseTemplateToken.GetTokenValue());
 			}
 
 			// If we found a comma, read the next thing, assume it's an allocator, and report that
@@ -3623,12 +3599,12 @@ void FHeaderParser::GetVarType(
 				Throwf(TEXT("Unexpected end of file when parsing TArray allocator."));
 			}
 
-			if (AllocatorToken.TokenType != TOKEN_Identifier)
+			if (!AllocatorToken.IsIdentifier())
 			{
-				Throwf(TEXT("Found '%s' - expected a '>' or ','."), AllocatorToken.Identifier);
+				Throwf(TEXT("Found '%s' - expected a '>' or ','."), *AllocatorToken.GetTokenValue());
 			}
 
-			if (FCString::Strcmp(AllocatorToken.Identifier, TEXT("FMemoryImageSetAllocator")) == 0)
+			if (AllocatorToken.IsIdentifier(TEXT("FMemoryImageSetAllocator"), ESearchCase::CaseSensitive))
 			{
 				RequireSymbol(TEXT('>'), TEXT("TMap template arguments"), ESymbolParseOption::CloseTemplateBracket);
 
@@ -3636,11 +3612,11 @@ void FHeaderParser::GetVarType(
 			}
 			else
 			{
-				Throwf(TEXT("Found '%s' - explicit allocators are not supported in TMap properties."), AllocatorToken.Identifier);
+				Throwf(TEXT("Found '%s' - explicit allocators are not supported in TMap properties."), *AllocatorToken.GetTokenValue());
 			}
 		}
 	}
-	else if ( VarType.Matches(TEXT("TSet"), ESearchCase::CaseSensitive) )
+	else if ( VarType.IsValue(TEXT("TSet"), ESearchCase::CaseSensitive) )
 	{
 		RequireSymbol( TEXT('<'), TEXT("'tset'") );
 
@@ -3674,25 +3650,25 @@ void FHeaderParser::GetVarType(
 			Throwf(TEXT("Missing token while parsing TArray."));
 		}
 
-		if (CloseTemplateToken.TokenType != TOKEN_Symbol || !CloseTemplateToken.Matches(TEXT('>')))
+		if (!CloseTemplateToken.IsSymbol(TEXT('>')))
 		{
 			// If we didn't find a comma, report it
-			if (!CloseTemplateToken.Matches(TEXT(',')))
+			if (!CloseTemplateToken.IsSymbol(TEXT(',')))
 			{
-				Throwf(TEXT("Expected '>' but found '%s'"), CloseTemplateToken.Identifier);
+				Throwf(TEXT("Expected '>' but found '%s'"), *CloseTemplateToken.GetTokenValue());
 			}
 
 			// If we found a comma, read the next thing, assume it's a keyfuncs, and report that
 			FToken AllocatorToken;
 			if (!GetToken(AllocatorToken, /*bNoConsts=*/ true, ESymbolParseOption::CloseTemplateBracket))
 			{
-				Throwf(TEXT("Expected '>' but found '%s'"), CloseTemplateToken.Identifier);
+				Throwf(TEXT("Expected '>' but found '%s'"), *CloseTemplateToken.GetTokenValue());
 			}
 
-			Throwf(TEXT("Found '%s' - explicit KeyFuncs are not supported in TSet properties."), AllocatorToken.Identifier);
+			Throwf(TEXT("Found '%s' - explicit KeyFuncs are not supported in TSet properties."), *AllocatorToken.GetTokenValue());
 		}
 	}
-	else if ( VarType.Matches(TEXT("FString"), ESearchCase::CaseSensitive) || VarType.Matches(TEXT("FMemoryImageString"), ESearchCase::CaseSensitive))
+	else if ( VarType.IsValue(TEXT("FString"), ESearchCase::CaseSensitive) || VarType.IsValue(TEXT("FMemoryImageString"), ESearchCase::CaseSensitive))
 	{
 		VarProperty = FPropertyBase(CPT_String);
 
@@ -3719,17 +3695,17 @@ void FHeaderParser::GetVarType(
 			}
 		}
 	}
-	else if ( VarType.Matches(TEXT("Text"), ESearchCase::IgnoreCase) )
+	else if ( VarType.IsValue(TEXT("Text"), ESearchCase::IgnoreCase) )
 	{
-		Throwf(TEXT("%s' is missing a prefix, expecting 'FText'"), VarType.Identifier);
+		Throwf(TEXT("%s' is missing a prefix, expecting 'FText'"), *VarType.GetTokenValue());
 	}
-	else if ( VarType.Matches(TEXT("FText"), ESearchCase::CaseSensitive) )
+	else if ( VarType.IsValue(TEXT("FText"), ESearchCase::CaseSensitive) )
 	{
 		VarProperty = FPropertyBase(CPT_Text);
 	}
-	else if (VarType.Matches(TEXT("TEnumAsByte"), ESearchCase::CaseSensitive))
+	else if (VarType.IsValue(TEXT("TEnumAsByte"), ESearchCase::CaseSensitive))
 	{
-		RequireSymbol(TEXT('<'), VarType.Identifier);
+		RequireSymbol(TEXT('<'), VarType.Value);
 
 		// Eat the forward declaration enum text if present
 		MatchIdentifier(TEXT("enum"), ESearchCase::CaseSensitive);
@@ -3739,7 +3715,7 @@ void FHeaderParser::GetVarType(
 		FToken InnerEnumType;
 		if (GetIdentifier(InnerEnumType, true))
 		{
-			if (FUnrealEnumDefinitionInfo* EnumDef = GTypeDefinitionInfoMap.FindByName<FUnrealEnumDefinitionInfo>(InnerEnumType.Identifier))
+			if (FUnrealEnumDefinitionInfo* EnumDef = GTypeDefinitionInfoMap.FindByName<FUnrealEnumDefinitionInfo>(*InnerEnumType.GetTokenValue()))
 			{
 				// In-scope enumeration.
 				VarProperty = FPropertyBase(*EnumDef, CPT_Byte);
@@ -3763,9 +3739,9 @@ void FHeaderParser::GetVarType(
 			Throwf(TEXT("Expected the name of a previously defined enum"));
 		}
 
-		RequireSymbol(TEXT('>'), VarType.Identifier, ESymbolParseOption::CloseTemplateBracket);
+		RequireSymbol(TEXT('>'), VarType.Value, ESymbolParseOption::CloseTemplateBracket);
 	}
-	else if (VarType.Matches(TEXT("TFieldPath"), ESearchCase::CaseSensitive ))
+	else if (VarType.IsValue(TEXT("TFieldPath"), ESearchCase::CaseSensitive ))
 	{
 		RequireSymbol( TEXT('<'), TEXT("'TFieldPath'") );
 
@@ -3777,18 +3753,18 @@ void FHeaderParser::GetVarType(
 		}
 		else
 		{
-			FieldClassName = FName(PropertyTypeToken.Identifier + 1, FNAME_Add);
+			FieldClassName = FName(PropertyTypeToken.Value.RightChop(1), FNAME_Add);
 			if (!FPropertyTraits::IsValidFieldClass(FieldClassName))
 			{
-				Throwf(TEXT("Undefined property type: %s"), PropertyTypeToken.Identifier);
+				Throwf(TEXT("Undefined property type: %s"), *PropertyTypeToken.GetTokenValue());
 			}
 		}
 
-		RequireSymbol(TEXT('>'), VarType.Identifier, ESymbolParseOption::CloseTemplateBracket);
+		RequireSymbol(TEXT('>'), VarType.Value, ESymbolParseOption::CloseTemplateBracket);
 
 		VarProperty = FPropertyBase(FieldClassName, CPT_FieldPath);
 	}
-	else if (FUnrealEnumDefinitionInfo* EnumDef = GTypeDefinitionInfoMap.FindByName<FUnrealEnumDefinitionInfo>(VarType.Identifier))
+	else if (FUnrealEnumDefinitionInfo* EnumDef = GTypeDefinitionInfoMap.FindByName<FUnrealEnumDefinitionInfo>(*VarType.GetTokenValue()))
 	{
 		EPropertyType UnderlyingType = CPT_Byte;
 
@@ -3819,9 +3795,9 @@ void FHeaderParser::GetVarType(
 	{
 		// Check for structs/classes
 		bool bHandledType = false;
-		FString IdentifierStripped = GetClassNameWithPrefixRemoved(VarType.Identifier);
+		FString IdentifierStripped = GetClassNameWithPrefixRemoved(FString(VarType.Value));
 		bool bStripped = false;
-		FUnrealScriptStructDefinitionInfo* StructDef = GTypeDefinitionInfoMap.FindByName<FUnrealScriptStructDefinitionInfo>(VarType.Identifier);
+		FUnrealScriptStructDefinitionInfo* StructDef = GTypeDefinitionInfoMap.FindByName<FUnrealScriptStructDefinitionInfo>(*VarType.GetTokenValue());
 		if (!StructDef)
 		{
 			StructDef = GTypeDefinitionInfoMap.FindByName<FUnrealScriptStructDefinitionInfo>(*IdentifierStripped);
@@ -3844,7 +3820,7 @@ void FHeaderParser::GetVarType(
 			{
 				if (FUnrealClassDefinitionInfo* LocalOwnerClassDef = FUnrealClassDefinitionInfo::FindClass(*IdentifierStripped))
 				{
-					const FString DelegateIdentifierStripped = GetClassNameWithPrefixRemoved(DelegateName.Identifier);
+					const FString DelegateIdentifierStripped = GetClassNameWithPrefixRemoved(FString(DelegateName.Value));
 					if (FUnrealFieldDefinitionInfo* FoundDef = LocalOwnerClassDef->GetScope()->FindTypeByName(*(DelegateIdentifierStripped + HEADER_GENERATED_DELEGATE_SIGNATURE_SUFFIX)))
 					{
 						if (FUnrealFunctionDefinitionInfo* DelegateFuncDef = FoundDef->AsFunction())
@@ -3856,7 +3832,7 @@ void FHeaderParser::GetVarType(
 				}
 				else
 				{
-					Throwf(TEXT("Cannot find class '%s', to resolve delegate '%s'"), *IdentifierStripped, DelegateName.Identifier);
+					Throwf(TEXT("Cannot find class '%s', to resolve delegate '%s'"), *IdentifierStripped, *DelegateName.GetTokenValue());
 				}
 			}
 		}
@@ -3870,15 +3846,15 @@ void FHeaderParser::GetVarType(
 			{
 				const TCHAR* PrefixCPP = UHTConfig.StructsWithTPrefix.Contains(IdentifierStripped) ? TEXT("T") : StructDef->GetPrefixCPP();
 				FString ExpectedStructName = FString::Printf(TEXT("%s%s"), PrefixCPP, *StructDef->GetName());
-				if (FString(VarType.Identifier) != ExpectedStructName)
+				if (!VarType.IsValue(*ExpectedStructName, ESearchCase::CaseSensitive))
 				{
-					Throwf(TEXT("Struct '%s' is missing or has an incorrect prefix, expecting '%s'"), VarType.Identifier, *ExpectedStructName);
+					Throwf(TEXT("Struct '%s' is missing or has an incorrect prefix, expecting '%s'"), *VarType.GetTokenValue(), *ExpectedStructName);
 				}
 			}
-			else if (!UHTConfig.StructsWithNoPrefix.Contains(VarType.Identifier))
+			else if (!UHTConfig.StructsWithNoPrefix.Contains(VarType.Value))
 			{
-				const TCHAR* PrefixCPP = UHTConfig.StructsWithTPrefix.Contains(VarType.Identifier) ? TEXT("T") : StructDef->GetPrefixCPP();
-				Throwf(TEXT("Struct '%s' is missing a prefix, expecting '%s'"), VarType.Identifier, *FString::Printf(TEXT("%s%s"), PrefixCPP, *StructDef->GetName()));
+				const TCHAR* PrefixCPP = UHTConfig.StructsWithTPrefix.Contains(VarType.Value) ? TEXT("T") : StructDef->GetPrefixCPP();
+				Throwf(TEXT("Struct '%s' is missing a prefix, expecting '%s'"), *VarType.GetTokenValue(), *FString::Printf(TEXT("%s%s"), PrefixCPP, *StructDef->GetName()));
 			}
 
 			bHandledType = true;
@@ -3909,35 +3885,35 @@ void FHeaderParser::GetVarType(
 				FUnrealClassDefinitionInfo* TempClassDef = nullptr;
 
 				EPropertyType PropertyType = CPT_ObjectReference;
-				const bool bIsSoftObjectPtrTemplate = VarType.Matches(TEXT("TSoftObjectPtr"), ESearchCase::CaseSensitive);
+				const bool bIsSoftObjectPtrTemplate = VarType.IsValue(TEXT("TSoftObjectPtr"), ESearchCase::CaseSensitive);
 
 				bool bWeakIsAuto = false;
 
-				if (VarType.Matches(TEXT("TSubclassOf"), ESearchCase::CaseSensitive))
+				if (VarType.IsValue(TEXT("TSubclassOf"), ESearchCase::CaseSensitive))
 				{
 					TempClassDef = GUClassDef;
 				}
-				else if (VarType.Matches(TEXT("FScriptInterface"), ESearchCase::CaseSensitive))
+				else if (VarType.IsValue(TEXT("FScriptInterface"), ESearchCase::CaseSensitive))
 				{
 					TempClassDef = GUInterfaceDef;
 					Flags |= CPF_UObjectWrapper;
 				}
-				else if (VarType.Matches(TEXT("TSoftClassPtr"), ESearchCase::CaseSensitive))
+				else if (VarType.IsValue(TEXT("TSoftClassPtr"), ESearchCase::CaseSensitive))
 				{
 					TempClassDef = GUClassDef;
 					PropertyType = CPT_SoftObjectReference;
 				}
 				else
 				{
-					const bool bIsLazyPtrTemplate = VarType.Matches(TEXT("TLazyObjectPtr"), ESearchCase::CaseSensitive);
-					const bool bIsObjectPtrTemplate = VarType.Matches(TEXT("TObjectPtr"), ESearchCase::CaseSensitive);
-					const bool bIsWeakPtrTemplate = VarType.Matches(TEXT("TWeakObjectPtr"), ESearchCase::CaseSensitive);
-					const bool bIsAutoweakPtrTemplate = VarType.Matches(TEXT("TAutoWeakObjectPtr"), ESearchCase::CaseSensitive);
-					const bool bIsScriptInterfaceWrapper = VarType.Matches(TEXT("TScriptInterface"), ESearchCase::CaseSensitive);
+					const bool bIsLazyPtrTemplate = VarType.IsValue(TEXT("TLazyObjectPtr"), ESearchCase::CaseSensitive);
+					const bool bIsObjectPtrTemplate = VarType.IsValue(TEXT("TObjectPtr"), ESearchCase::CaseSensitive);
+					const bool bIsWeakPtrTemplate = VarType.IsValue(TEXT("TWeakObjectPtr"), ESearchCase::CaseSensitive);
+					const bool bIsAutoweakPtrTemplate = VarType.IsValue(TEXT("TAutoWeakObjectPtr"), ESearchCase::CaseSensitive);
+					const bool bIsScriptInterfaceWrapper = VarType.IsValue(TEXT("TScriptInterface"), ESearchCase::CaseSensitive);
 
 					if (bIsLazyPtrTemplate || bIsObjectPtrTemplate || bIsWeakPtrTemplate || bIsAutoweakPtrTemplate || bIsScriptInterfaceWrapper || bIsSoftObjectPtrTemplate)
 					{
-						RequireSymbol(TEXT('<'), VarType.Identifier);
+						RequireSymbol(TEXT('<'), VarType.Value);
 
 						// Also consume const
 						bNativeConstTemplateArg |= MatchIdentifier(TEXT("const"), ESearchCase::CaseSensitive);
@@ -3951,10 +3927,11 @@ void FHeaderParser::GetVarType(
 						{
 							RedirectTypeIdentifier(InnerClass);
 
-							TempClassDef = FUnrealClassDefinitionInfo::FindScriptClass(InnerClass.Identifier);
+							TempClassDef = FUnrealClassDefinitionInfo::FindScriptClass(FString(InnerClass.Value));
 							if (TempClassDef == nullptr)
 							{
-								Throwf(TEXT("Unrecognized type '%s' (in expression %s<%s>) - type must be a UCLASS"), InnerClass.Identifier, VarType.Identifier, InnerClass.Identifier);
+								FTokenValue InnerClassName = InnerClass.GetTokenValue();
+								Throwf(TEXT("Unrecognized type '%s' (in expression %s<%s>) - type must be a UCLASS"), *InnerClassName, *VarType.GetTokenValue(), *InnerClassName);
 							}
 
 							if (bIsAutoweakPtrTemplate)
@@ -3983,17 +3960,17 @@ void FHeaderParser::GetVarType(
 						}
 						else
 						{
-							Throwf(TEXT("%s: Missing template type"), VarType.Identifier);
+							Throwf(TEXT("%s: Missing template type"), *VarType.GetTokenValue());
 						}
 
 						// Const after template argument type but before end of template symbol
 						bNativeConstTemplateArg |= MatchIdentifier(TEXT("const"), ESearchCase::CaseSensitive);
 
-						RequireSymbol(TEXT('>'), VarType.Identifier, ESymbolParseOption::CloseTemplateBracket);
+						RequireSymbol(TEXT('>'), VarType.Value, ESymbolParseOption::CloseTemplateBracket);
 					}
 					else
 					{
-						TempClassDef = FUnrealClassDefinitionInfo::FindScriptClass(VarType.Identifier);
+						TempClassDef = FUnrealClassDefinitionInfo::FindScriptClass(*VarType.GetTokenValue());
 					}
 				}
 
@@ -4031,7 +4008,7 @@ void FHeaderParser::GetVarType(
 
 							RedirectTypeIdentifier(Limitor);
 
-							VarProperty.MetaClassDef = FUnrealClassDefinitionInfo::FindScriptClassOrThrow(*this, Limitor.Identifier);
+							VarProperty.MetaClassDef = FUnrealClassDefinitionInfo::FindScriptClassOrThrow(*this, FString(Limitor.Value));
 
 							RequireSymbol(TEXT('>'), TEXT("'class limitor'"), ESymbolParseOption::CloseTemplateBracket);
 						}
@@ -4102,7 +4079,7 @@ void FHeaderParser::GetVarType(
 
 				if (!bHandledType)
 				{
-					Throwf(TEXT("Unrecognized type '%s' - type must be a UCLASS, USTRUCT or UENUM"), VarType.Identifier);
+					Throwf(TEXT("Unrecognized type '%s' - type must be a UCLASS, USTRUCT or UENUM"), *VarType.GetTokenValue());
 				}
 			}
 		}
@@ -4126,28 +4103,28 @@ void FHeaderParser::GetVarType(
 		}
 		else
 		{
-			Throwf(TEXT("Inappropriate keyword 'const' on variable of type '%s'"), VarType.Identifier);
+			Throwf(TEXT("Inappropriate keyword 'const' on variable of type '%s'"), *VarType.GetTokenValue());
 		}
 	}
 
 	if (bUnconsumedClassKeyword)
 	{
-		Throwf(TEXT("Inappropriate keyword 'class' on variable of type '%s'"), VarType.Identifier );
+		Throwf(TEXT("Inappropriate keyword 'class' on variable of type '%s'"), *VarType.GetTokenValue());
 	}
 
 	if (bUnconsumedStructKeyword)
 	{
-		Throwf(TEXT("Inappropriate keyword 'struct' on variable of type '%s'"), VarType.Identifier );
+		Throwf(TEXT("Inappropriate keyword 'struct' on variable of type '%s'"), *VarType.GetTokenValue());
 	}
 
 	if (bUnconsumedEnumKeyword)
 	{
-		Throwf(TEXT("Inappropriate keyword 'enum' on variable of type '%s'"), VarType.Identifier );
+		Throwf(TEXT("Inappropriate keyword 'enum' on variable of type '%s'"), *VarType.GetTokenValue());
 	}
 
 	if (MatchSymbol(TEXT('*')))
 	{
-		Throwf(TEXT("Inappropriate '*' on variable of type '%s', cannot have an exposed pointer to this type."), VarType.Identifier );
+		Throwf(TEXT("Inappropriate '*' on variable of type '%s', cannot have an exposed pointer to this type."), *VarType.GetTokenValue());
 	}
 
 	//@TODO: UCREMOVAL: 'const' member variables that will get written post-construction by defaultproperties
@@ -4179,7 +4156,7 @@ void FHeaderParser::GetVarType(
 			{
 				if (!(Flags & CPF_ConstParm))
 				{
-					Throwf(TEXT("Replicated %s parameters cannot be passed by non-const reference"), VarType.Identifier);
+					Throwf(TEXT("Replicated %s parameters cannot be passed by non-const reference"), *VarType.GetTokenValue());
 				}
 
 				Flags |= CPF_ReferenceParm;
@@ -4311,13 +4288,13 @@ FUnrealPropertyDefinitionInfo& FHeaderParser::GetVarNameAndDim
 	const TCHAR* HintText = GetHintText(*this, VariableCategory);
 
 	AddModuleRelativePathToMetadata(ParentStruct, VarProperty.MetaData);
-	TCHAR Identifier[NAME_SIZE];
+	FStringView Identifier;
 
 	// Get variable name.
 	if (VariableCategory == EVariableCategory::Return)
 	{
 		// Hard-coded variable name, such as with return value.
-		FCString::Strcpy( Identifier, TEXT("ReturnValue") );
+		Identifier = FStringView(TEXT("ReturnValue"));
 	}
 	else
 	{
@@ -4341,19 +4318,17 @@ FUnrealPropertyDefinitionInfo& FHeaderParser::GetVarNameAndDim
 				break;
 		}
 
-		FCString::Strcpy(Identifier, VarToken.Identifier);
+		Identifier = VarToken.Value;
 	}
 
 	// Check to see if the variable is deprecated, and if so set the flag
 	{
-		FString VarName(Identifier);
-
-		const int32 DeprecatedIndex = VarName.Find(TEXT("_DEPRECATED"));
-		const int32 NativizedPropertyPostfixIndex = VarName.Find(TEXT("__pf")); //TODO: check OverrideNativeName in Meta Data, to be sure it's not a random occurrence of the "__pf" string.
+		const int32 DeprecatedIndex = Identifier.Find(TEXT("_DEPRECATED"));
+		const int32 NativizedPropertyPostfixIndex = Identifier.Find(TEXT("__pf")); //TODO: check OverrideNativeName in Meta Data, to be sure it's not a random occurrence of the "__pf" string.
 		bool bIgnoreDeprecatedWord = (NativizedPropertyPostfixIndex != INDEX_NONE) && (NativizedPropertyPostfixIndex > DeprecatedIndex);
 		if ((DeprecatedIndex != INDEX_NONE) && !bIgnoreDeprecatedWord)
 		{
-			if (DeprecatedIndex != VarName.Len() - 11)
+			if (DeprecatedIndex != Identifier.Len() - 11)
 			{
 				Throwf(TEXT("Deprecated variables must end with _DEPRECATED"));
 			}
@@ -4365,12 +4340,12 @@ FUnrealPropertyDefinitionInfo& FHeaderParser::GetVarNameAndDim
 
 			if (bWarnOnGetter)
 			{
-				LogWarning(TEXT("%s: Deprecated property '%s' should not be marked as blueprint visible without having a BlueprintGetter"), HintText, *VarName);
+				LogWarning(TEXT("%s: Deprecated property '%s' should not be marked as blueprint visible without having a BlueprintGetter"), HintText, *FString(Identifier));
 			}
 
 			if (bWarnOnSetter)
 			{
-				LogWarning(TEXT("%s: Deprecated property '%s' should not be marked as blueprint writeable without having a BlueprintSetter"), HintText, *VarName);
+				LogWarning(TEXT("%s: Deprecated property '%s' should not be marked as blueprint writeable without having a BlueprintSetter"), HintText, *FString(Identifier));
 			}
 
 
@@ -4378,20 +4353,19 @@ FUnrealPropertyDefinitionInfo& FHeaderParser::GetVarNameAndDim
 			if (VarProperty.PropertyFlags & (CPF_Edit | CPF_EditConst) || // Property is marked as editable
 				(!bBlueprintVisible && (VarProperty.PropertyFlags & CPF_BlueprintReadOnly) && !(VarProperty.ImpliedPropertyFlags & CPF_BlueprintReadOnly)) ) // Is BPRO, but not via Implied Flags and not caught by Getter/Setter path above
 			{
-				LogWarning(TEXT("%s: Deprecated property '%s' should not be marked as visible or editable"), HintText, *VarName);
+				LogWarning(TEXT("%s: Deprecated property '%s' should not be marked as visible or editable"), HintText, *FString(Identifier));
 			}
 
 			VarProperty.PropertyFlags |= CPF_Deprecated;
-			VarName.MidInline(0, DeprecatedIndex, false);
-
-			FCString::Strcpy(Identifier, *VarName);
+			Identifier.MidInline(0, DeprecatedIndex);
 		}
 	}
 
 	// Make sure it doesn't conflict.
+	FString VarName(Identifier);
 	int32 OuterContextCount = 0;
-	FUnrealFunctionDefinitionInfo* ExistingFunctionDef = FindFunction(ParentStruct, Identifier, true, nullptr);
-	FUnrealPropertyDefinitionInfo* ExistingPropertyDef = FindProperty(ParentStruct, Identifier, true);
+	FUnrealFunctionDefinitionInfo* ExistingFunctionDef = FindFunction(ParentStruct, *VarName, true, nullptr);
+	FUnrealPropertyDefinitionInfo* ExistingPropertyDef = FindProperty(ParentStruct, *VarName, true);
 
 	if (ExistingFunctionDef != nullptr || ExistingPropertyDef != nullptr)
 	{
@@ -4420,14 +4394,14 @@ FUnrealPropertyDefinitionInfo& FHeaderParser::GetVarNameAndDim
 		{
 			Throwf(TEXT("%s: '%s' cannot be defined in '%s' as it is already defined in scope '%s' (shadowing is not allowed)"),
 				HintText,
-				Identifier,
+				*FString(Identifier),
 				*ParentStruct.GetName(),
 				ExistingFunctionDef ? *ExistingFunctionDef->GetOuter()->GetName() : *ExistingPropertyDef->GetFullName());
 		}
 	}
 
 	// Get optional dimension immediately after name.
-	FToken Dimensions;
+	FTokenString Dimensions;
 	if ((LayoutMacroType == ELayoutMacroType::None && MatchSymbol(TEXT('['))) || LayoutMacroType == ELayoutMacroType::Array || LayoutMacroType == ELayoutMacroType::ArrayEditorOnly)
 	{
 		switch (VariableCategory)
@@ -4457,17 +4431,17 @@ FUnrealPropertyDefinitionInfo& FHeaderParser::GetVarNameAndDim
 		if (LayoutMacroType == ELayoutMacroType::None)
 		{
 			// Ignore how the actual array dimensions are actually defined - we'll calculate those with the compiler anyway.
-			if (!GetRawToken(Dimensions, TEXT(']')))
+			if (!GetRawString(Dimensions, TEXT(']')))
 			{
-				Throwf(TEXT("%s %s: Missing ']'"), HintText, Identifier);
+				Throwf(TEXT("%s %s: Missing ']'"), HintText, *FString(Identifier));
 			}
 		}
 		else
 		{
 			// Ignore how the actual array dimensions are actually defined - we'll calculate those with the compiler anyway.
-			if (!GetRawToken(Dimensions, TEXT(')')))
+			if (!GetRawString(Dimensions, TEXT(')')))
 			{
-				Throwf(TEXT("%s %s: Missing ']'"), HintText, Identifier);
+				Throwf(TEXT("%s %s: Missing ']'"), HintText, *FString(Identifier));
 			}
 		}
 
@@ -4483,7 +4457,7 @@ FUnrealPropertyDefinitionInfo& FHeaderParser::GetVarNameAndDim
 	// Try gathering metadata for member fields
 	if (VariableCategory == EVariableCategory::Member)
 	{
-		ParseFieldMetaData(VarProperty.MetaData, Identifier);
+		ParseFieldMetaData(VarProperty.MetaData, *VarName);
 		AddFormattedPrevCommentAsTooltipMetaData(VarProperty.MetaData);
 	}
 	// validate UFunction parameters
@@ -4526,14 +4500,14 @@ bool FHeaderParser::CompileDeclaration(TArray<FUnrealFunctionDefinitionInfo*>& D
 	{
 		if (!IsAllowedInThisNesting(ENestAllowFlags::VarDecl) && !IsAllowedInThisNesting(ENestAllowFlags::Function))
 		{
-			Throwf(TEXT("Access specifier %s not allowed here."), Token.Identifier);
+			Throwf(TEXT("Access specifier %s not allowed here."), *Token.GetTokenValue());
 		}
 		check(TopNest->NestType == ENestType::Class || TopNest->NestType == ENestType::Interface || TopNest->NestType == ENestType::NativeInterface);
 		CurrentAccessSpecifier = AccessSpecifier;
 		return true;
 	}
 
-	if (Token.Matches(TEXT("class"), ESearchCase::CaseSensitive) && (TopNest->NestType == ENestType::GlobalScope))
+	if (Token.IsIdentifier(TEXT("class"), ESearchCase::CaseSensitive) && (TopNest->NestType == ENestType::GlobalScope))
 	{
 		// Make sure the previous class ended with valid nesting.
 		if (bEncounteredNewStyleClass_UnmatchedBrackets)
@@ -4554,15 +4528,15 @@ bool FHeaderParser::CompileDeclaration(TArray<FUnrealFunctionDefinitionInfo*>& D
 		return true;
 	}
 
-	if (Token.Matches(TEXT("GENERATED_IINTERFACE_BODY"), ESearchCase::CaseSensitive) || (Token.Matches(TEXT("GENERATED_BODY"), ESearchCase::CaseSensitive) && TopNest->NestType == ENestType::NativeInterface))
+	if (Token.IsIdentifier(TEXT("GENERATED_IINTERFACE_BODY"), ESearchCase::CaseSensitive) || (Token.IsIdentifier(TEXT("GENERATED_BODY"), ESearchCase::CaseSensitive) && TopNest->NestType == ENestType::NativeInterface))
 	{
 		if (TopNest->NestType != ENestType::NativeInterface)
 		{
-			Throwf(TEXT("%s must occur inside the native interface definition"), Token.Identifier);
+			Throwf(TEXT("%s must occur inside the native interface definition"), *Token.GetTokenValue());
 		}
-		RequireSymbol(TEXT('('), Token.Identifier);
+		RequireSymbol(TEXT('('), Token.Value);
 		CompileVersionDeclaration(GetCurrentClassDef());
-		RequireSymbol(TEXT(')'), Token.Identifier);
+		RequireSymbol(TEXT(')'), Token.Value);
 
 		FUnrealClassDefinitionInfo& ClassDef = GetCurrentClassDef();
 		if (ClassDef.GetParsedInterfaceState() == EParsedInterface::NotAnInterface)
@@ -4585,27 +4559,27 @@ bool FHeaderParser::CompileDeclaration(TArray<FUnrealFunctionDefinitionInfo*>& D
 
 		bClassHasGeneratedIInterfaceBody = true;
 
-		if (Token.Matches(TEXT("GENERATED_IINTERFACE_BODY"), ESearchCase::CaseSensitive))
+		if (Token.IsIdentifier(TEXT("GENERATED_IINTERFACE_BODY"), ESearchCase::CaseSensitive))
 		{
 			CurrentAccessSpecifier = ACCESS_Public;
 		}
 
-		if (Token.Matches(TEXT("GENERATED_BODY"), ESearchCase::CaseSensitive))
+		if (Token.IsIdentifier(TEXT("GENERATED_BODY"), ESearchCase::CaseSensitive))
 		{
 			GetCurrentClassDef().MarkGeneratedBody();
 		}
 		return true;
 	}
 
-	if (Token.Matches(TEXT("GENERATED_UINTERFACE_BODY"), ESearchCase::CaseSensitive) || (Token.Matches(TEXT("GENERATED_BODY"), ESearchCase::CaseSensitive) && TopNest->NestType == ENestType::Interface))
+	if (Token.IsIdentifier(TEXT("GENERATED_UINTERFACE_BODY"), ESearchCase::CaseSensitive) || (Token.IsIdentifier(TEXT("GENERATED_BODY"), ESearchCase::CaseSensitive) && TopNest->NestType == ENestType::Interface))
 	{
 		if (TopNest->NestType != ENestType::Interface)
 		{
-			Throwf(TEXT("%s must occur inside the interface definition"), Token.Identifier);
+			Throwf(TEXT("%s must occur inside the interface definition"), *Token.GetTokenValue());
 		}
-		RequireSymbol(TEXT('('), Token.Identifier);
+		RequireSymbol(TEXT('('), Token.Value);
 		CompileVersionDeclaration(GetCurrentClassDef());
-		RequireSymbol(TEXT(')'), Token.Identifier);
+		RequireSymbol(TEXT(')'), Token.Value);
 
 		FUnrealClassDefinitionInfo& ClassDef = GetCurrentClassDef();
 
@@ -4614,23 +4588,23 @@ bool FHeaderParser::CompileDeclaration(TArray<FUnrealFunctionDefinitionInfo*>& D
 
 		bClassHasGeneratedUInterfaceBody = true;
 
-		if (Token.Matches(TEXT("GENERATED_UINTERFACE_BODY"), ESearchCase::CaseSensitive))
+		if (Token.IsIdentifier(TEXT("GENERATED_UINTERFACE_BODY"), ESearchCase::CaseSensitive))
 		{
 			CurrentAccessSpecifier = ACCESS_Public;
 		}
 		return true;
 	}
 
-	if (Token.Matches(TEXT("GENERATED_UCLASS_BODY"), ESearchCase::CaseSensitive) || (Token.Matches(TEXT("GENERATED_BODY"), ESearchCase::CaseSensitive) && TopNest->NestType == ENestType::Class))
+	if (Token.IsIdentifier(TEXT("GENERATED_UCLASS_BODY"), ESearchCase::CaseSensitive) || (Token.IsIdentifier(TEXT("GENERATED_BODY"), ESearchCase::CaseSensitive) && TopNest->NestType == ENestType::Class))
 	{
 		if (TopNest->NestType != ENestType::Class)
 		{
-			Throwf(TEXT("%s must occur inside the class definition"), Token.Identifier);
+			Throwf(TEXT("%s must occur inside the class definition"), *Token.GetTokenValue());
 		}
 
 		FUnrealClassDefinitionInfo& ClassDef = GetCurrentClassDef();
 
-		if (Token.Matches(TEXT("GENERATED_BODY"), ESearchCase::CaseSensitive))
+		if (Token.IsIdentifier(TEXT("GENERATED_BODY"), ESearchCase::CaseSensitive))
 		{
 			GetCurrentClassDef().MarkGeneratedBody();
 
@@ -4641,9 +4615,9 @@ bool FHeaderParser::CompileDeclaration(TArray<FUnrealFunctionDefinitionInfo*>& D
 			CurrentAccessSpecifier = ACCESS_Public;
 		}
 
-		RequireSymbol(TEXT('('), Token.Identifier);
+		RequireSymbol(TEXT('('), Token.Value);
 		CompileVersionDeclaration(GetCurrentClassDef());
-		RequireSymbol(TEXT(')'), Token.Identifier);
+		RequireSymbol(TEXT(')'), Token.Value);
 
 		ClassDef.SetGeneratedBodyLine(InputLine);
 
@@ -4651,7 +4625,7 @@ bool FHeaderParser::CompileDeclaration(TArray<FUnrealFunctionDefinitionInfo*>& D
 		return true;
 	}
 
-	if (Token.Matches(TEXT("UCLASS"), ESearchCase::CaseSensitive))
+	if (Token.IsIdentifier(TEXT("UCLASS"), ESearchCase::CaseSensitive))
 	{
 		bHaveSeenUClass = true;
 		bEncounteredNewStyleClass_UnmatchedBrackets = true;
@@ -4659,7 +4633,7 @@ bool FHeaderParser::CompileDeclaration(TArray<FUnrealFunctionDefinitionInfo*>& D
 		return true;
 	}
 
-	if (Token.Matches(TEXT("UINTERFACE"), ESearchCase::CaseSensitive))
+	if (Token.IsIdentifier(TEXT("UINTERFACE"), ESearchCase::CaseSensitive))
 	{
 		bHaveSeenUClass = true;
 		bEncounteredNewStyleClass_UnmatchedBrackets = true;
@@ -4667,27 +4641,27 @@ bool FHeaderParser::CompileDeclaration(TArray<FUnrealFunctionDefinitionInfo*>& D
 		return true;
 	}
 
-	if (Token.Matches(TEXT("UFUNCTION"), ESearchCase::CaseSensitive))
+	if (Token.IsIdentifier(TEXT("UFUNCTION"), ESearchCase::CaseSensitive))
 	{
 		FUnrealFunctionDefinitionInfo& FunctionDef = CompileFunctionDeclaration();
 		return true;
 	}
 
-	if (Token.Matches(TEXT("UDELEGATE"), ESearchCase::CaseSensitive))
+	if (Token.IsIdentifier(TEXT("UDELEGATE"), ESearchCase::CaseSensitive))
 	{
-		FUnrealFunctionDefinitionInfo &DelegateDef = CompileDelegateDeclaration(Token.Identifier, EDelegateSpecifierAction::Parse);
+		FUnrealFunctionDefinitionInfo &DelegateDef = CompileDelegateDeclaration(Token.Value, EDelegateSpecifierAction::Parse);
 		DelegatesToFixup.Add(&DelegateDef);
 		return true;
 	}
 
 	if (IsValidDelegateDeclaration(Token)) // Legacy delegate parsing - it didn't need a UDELEGATE
 	{
-		FUnrealFunctionDefinitionInfo& DelegateDef = CompileDelegateDeclaration(Token.Identifier);
+		FUnrealFunctionDefinitionInfo& DelegateDef = CompileDelegateDeclaration(Token.Value);
 		DelegatesToFixup.Add(&DelegateDef);
 		return true;
 	}
 
-	if (Token.Matches(TEXT("UPROPERTY"), ESearchCase::CaseSensitive))
+	if (Token.IsIdentifier(TEXT("UPROPERTY"), ESearchCase::CaseSensitive))
 	{
 		CheckAllow(TEXT("'Member variable declaration'"), ENestAllowFlags::VarDecl);
 		check(TopNest->NestType == ENestType::Class);
@@ -4696,28 +4670,28 @@ bool FHeaderParser::CompileDeclaration(TArray<FUnrealFunctionDefinitionInfo*>& D
 		return true;
 	}
 
-	if (Token.Matches(TEXT("UENUM"), ESearchCase::CaseSensitive))
+	if (Token.IsIdentifier(TEXT("UENUM"), ESearchCase::CaseSensitive))
 	{
 		// Enumeration definition.
 		CompileEnum();
 		return true;
 	}
 
-	if (Token.Matches(TEXT("USTRUCT"), ESearchCase::CaseSensitive))
+	if (Token.IsIdentifier(TEXT("USTRUCT"), ESearchCase::CaseSensitive))
 	{
 		// Struct definition.
 		CompileStructDeclaration();
 		return true;
 	}
 
-	if (Token.Matches(TEXT('#')))
+	if (Token.IsSymbol(TEXT('#')))
 	{
 		// Compiler directive.
 		CompileDirective();
 		return true;
 	}
 
-	if (bEncounteredNewStyleClass_UnmatchedBrackets && Token.Matches(TEXT('}')))
+	if (bEncounteredNewStyleClass_UnmatchedBrackets && Token.IsSymbol(TEXT('}')))
 	{
 		FUnrealClassDefinitionInfo& CurrentClassDef = GetCurrentClassDef();
 		CurrentClassDef.GetDefinitionRange().End = &Input[InputPos];
@@ -4768,11 +4742,11 @@ bool FHeaderParser::CompileDeclaration(TArray<FUnrealFunctionDefinitionInfo*>& D
 		return true;
 	}
 
-	if (Token.Matches(TEXT(';')))
+	if (Token.IsSymbol(TEXT(';')))
 	{
 		if (GetToken(Token))
 		{
-			Throwf(TEXT("Extra ';' before '%s'"), Token.Identifier);
+			Throwf(TEXT("Extra ';' before '%s'"), *Token.GetTokenValue());
 		}
 		else
 		{
@@ -4786,14 +4760,14 @@ bool FHeaderParser::CompileDeclaration(TArray<FUnrealFunctionDefinitionInfo*>& D
 		FToken ConstructorToken = Token;
 
 		// Allow explicit constructors
-		bool bFoundExplicit = ConstructorToken.Matches(TEXT("explicit"), ESearchCase::CaseSensitive);
+		bool bFoundExplicit = ConstructorToken.IsIdentifier(TEXT("explicit"), ESearchCase::CaseSensitive);
 		if (bFoundExplicit)
 		{
 			GetToken(ConstructorToken);
 		}
 
 		bool bSkippedAPIToken = false;
-		if (FString(ConstructorToken.Identifier).EndsWith(TEXT("_API"), ESearchCase::CaseSensitive))
+		if (ConstructorToken.Value.EndsWith(TEXT("_API"), ESearchCase::CaseSensitive))
 		{
 			if (!bFoundExplicit)
 			{
@@ -4805,7 +4779,7 @@ bool FHeaderParser::CompileDeclaration(TArray<FUnrealFunctionDefinitionInfo*>& D
 			bSkippedAPIToken = true;
 		}
 
-		if (ConstructorToken.Matches(*StructDef.GetAlternateNameCPP(), ESearchCase::IgnoreCase))
+		if (ConstructorToken.IsIdentifier(*StructDef.GetAlternateNameCPP(), ESearchCase::IgnoreCase))
 		{
 			if (TryToMatchConstructorParameterList(ConstructorToken))
 			{
@@ -4828,76 +4802,76 @@ bool FHeaderParser::CompileDeclaration(TArray<FUnrealFunctionDefinitionInfo*>& D
 	// Determine if this statement is a serialize function declaration
 	if (bEncounteredNewStyleClass_UnmatchedBrackets && IsInAClass() && TopNest->NestType == ENestType::Class)
 	{
-		while (Token.Matches(TEXT("virtual"), ESearchCase::CaseSensitive) || FString(Token.Identifier).EndsWith(TEXT("_API"), ESearchCase::CaseSensitive))
+		while (Token.IsIdentifier(TEXT("virtual"), ESearchCase::CaseSensitive) || Token.Value.EndsWith(TEXT("_API"), ESearchCase::CaseSensitive))
 		{
 			GetToken(Token);
 		}
 
-		if (Token.Matches(TEXT("void"), ESearchCase::CaseSensitive))
+		if (Token.IsIdentifier(TEXT("void"), ESearchCase::CaseSensitive))
 		{
 			GetToken(Token);
-			if (Token.Matches(TEXT("Serialize"), ESearchCase::CaseSensitive))
+			if (Token.IsIdentifier(TEXT("Serialize"), ESearchCase::CaseSensitive))
 			{
 				GetToken(Token);
-				if (Token.Matches(TEXT('(')))
+				if (Token.IsSymbol(TEXT('(')))
 				{
 					GetToken(Token);
 
 					ESerializerArchiveType ArchiveType = ESerializerArchiveType::None;
-					if (Token.Matches(TEXT("FArchive"), ESearchCase::CaseSensitive))
+					if (Token.IsIdentifier(TEXT("FArchive"), ESearchCase::CaseSensitive))
 					{
 						GetToken(Token);
-						if (Token.Matches(TEXT('&')))
+						if (Token.IsSymbol(TEXT('&')))
 						{
 							GetToken(Token);
 
 							// Allow the declaration to not define a name for the archive parameter
-							if (!Token.Matches(TEXT(')')))
+							if (!Token.IsSymbol(TEXT(')')))
 							{
 								GetToken(Token);
 							}
 
-							if (Token.Matches(TEXT(')')))
+							if (Token.IsSymbol(TEXT(')')))
 							{
 								ArchiveType = ESerializerArchiveType::Archive;
 							}
 						}
 					}
-					else if (Token.Matches(TEXT("FStructuredArchive"), ESearchCase::CaseSensitive))
+					else if (Token.IsIdentifier(TEXT("FStructuredArchive"), ESearchCase::CaseSensitive))
 					{
 						GetToken(Token);
-						if (Token.Matches(TEXT("::"), ESearchCase::CaseSensitive))
+						if (Token.IsSymbol(TEXT("::"), ESearchCase::CaseSensitive))
 						{
 							GetToken(Token);
 
-							if (Token.Matches(TEXT("FRecord"), ESearchCase::CaseSensitive))
+							if (Token.IsIdentifier(TEXT("FRecord"), ESearchCase::CaseSensitive))
 							{
 								GetToken(Token);
 
 								// Allow the declaration to not define a name for the slot parameter
-								if (!Token.Matches(TEXT(')')))
+								if (!Token.IsSymbol(TEXT(')')))
 								{
 									GetToken(Token);
 								}
 
-								if (Token.Matches(TEXT(')')))
+								if (Token.IsSymbol(TEXT(')')))
 								{
 									ArchiveType = ESerializerArchiveType::StructuredArchiveRecord;
 								}
 							}
 						}
 					}
-					else if (Token.Matches(TEXT("FStructuredArchiveRecord"), ESearchCase::CaseSensitive))
+					else if (Token.IsIdentifier(TEXT("FStructuredArchiveRecord"), ESearchCase::CaseSensitive))
 					{
 						GetToken(Token);
 
 						// Allow the declaration to not define a name for the slot parameter
-						if (!Token.Matches(TEXT(')')))
+						if (!Token.IsSymbol(TEXT(')')))
 						{
 							GetToken(Token);
 						}
 
-						if (Token.Matches(TEXT(')')))
+						if (Token.IsSymbol(TEXT(')')))
 						{
 							ArchiveType = ESerializerArchiveType::StructuredArchiveRecord;
 						}
@@ -4935,10 +4909,10 @@ bool FHeaderParser::SkipDeclaration(FToken& Token)
 	// Consume all tokens until the end of declaration/definition has been found.
 	int32 NestedScopes = 0;
 	// Check if this is a class/struct declaration in which case it can be followed by member variable declaration.	
-	bool bPossiblyClassDeclaration = Token.Matches(TEXT("class"), ESearchCase::CaseSensitive) || Token.Matches(TEXT("struct"), ESearchCase::CaseSensitive);
+	bool bPossiblyClassDeclaration = Token.IsIdentifier(TEXT("class"), ESearchCase::CaseSensitive) || Token.IsIdentifier(TEXT("struct"), ESearchCase::CaseSensitive);
 	// (known) macros can end without ; or } so use () to find the end of the declaration.
 	// However, we don't want to use it with DECLARE_FUNCTION, because we need it to be treated like a function.
-	bool bMacroDeclaration      = ProbablyAMacro(Token.Identifier) && !Token.Matches(TEXT("DECLARE_FUNCTION"), ESearchCase::CaseSensitive);
+	bool bMacroDeclaration      = ProbablyAMacro(Token.Value) && !Token.IsIdentifier(TEXT("DECLARE_FUNCTION"), ESearchCase::CaseSensitive);
 	bool bEndOfDeclarationFound = false;
 	bool bDefinitionFound       = false;
 	TCHAR OpeningBracket = bMacroDeclaration ? TEXT('(') : TEXT('{');
@@ -4948,31 +4922,31 @@ bool FHeaderParser::SkipDeclaration(FToken& Token)
 	{
 		// If we find parentheses at top-level and we think it's a class declaration then it's more likely
 		// to be something like: class UThing* GetThing();
-		if (bPossiblyClassDeclaration && NestedScopes == 0 && Token.Matches(TEXT('(')))
+		if (bPossiblyClassDeclaration && NestedScopes == 0 && Token.IsSymbol(TEXT('(')))
 		{
 			bPossiblyClassDeclaration = false;
 		}
 
 		bRetestCurrentToken = false;
-		if (Token.Matches(TEXT(';')) && NestedScopes == 0)
+		if (Token.IsSymbol(TEXT(';')) && NestedScopes == 0)
 		{
 			bEndOfDeclarationFound = true;
 			break;
 		}
 
-		if (!bMacroDeclaration && Token.Matches(TEXT("PURE_VIRTUAL"), ESearchCase::CaseSensitive) && NestedScopes == 0)
+		if (!bMacroDeclaration && Token.IsIdentifier(TEXT("PURE_VIRTUAL"), ESearchCase::CaseSensitive) && NestedScopes == 0)
 		{
 			OpeningBracket = TEXT('(');
 			ClosingBracket = TEXT(')');
 		}
 
-		if (Token.Matches(OpeningBracket))
+		if (Token.IsSymbol(OpeningBracket))
 		{
 			// This is a function definition or class declaration.
 			bDefinitionFound = true;
 			NestedScopes++;
 		}
-		else if (Token.Matches(ClosingBracket))
+		else if (Token.IsSymbol(ClosingBracket))
 		{
 			NestedScopes--;
 			if (NestedScopes == 0)
@@ -4987,7 +4961,7 @@ bool FHeaderParser::SkipDeclaration(FToken& Token)
 					GetToken(Token);
 
 					// If Strcmp returns 0, it is probably a class, else a macro.
-					bReallyEndDeclaration = !PossibleBracketToken.Matches(TEXT('{'));
+					bReallyEndDeclaration = !PossibleBracketToken.IsSymbol(TEXT('{'));
 				}
 
 				if (bReallyEndDeclaration)
@@ -5022,14 +4996,14 @@ bool FHeaderParser::SkipDeclaration(FToken& Token)
 			{
 				return false;
 			}
-			if (VariableName.TokenType != TOKEN_Identifier)
+			if (!VariableName.IsIdentifier())
 			{
 				// Not a variable name.
 				UngetToken(VariableName);
 			}
 			else if (!SafeMatchSymbol(TEXT(';')))
 			{
-				Throwf(TEXT("Unexpected '%s'. Did you miss a semi-colon?"), VariableName.Identifier);
+				Throwf(TEXT("Unexpected '%s'. Did you miss a semi-colon?"), *VariableName.GetTokenValue());
 			}
 		}
 
@@ -5047,24 +5021,22 @@ bool FHeaderParser::SkipDeclaration(FToken& Token)
 
 bool FHeaderParser::SafeMatchSymbol( const TCHAR Match )
 {
-	FToken Token;
-
 	// Remember the position before the next token (this can include comments before the next symbol).
 	FScriptLocation LocationBeforeNextSymbol;
 	InitScriptLocation(LocationBeforeNextSymbol);
 
+	FToken Token;
 	if (GetToken(Token, /*bNoConsts=*/ true))
 	{
-		if (Token.TokenType==TOKEN_Symbol && Token.Identifier[0] == Match && Token.Identifier[1] == 0)
+		if (Token.IsSymbol(Match))
 		{
 			return true;
 		}
-
 		UngetToken(Token);
 	}
+
 	// Return to the stored position.
 	ReturnToLocation(LocationBeforeNextSymbol);
-
 	return false;
 }
 
@@ -5437,7 +5409,7 @@ void FHeaderParser::CompileRigVMMethodDeclaration(FUnrealStructDefinitionInfo& S
 		return;
 	}
 
-	if (FString(PrefixToken.Identifier).Equals(TEXT("virtual")))
+	if (PrefixToken.IsIdentifier(TEXT("virtual"), ESearchCase::CaseSensitive))
 	{
 		if (!GetToken(ReturnTypeToken))
 		{
@@ -5467,10 +5439,10 @@ void FHeaderParser::CompileRigVMMethodDeclaration(FUnrealStructDefinitionInfo& S
 		{
 			break;
 		}
-		ParamsContent.Add(FString(Token.Identifier));
+		ParamsContent.Add(FString(Token.Value));
 	}
 
-	while (!FString(PostfixToken.Identifier).Equals(TEXT(";")))
+	while (!PostfixToken.IsSymbol(TEXT(';')))
 	{
 		if (!GetToken(PostfixToken))
 		{
@@ -5479,8 +5451,8 @@ void FHeaderParser::CompileRigVMMethodDeclaration(FUnrealStructDefinitionInfo& S
 	}
 
 	FRigVMMethodInfo MethodInfo;
-	MethodInfo.ReturnType = ReturnTypeToken.Identifier;
-	MethodInfo.Name = NameToken.Identifier;
+	MethodInfo.ReturnType = FString(ReturnTypeToken.Value);
+	MethodInfo.Name = FString(NameToken.Value);
 	
 	FString ParamString = FString::Join(ParamsContent, TEXT(" "));
 	if (!ParamString.IsEmpty())
@@ -5638,7 +5610,7 @@ void FHeaderParser::ParseRigVMMethodParameters(FUnrealStructDefinitionInfo& Stru
 // Returns true if the token is a dynamic delegate declaration
 bool FHeaderParser::IsValidDelegateDeclaration(const FToken& Token) const
 {
-	return (Token.TokenType == TOKEN_Identifier) && !FCString::Strncmp(Token.Identifier, TEXT("DECLARE_DYNAMIC_"), 16);
+	return Token.IsIdentifier() && Token.ValueStartsWith(TEXT("DECLARE_DYNAMIC_"), ESearchCase::CaseSensitive);
 }
 
 // Parse the parameter list of a function or delegate declaration
@@ -5783,18 +5755,18 @@ void FHeaderParser::ParseParameterList(FUnrealFunctionDefinitionInfo& FunctionDe
 					StartPos = SkipToken.StartPos;
 				}
 				if ( ParenthesisNestCount == 0
-					&& (SkipToken.Matches(TEXT(')')) || SkipToken.Matches(TEXT(','))) )
+					&& (SkipToken.IsSymbol(TEXT(')')) || SkipToken.IsSymbol(TEXT(','))) )
 				{
-					EndPos = SkipToken.StartPos;
 					// went too far
 					UngetToken(SkipToken);
 					break;
 				}
-				if ( SkipToken.Matches(TEXT('(')) )
+				EndPos = InputPos;
+				if ( SkipToken.IsSymbol(TEXT('(')) )
 				{
 					ParenthesisNestCount++;
 				}
-				else if ( SkipToken.Matches(TEXT(')')) )
+				else if ( SkipToken.IsSymbol(TEXT(')')) )
 				{
 					ParenthesisNestCount--;
 				}
@@ -5826,7 +5798,7 @@ void FHeaderParser::ParseParameterList(FUnrealFunctionDefinitionInfo& FunctionDe
 	} while( MatchSymbol(TEXT(',')) );
 	RequireSymbol( TEXT(')'), TEXT("parameter list") );
 }
-FUnrealFunctionDefinitionInfo& FHeaderParser::CompileDelegateDeclaration(const TCHAR* DelegateIdentifier, EDelegateSpecifierAction::Type SpecifierAction)
+FUnrealFunctionDefinitionInfo& FHeaderParser::CompileDelegateDeclaration(const FStringView& DelegateIdentifier, EDelegateSpecifierAction::Type SpecifierAction)
 {
 	const TCHAR* CurrentScopeName = TEXT("Delegate Declaration");
 
@@ -5849,10 +5821,10 @@ FUnrealFunctionDefinitionInfo& FHeaderParser::CompileDelegateDeclaration(const T
 		GetToken(Token);
 		if (!IsValidDelegateDeclaration(Token))
 		{
-			Throwf(TEXT("Unexpected token following UDELEGATE(): %s"), Token.Identifier);
+			Throwf(TEXT("Unexpected token following UDELEGATE(): %s"), *Token.GetTokenValue());
 		}
 
-		DelegateMacro = Token.Identifier;
+		DelegateMacro = FString(Token.Value);
 
 		//Workaround for UE-28897
 		const FStructScope* CurrentStructScope = TopNest->GetScope() ? TopNest->GetScope()->AsStructScope() : nullptr;
@@ -5916,7 +5888,7 @@ FUnrealFunctionDefinitionInfo& FHeaderParser::CompileDelegateDeclaration(const T
 	}
 
 	// Skip whitespaces to get InputPos exactly on beginning of function name.
-	while (FChar::IsWhitespace(PeekChar())) { GetChar(); }
+	SkipWhitespaceAndComments();
 
 	FuncInfo.InputPos = InputPos;
 
@@ -5926,6 +5898,7 @@ FUnrealFunctionDefinitionInfo& FHeaderParser::CompileDelegateDeclaration(const T
 	{
 		Throwf(TEXT("Missing name for %s"), CurrentScopeName );
 	}
+	FString FuncName(FuncNameToken.Value);
 
 	// If this is a delegate function then go ahead and mangle the name so we don't collide with
 	// actual functions or properties
@@ -5933,23 +5906,19 @@ FUnrealFunctionDefinitionInfo& FHeaderParser::CompileDelegateDeclaration(const T
 		//@TODO: UCREMOVAL: Eventually this mangling shouldn't occur
 
 		// Remove the leading F
-		FString Name(FuncNameToken.Identifier);
 
-		if (!Name.StartsWith(TEXT("F"), ESearchCase::CaseSensitive))
+		if (!FuncName.StartsWith(TEXT("F"), ESearchCase::CaseSensitive))
 		{
 			Throwf(TEXT("Delegate type declarations must start with F"));
 		}
 
-		Name.RightChopInline(1, false);
+		FuncName.RightChopInline(1, false);
 
 		// Append the signature goo
-		Name += HEADER_GENERATED_DELEGATE_SIGNATURE_SUFFIX;
-
-		// Replace the name
-		FCString::Strcpy(FuncNameToken.Identifier, *Name );
+		FuncName += HEADER_GENERATED_DELEGATE_SIGNATURE_SUFFIX;
 	}
 
-	FUnrealFunctionDefinitionInfo& DelegateSignatureFunctionDef = CreateFunction(FuncNameToken.Identifier, MoveTemp(FuncInfo), bIsSparse ? EFunctionType::SparseDelegate : EFunctionType::Delegate);
+	FUnrealFunctionDefinitionInfo& DelegateSignatureFunctionDef = CreateFunction(*FuncName, MoveTemp(FuncInfo), bIsSparse ? EFunctionType::SparseDelegate : EFunctionType::Delegate);
 	DelegateSignatureFunctionDef.GetDefinitionRange().Start = StartPos;
 
 	// determine whether this function should be 'const'
@@ -5976,8 +5945,8 @@ FUnrealFunctionDefinitionInfo& FHeaderParser::CompileDelegateDeclaration(const T
 			Throwf(TEXT("Missing Delegate Name."));
 		}
 
-		DelegateSignatureFunctionDef.SetSparseOwningClassName(*GetClassNameWithoutPrefix(FString(OwningClass.Identifier)));
-		DelegateSignatureFunctionDef.SetSparseDelegateName(DelegateName.Identifier);
+		DelegateSignatureFunctionDef.SetSparseOwningClassName(*GetClassNameWithoutPrefix(FString(OwningClass.Value)));
+		DelegateSignatureFunctionDef.SetSparseDelegateName(FName(DelegateName.Value, FNAME_Add));
 	}
 
 	DelegateSignatureFunctionDef.SetLineNumber(InputLine);
@@ -6202,9 +6171,9 @@ FUnrealFunctionDefinitionInfo& FHeaderParser::CompileFunctionDeclaration()
 		if (GetToken(Token, true))
 		{
 			bool bThrowTokenBack = true;
-			if (Token.TokenType == TOKEN_Identifier)
+			if (Token.IsIdentifier())
 			{
-				FString RequiredAPIMacroIfPresent(Token.Identifier);
+				FString RequiredAPIMacroIfPresent(Token.Value);
 				if (RequiredAPIMacroIfPresent.EndsWith(TEXT("_API"), ESearchCase::CaseSensitive))
 				{
 					//@TODO: Validate the module name for RequiredAPIMacroIfPresent
@@ -6296,7 +6265,7 @@ FUnrealFunctionDefinitionInfo& FHeaderParser::CompileFunctionDeclaration()
 	bool bHasReturnValue = ReturnType.Type != CPT_None;
 
 	// Skip whitespaces to get InputPos exactly on beginning of function name.
-	while (FChar::IsWhitespace(PeekChar())) { GetChar(); }
+	SkipWhitespaceAndComments();
 
 	FuncInfo.InputPos = InputPos;
 
@@ -6306,6 +6275,7 @@ FUnrealFunctionDefinitionInfo& FHeaderParser::CompileFunctionDeclaration()
 	{
 		Throwf(TEXT("Missing %s name"), TypeOfFunction);
 	}
+	FString FuncName(FuncNameToken.Value);
 
 	const TCHAR* StartPos = &Input[InputPos];
 
@@ -6329,7 +6299,7 @@ FUnrealFunctionDefinitionInfo& FHeaderParser::CompileFunctionDeclaration()
 				Throwf(TEXT("Function %s already uses identifier %d"), **ExistingFunc, FuncInfo.RPCId);
 			}
 
-			UsedRPCIds.Add(FuncInfo.RPCId, FuncNameToken.Identifier);
+			UsedRPCIds.Add(FuncInfo.RPCId, FuncName);
 			if (FuncInfo.FunctionFlags & FUNC_NetResponse)
 			{
 				// Look for another function expecting this response
@@ -6348,12 +6318,12 @@ FUnrealFunctionDefinitionInfo& FHeaderParser::CompileFunctionDeclaration()
 			if (ExistingFunc == NULL)
 			{
 				// If this list isn't empty at end of class, throw error
-				RPCsNeedingHookup.Add(FuncInfo.RPCResponseId, FuncNameToken.Identifier);
+				RPCsNeedingHookup.Add(FuncInfo.RPCResponseId, FuncName);
 			}
 		}
 	}
 
-	FUnrealFunctionDefinitionInfo& FuncDef = CreateFunction(FuncNameToken.Identifier, MoveTemp(FuncInfo), EFunctionType::Function);
+	FUnrealFunctionDefinitionInfo& FuncDef = CreateFunction(*FuncName, MoveTemp(FuncInfo), EFunctionType::Function);
 	FuncDef.GetDefinitionRange().Start = StartPos;
 	const FFuncInfo& FuncDefFuncInfo = FuncDef.GetFunctionData();
 	FUnrealFunctionDefinitionInfo* SuperFuncDef = FuncDef.GetSuperFunction();
@@ -6505,11 +6475,11 @@ FUnrealFunctionDefinitionInfo& FHeaderParser::CompileFunctionDeclaration()
 
 	// Optionally consume a semicolon
 	// This is optional to allow inline function definitions
-	if (Token.TokenType == TOKEN_Symbol && Token.Matches(TEXT(';')))
+	if (Token.IsSymbol(TEXT(';')))
 	{
 		// Do nothing (consume it)
 	}
-	else if (Token.TokenType == TOKEN_Symbol && Token.Matches(TEXT('{')))
+	else if (Token.IsSymbol(TEXT('{')))
 	{
 		// Skip inline function bodies
 		UngetToken(Token);
@@ -6530,7 +6500,7 @@ FUnrealFunctionDefinitionInfo& FHeaderParser::CompileFunctionDeclaration()
 /** Parses optional metadata text. */
 void FHeaderParser::ParseFieldMetaData(TMap<FName, FString>& MetaData, const TCHAR* FieldName)
 {
-	FToken PropertyMetaData;
+	FTokenString PropertyMetaData;
 	bool bMetadataPresent = false;
 	if (MatchIdentifier(TEXT("UMETA"), ESearchCase::CaseSensitive))
 	{
@@ -6538,7 +6508,7 @@ void FHeaderParser::ParseFieldMetaData(TMap<FName, FString>& MetaData, const TCH
 
 		bMetadataPresent = true;
 		RequireSymbol( TEXT('('), ErrorMessageGetter );
-		if (!GetRawTokenRespectingQuotes(PropertyMetaData, TCHAR(')')))
+		if (!GetRawStringRespectingQuotes(PropertyMetaData, TCHAR(')')))
 		{
 			Throwf(TEXT("'%s': No metadata specified"), FieldName);
 		}
@@ -6624,7 +6594,7 @@ bool FHeaderParser::IsBitfieldProperty(ELayoutMacroType LayoutMacroType)
 		FToken Token;
 		if (GetToken(Token, /*bNoConsts=*/ true))
 		{
-			if (Token.TokenType == TOKEN_Symbol && Token.Matches(TEXT(':')))
+			if (Token.IsSymbol(TEXT(':')))
 			{
 				bIsBitfield = true;
 			}
@@ -6884,11 +6854,11 @@ void FHeaderParser::CompileVariableDeclaration(FUnrealStructDefinitionInfo& Stru
 		int Nesting = 1;
 		while (GetToken(SkipToken))
 		{
-			if (SkipToken.Matches(TEXT('(')))
+			if (SkipToken.IsSymbol(TEXT('(')))
 			{
 				++Nesting;
 			}
-			else if (SkipToken.Matches(TEXT(')')))
+			else if (SkipToken.IsSymbol(TEXT(')')))
 			{
 				--Nesting;
 				if (Nesting == 0)
@@ -6905,7 +6875,7 @@ void FHeaderParser::CompileVariableDeclaration(FUnrealStructDefinitionInfo& Stru
 		FToken SkipToken;
 		while (GetToken(SkipToken))
 		{
-			if (SkipToken.Matches(TEXT(';')))
+			if (SkipToken.IsSymbol(TEXT(';')))
 			{
 				// went too far
 				UngetToken(SkipToken);
@@ -6920,11 +6890,11 @@ void FHeaderParser::CompileVariableDeclaration(FUnrealStructDefinitionInfo& Stru
 		int BraceLevel = 1;
 		while (GetToken(SkipToken))
 		{
-			if (SkipToken.Matches(TEXT('{')))
+			if (SkipToken.IsSymbol(TEXT('{')))
 			{
 				++BraceLevel;
 			}
-			else if (SkipToken.Matches(TEXT('}')))
+			else if (SkipToken.IsSymbol(TEXT('}')))
 			{
 				--BraceLevel;
 				if (BraceLevel == 0)
@@ -6958,7 +6928,7 @@ void FHeaderParser::CompileVariableDeclaration(FUnrealStructDefinitionInfo& Stru
 			break;
 		}
 
-		if (Token.TokenType != TOKEN_Symbol || !Token.Matches(TEXT(';')))
+		if (!Token.IsSymbol(TEXT(';')))
 		{
 			InputPos  = CurrInputPos;
 			InputLine = CurrInputLine;
@@ -6982,7 +6952,7 @@ bool FHeaderParser::CompileStatement(TArray<FUnrealFunctionDefinitionInfo*>& Del
 	}
 	else if (!CompileDeclaration(DelegatesToFixup, Token))
 	{
-		Throwf(TEXT("'%s': Bad command or expression"), Token.Identifier );
+		Throwf(TEXT("'%s': Bad command or expression"), *Token.GetTokenValue() );
 	}
 	return true;
 }
@@ -7005,15 +6975,15 @@ void FHeaderParser::SkipStatements( int32 NestCount, const TCHAR* ErrorTag  )
 
 	while( GetToken( Token, true ) )
 	{
-		if ( Token.Matches(TEXT('{')) )
+		if ( Token.IsSymbol(TEXT('{')) )
 		{
 			NestCount++;
 		}
-		else if	( Token.Matches(TEXT('}')) )
+		else if	( Token.IsSymbol(TEXT('}')) )
 		{
 			NestCount--;
 		}
-		else if ( Token.Matches(TEXT(';')) && OriginalNestCount == 0 )
+		else if ( Token.IsSymbol(TEXT(';')) && OriginalNestCount == 0 )
 		{
 			break;
 		}
@@ -7972,7 +7942,7 @@ TSharedRef<FUnrealTypeDefinitionInfo> FHeaderPreParser::ParseEnumDeclaration(con
 		Throwf(TEXT("Missing identifier after UENUM()"));
 	}
 
-	if (EnumToken.Matches(TEXT("namespace"), ESearchCase::CaseSensitive))
+	if (EnumToken.IsValue(TEXT("namespace"), ESearchCase::CaseSensitive))
 	{
 		CppForm = UEnum::ECppForm::Namespaced;
 
@@ -7980,14 +7950,14 @@ TSharedRef<FUnrealTypeDefinitionInfo> FHeaderPreParser::ParseEnumDeclaration(con
 
 		bReadEnumName = GetIdentifier(EnumToken);
 	}
-	else if (EnumToken.Matches(TEXT("enum"), ESearchCase::CaseSensitive))
+	else if (EnumToken.IsValue(TEXT("enum"), ESearchCase::CaseSensitive))
 	{
 		if (!GetIdentifier(EnumToken))
 		{
 			Throwf(TEXT("Missing identifier after enum"));
 		}
 
-		if (EnumToken.Matches(TEXT("class"), ESearchCase::CaseSensitive) || EnumToken.Matches(TEXT("struct"), ESearchCase::CaseSensitive))
+		if (EnumToken.IsValue(TEXT("class"), ESearchCase::CaseSensitive) || EnumToken.IsValue(TEXT("struct"), ESearchCase::CaseSensitive))
 		{
 			CppForm = UEnum::ECppForm::EnumClass;
 		}
@@ -8021,7 +7991,7 @@ TSharedRef<FUnrealTypeDefinitionInfo> FHeaderPreParser::ParseEnumDeclaration(con
 		UnderlyingType = ParseUnderlyingEnumType();
 	}
 
-	TSharedRef<FUnrealEnumDefinitionInfo> EnumDef = MakeShareable(new FUnrealEnumDefinitionInfo(SourceFile, InLineNumber, EnumToken.Identifier, FName(EnumToken.Identifier, FNAME_Add), CppForm, UnderlyingType));
+	TSharedRef<FUnrealEnumDefinitionInfo> EnumDef = MakeShareable(new FUnrealEnumDefinitionInfo(SourceFile, InLineNumber, FString(EnumToken.Value), FName(EnumToken.Value, FNAME_Add), CppForm, UnderlyingType));
 	return EnumDef;
 }
 
@@ -8132,7 +8102,7 @@ bool FHeaderParser::TryToMatchConstructorParameterList(FToken Token)
 		return false;
 	}
 
-	if (!PotentialParenthesisToken.Matches(TEXT('(')))
+	if (!PotentialParenthesisToken.IsSymbol(TEXT('(')))
 	{
 		UngetToken(PotentialParenthesisToken);
 		return false;
@@ -8159,45 +8129,45 @@ bool FHeaderParser::TryToMatchConstructorParameterList(FToken Token)
 		while (ParenthesesNestingLevel && GetToken(ObjectInitializerParamParsingToken))
 		{
 			// Template instantiation or additional parameter excludes ObjectInitializer constructor.
-			if (ObjectInitializerParamParsingToken.Matches(TEXT(',')) || ObjectInitializerParamParsingToken.Matches(TEXT('<')))
+			if (ObjectInitializerParamParsingToken.IsSymbol(TEXT(',')) || ObjectInitializerParamParsingToken.IsSymbol(TEXT('<')))
 			{
 				bOICtor = false;
 				bVTCtor = false;
 				break;
 			}
 
-			if (ObjectInitializerParamParsingToken.Matches(TEXT('(')))
+			if (ObjectInitializerParamParsingToken.IsSymbol(TEXT('(')))
 			{
 				ParenthesesNestingLevel++;
 				continue;
 			}
 
-			if (ObjectInitializerParamParsingToken.Matches(TEXT(')')))
+			if (ObjectInitializerParamParsingToken.IsSymbol(TEXT(')')))
 			{
 				ParenthesesNestingLevel--;
 				continue;
 			}
 
-			if (ObjectInitializerParamParsingToken.Matches(TEXT("const"), ESearchCase::CaseSensitive))
+			if (ObjectInitializerParamParsingToken.IsIdentifier(TEXT("const"), ESearchCase::CaseSensitive))
 			{
 				bIsConst = true;
 				continue;
 			}
 
-			if (ObjectInitializerParamParsingToken.Matches(TEXT('&')))
+			if (ObjectInitializerParamParsingToken.IsSymbol(TEXT('&')))
 			{
 				bIsRef = true;
 				continue;
 			}
 
-			if (ObjectInitializerParamParsingToken.Matches(TEXT("FObjectInitializer"), ESearchCase::CaseSensitive)
-				|| ObjectInitializerParamParsingToken.Matches(TEXT("FPostConstructInitializeProperties"), ESearchCase::CaseSensitive) // Deprecated, but left here, so it won't break legacy code.
+			if (ObjectInitializerParamParsingToken.IsIdentifier(TEXT("FObjectInitializer"), ESearchCase::CaseSensitive)
+				|| ObjectInitializerParamParsingToken.IsIdentifier(TEXT("FPostConstructInitializeProperties"), ESearchCase::CaseSensitive) // Deprecated, but left here, so it won't break legacy code.
 				)
 			{
 				bOICtor = true;
 			}
 
-			if (ObjectInitializerParamParsingToken.Matches(TEXT("FVTableHelper"), ESearchCase::CaseSensitive))
+			if (ObjectInitializerParamParsingToken.IsIdentifier(TEXT("FVTableHelper"), ESearchCase::CaseSensitive))
 			{
 				bVTCtor = true;
 			}
@@ -8206,13 +8176,13 @@ bool FHeaderParser::TryToMatchConstructorParameterList(FToken Token)
 		// Parse until finish.
 		while (ParenthesesNestingLevel && GetToken(ObjectInitializerParamParsingToken))
 		{
-			if (ObjectInitializerParamParsingToken.Matches(TEXT('(')))
+			if (ObjectInitializerParamParsingToken.IsSymbol(TEXT('(')))
 			{
 				ParenthesesNestingLevel++;
 				continue;
 			}
 
-			if (ObjectInitializerParamParsingToken.Matches(TEXT(')')))
+			if (ObjectInitializerParamParsingToken.IsSymbol(TEXT(')')))
 			{
 				ParenthesesNestingLevel--;
 				continue;
@@ -8264,8 +8234,7 @@ void FHeaderParser::CompileVersionDeclaration(FUnrealStructDefinitionInfo& Struc
 		Version = PackageDef.GetModule().GeneratedCodeVersion;
 	}
 
-	if (Token.TokenType == ETokenType::TOKEN_Symbol
-		&& Token.Matches(TEXT(')')))
+	if (Token.IsSymbol(TEXT(')')))
 	{
 		StructDef.SetGeneratedCodeVersion(Version);
 		UngetToken(Token);
@@ -8273,7 +8242,7 @@ void FHeaderParser::CompileVersionDeclaration(FUnrealStructDefinitionInfo& Struc
 	}
 
 	// Overwrite with version specified by macro.
-	Version = ToGeneratedCodeVersion(Token.Identifier);
+	Version = ToGeneratedCodeVersion(Token.Value);
 	StructDef.SetGeneratedCodeVersion(Version);
 }
 
@@ -8708,18 +8677,6 @@ bool FHeaderParser::IsReservedTypeName(const FString& TypeName)
 	return false;
 }
 
-bool FHeaderParser::IsReservedTypeName(const FToken& Token)
-{
-	for (const FString& ReservedName : ReservedTypeNames)
-	{
-		if (Token.Matches(*ReservedName, ESearchCase::IgnoreCase))
-		{
-			return true;
-		}
-	}
-	return false;
-}
-
 bool FHeaderParser::CheckUIMinMaxRangeFromMetaData(const FString& UIMin, const FString& UIMax)
 {
 	if (UIMin.IsEmpty() || UIMax.IsEmpty())
diff --git a/Engine/Source/Programs/UnrealHeaderTool/Private/HeaderParser.h b/Engine/Source/Programs/UnrealHeaderTool/Private/HeaderParser.h
--- a/Engine/Source/Programs/UnrealHeaderTool/Private/HeaderParser.h
+++ b/Engine/Source/Programs/UnrealHeaderTool/Private/HeaderParser.h
@@ -390,7 +390,7 @@ protected:
 	FUnrealFunctionDefinitionInfo& CreateFunction(const TCHAR* FuncName, FFuncInfo&& FuncInfo, EFunctionType InFunctionType) const;
 
 	FUnrealClassDefinitionInfo& CompileClassDeclaration();
-	FUnrealFunctionDefinitionInfo& CompileDelegateDeclaration(const TCHAR* DelegateIdentifier, EDelegateSpecifierAction::Type SpecifierAction = EDelegateSpecifierAction::DontParse);
+	FUnrealFunctionDefinitionInfo& CompileDelegateDeclaration(const FStringView& DelegateIdentifier, EDelegateSpecifierAction::Type SpecifierAction = EDelegateSpecifierAction::DontParse);
 	FUnrealFunctionDefinitionInfo& CompileFunctionDeclaration();
 	void CompileVariableDeclaration (FUnrealStructDefinitionInfo& StructDef);
 	void CompileInterfaceDeclaration();
@@ -615,14 +615,6 @@ public:
 	*/
 	static bool IsReservedTypeName(const FString& TypeName);
 
-	/**
-	* Checks if the given token uses one of the reserved type names.
-	*
-	* @param  Token		The token to check
-	* @return True if the Token is using a reserved name
-	*/
-	static bool IsReservedTypeName(const FToken& Token);
-	
 	static const FName NAME_InputText;
 	static const FName NAME_OutputText;
 	static const FName NAME_ConstantText;
diff --git a/Engine/Source/Programs/UnrealHeaderTool/Private/ParserHelper.cpp b/Engine/Source/Programs/UnrealHeaderTool/Private/ParserHelper.cpp
--- a/Engine/Source/Programs/UnrealHeaderTool/Private/ParserHelper.cpp
+++ b/Engine/Source/Programs/UnrealHeaderTool/Private/ParserHelper.cpp
@@ -281,41 +281,6 @@ bool FPropertyBase::MatchesType(const FPropertyBase& Other, bool bDisallowGenera
 	}
 }
 
-
-/////////////////////////////////////////////////////
-// FToken
-
-/**
- * Copies the properties from this token into another.
- *
- * @param	Other	the token to copy this token's properties to.
- */
-FToken& FToken::operator=(const FToken& Other)
-{
-	ConstantType = Other.ConstantType;
-	TokenType = Other.TokenType;
-	StartPos = Other.StartPos;
-	StartLine = Other.StartLine;
-
-	FCString::Strncpy(Identifier, Other.Identifier, NAME_SIZE);
-	FMemory::Memcpy(String, Other.String, sizeof(String));
-
-	return *this;
-}
-
-FToken& FToken::operator=(FToken&& Other)
-{
-	ConstantType = Other.ConstantType;
-	TokenType = Other.TokenType;
-	StartPos = Other.StartPos;
-	StartLine = Other.StartLine;
-
-	FCString::Strncpy(Identifier, Other.Identifier, NAME_SIZE);
-	FMemory::Memcpy(String, Other.String, sizeof(String));
-
-	return *this;
-}
-
 /////////////////////////////////////////////////////
 // FFuncData
 
diff --git a/Engine/Source/Programs/UnrealHeaderTool/Private/ParserHelper.h b/Engine/Source/Programs/UnrealHeaderTool/Private/ParserHelper.h
--- a/Engine/Source/Programs/UnrealHeaderTool/Private/ParserHelper.h
+++ b/Engine/Source/Programs/UnrealHeaderTool/Private/ParserHelper.h
@@ -344,262 +344,6 @@ public:
 	friend struct FPropertyBaseArchiveProxy;
 };
 
-//
-// Token types.
-//
-enum ETokenType
-{
-	TOKEN_None				= 0x00,		// No token.
-	TOKEN_Identifier		= 0x01,		// Alphanumeric identifier.
-	TOKEN_Symbol			= 0x02,		// Symbol.
-	TOKEN_Const				= 0x03,		// A constant.
-	TOKEN_Max				= 0x0D
-};
-
-/*-----------------------------------------------------------------------------
-	FToken.
------------------------------------------------------------------------------*/
-//
-// Information about a token that was just parsed.
-//
-class FToken
-{
-public:
-	/** @name Variables */
-	//@{
-	/** Type of token. */
-	ETokenType TokenType;
-	/** Starting position in script where this token came from. */
-	int32 StartPos;
-	/** Starting line in script. */
-	int32 StartLine;
-	/** Always valid. */
-	TCHAR Identifier[NAME_SIZE];
-	/** Type of the constant value */
-	EPropertyType ConstantType = CPT_None;
-
-	union
-	{
-		// TOKEN_Const values.
-		uint8 Byte;								// If CPT_Byte.
-		int64 Int64;							// If CPT_Int64.
-		int32 Int;								// If CPT_Int.
-		bool NativeBool;						// if CPT_Bool
-		float Float;							// If CPT_Float.
-		double Double;							// If CPT_Double.
-		uint8 NameBytes[sizeof(FName)];			// If CPT_Name.
-		TCHAR String[MAX_STRING_CONST_SIZE];	// If CPT_String
-	};
-	//@}
-
-	FToken& operator=(const FToken& Other);
-	FToken& operator=(FToken&& Other);
-
-	FString GetConstantValue() const
-	{
-		if (TokenType == TOKEN_Const)
-		{
-			switch (ConstantType)
-			{
-				case CPT_Byte:
-					return FString::Printf(TEXT("%u"), Byte);
-				case CPT_Int64:
-					return FString::Printf(TEXT("%" INT64_FMT), Int64);
-				case CPT_Int:
-					return FString::Printf(TEXT("%i"), Int);
-				case CPT_Bool:
-					// Don't use FCoreTexts::True/FCoreTexts::False here because they can be localized
-					return FString::Printf(TEXT("%s"), NativeBool ? *(FName::GetEntry(NAME_TRUE)->GetPlainNameString()) : *(FName::GetEntry(NAME_FALSE)->GetPlainNameString()));
-				case CPT_Float:
-					return FString::Printf(TEXT("%f"), Float);
-				case CPT_Double:
-					return FString::Printf(TEXT("%f"), Double);
-				case CPT_Name:
-					return FString::Printf(TEXT("%s"), *(*(FName*)NameBytes).ToString());
-				case CPT_String:
-					return String;
-
-				// unsupported (parsing never produces a constant token of these types)
-				default:
-					return TEXT("InvalidTypeForAToken");
-			}
-		}
-		else
-		{
-			return TEXT("NotConstant");
-		}
-	}
-
-	// Constructors.
-	FToken()
-	{
-		InitToken(CPT_None);
-	}
-
-	// copy constructors
-	explicit FToken(EPropertyType InType)
-	{
-		InitToken(InType);
-	}
-
-	FToken(const FToken& InToken)
-	{
-		*this = InToken;
-	}
-
-	FToken(FToken&& InToken)
-	{
-		*this = MoveTemp(InToken);
-	}
-
-	// Inlines.
-	void InitToken( EPropertyType InType )
-	{
-		ConstantType    = InType;
-		TokenType		= TOKEN_None;
-		StartPos		= 0;
-		StartLine		= 0;
-		*Identifier		= 0;
-		FMemory::Memzero(String);
-	}
-	bool Matches(const TCHAR Ch) const
-	{
-		return TokenType==TOKEN_Symbol && Identifier[0] == Ch && Identifier[1] == 0;
-	}
-	bool Matches( const TCHAR* Str, ESearchCase::Type SearchCase ) const
-	{
-		return (TokenType==TOKEN_Identifier || TokenType==TOKEN_Symbol) && ((SearchCase == ESearchCase::CaseSensitive) ? !FCString::Strcmp(Identifier, Str) : !FCString::Stricmp(Identifier, Str));
-	}
-
-	// Setters.
-	
-	void SetIdentifier( const TCHAR* InString)
-	{
-		InitToken(CPT_None);
-		TokenType = TOKEN_Identifier;
-		FCString::Strncpy(Identifier, InString, NAME_SIZE);
-	}
-
-	void SetConstInt64( int64 InInt64 )
-	{
-		ConstantType    = CPT_Int64;
-		Int64			= InInt64;
-		TokenType		= TOKEN_Const;
-	}
-	void SetConstInt( int32 InInt )
-	{
-		ConstantType    = CPT_Int;
-		Int				= InInt;
-		TokenType		= TOKEN_Const;
-	}
-	void SetConstBool( bool InBool )
-	{
-		ConstantType    = CPT_Bool;
-		NativeBool 		= InBool;
-		TokenType		= TOKEN_Const;
-	}
-	void SetConstFloat( float InFloat )
-	{
-		ConstantType    = CPT_Float;
-		Float			= InFloat;
-		TokenType		= TOKEN_Const;
-	}
-	void SetConstDouble( double InDouble )
-	{
-		ConstantType    = CPT_Double;
-		Double			= InDouble;
-		TokenType		= TOKEN_Const;
-	}
-	void SetConstName( FName InName )
-	{
-		ConstantType    = CPT_Name;
-		*(FName *)NameBytes = InName;
-		TokenType		= TOKEN_Const;
-	}
-	void SetConstString( const TCHAR* InString, int32 MaxLength=MAX_STRING_CONST_SIZE )
-	{
-		check(MaxLength>0);
-		ConstantType = CPT_String;
-		if( InString != String )
-		{
-			FCString::Strncpy( String, InString, MaxLength );
-		}
-		TokenType = TOKEN_Const;
-	}
-	void SetConstChar(TCHAR InChar)
-	{
-		//@TODO: Treating this like a string for now, nothing consumes it
-		ConstantType = CPT_String;
-		String[0] = InChar;
-		String[1] = 0;
-		TokenType = TOKEN_Const;
-	}
-	//!!struct constants
-
-	// Getters.
-	bool GetConstInt( int32& I ) const
-	{
-		if( TokenType==TOKEN_Const && ConstantType ==CPT_Int64 )
-		{
-			I = (int32)Int64;
-			return 1;
-		}
-		else if( TokenType==TOKEN_Const && ConstantType ==CPT_Int )
-		{
-			I = Int;
-			return 1;
-		}
-		else if( TokenType==TOKEN_Const && ConstantType ==CPT_Byte )
-		{
-			I = Byte;
-			return 1;
-		}
-		else if( TokenType==TOKEN_Const && ConstantType ==CPT_Float && Float==FMath::TruncToInt(Float))
-		{
-			I = (int32) Float;
-			return 1;
-		}
-		else if (TokenType == TOKEN_Const && ConstantType == CPT_Double && Double == FMath::TruncToInt((float)Double))
-		{
-			I = (int32) Double;
-			return 1;
-		}
-		else return 0;
-	}
-
-	bool GetConstInt64( int64& I ) const
-	{
-		if( TokenType==TOKEN_Const && ConstantType ==CPT_Int64 )
-		{
-			I = Int64;
-			return 1;
-		}
-		else if( TokenType==TOKEN_Const && ConstantType ==CPT_Int )
-		{
-			I = Int;
-			return 1;
-		}
-		else if( TokenType==TOKEN_Const && ConstantType ==CPT_Byte )
-		{
-			I = Byte;
-			return 1;
-		}
-		else if( TokenType==TOKEN_Const && ConstantType ==CPT_Float && Float==FMath::TruncToInt(Float))
-		{
-			I = (int32) Float;
-			return 1;
-		}
-		else if (TokenType == TOKEN_Const && ConstantType == CPT_Double && Double == FMath::TruncToInt((float)Double))
-		{
-			I = (int32) Double;
-			return 1;
-		}
-		else return 0;
-	}
-
-	friend struct FTokenArchiveProxy;
-};
-
 /**
  * Information about a function being compiled.
  */
