Fixing MAC compile error

#rnx

diff --git a/Engine/Source/Programs/UnrealHeaderTool/Private/BaseParser.h b/Engine/Source/Programs/UnrealHeaderTool/Private/BaseParser.h
--- a/Engine/Source/Programs/UnrealHeaderTool/Private/BaseParser.h
+++ b/Engine/Source/Programs/UnrealHeaderTool/Private/BaseParser.h
@@ -97,9 +97,7 @@ struct FTokenString
 // Helper structure for null terminated value.  This will be the raw text from the token.
 struct FTokenValue
 {
-	static constexpr int32 MAX_SIZE = std::max(int32(MAX_STRING_CONST_SIZE), int32(NAME_SIZE));
-
-	TCHAR Value[MAX_SIZE];
+	TCHAR Value[NAME_SIZE];
 
 	FTokenValue() { Value[0] = 0; }
 
