Fixing issue with WRITEF/VERIFYREF where a duplicate file name would result in base verifications.

#rnx

diff --git a/Engine/Source/Programs/UnrealHeaderTool/Private/CodeGenerator.cpp b/Engine/Source/Programs/UnrealHeaderTool/Private/CodeGenerator.cpp
--- a/Engine/Source/Programs/UnrealHeaderTool/Private/CodeGenerator.cpp
+++ b/Engine/Source/Programs/UnrealHeaderTool/Private/CodeGenerator.cpp
@@ -6045,7 +6045,49 @@ bool FNativeClassHeaderGenerator::SaveHeaderIfChanged(FGeneratedFileInfo& FileIn
 
 		const FString& ProjectSavedDir = FPaths::ProjectSavedDir();
 		const FString CleanFilename = FPaths::GetCleanFilename(FileInfo.GetFilename());
-		const FString Ref = ProjectSavedDir / TEXT("ReferenceGeneratedCode") / CleanFilename;
+
+		FString PathPrefix;
+		{
+			FString StdFilename = FileInfo.GetFilename();
+
+			FPaths::MakeStandardFilename(StdFilename);
+
+			bool bRelativePath = FPaths::IsRelative(StdFilename);
+
+			if (!bRelativePath)
+			{
+				// If path is still absolute that means MakeStandardFilename has failed
+				// In this case make it relative to the current project. 
+				bRelativePath = FPaths::MakePathRelativeTo(StdFilename, *FPaths::GetPath(FPaths::GetProjectFilePath()));
+			}
+
+			// If the path has passed either MakeStandardFilename or MakePathRelativeTo it should be using internal path separators
+			if (bRelativePath)
+			{
+				// Remove any preceding parent directory paths
+				while (StdFilename.RemoveFromStart(TEXT("../")));
+			}
+
+			StdFilename = FPaths::GetPath(StdFilename);
+
+			FStringOutputDevice Out;
+
+			for (TCHAR Char : StdFilename)
+			{
+				if (FChar::IsAlnum(Char))
+				{
+					Out.AppendChar(Char);
+				}
+				else
+				{
+					Out.AppendChar('_');
+				}
+			}
+
+			PathPrefix = MoveTemp(Out);
+		}
+
+		const FString Ref = ProjectSavedDir / TEXT("ReferenceGeneratedCode") / (PathPrefix + TEXT("_") + CleanFilename);
 
 		if (bWriteContents)
 		{
@@ -6057,7 +6099,7 @@ bool FNativeClassHeaderGenerator::SaveHeaderIfChanged(FGeneratedFileInfo& FileIn
 		{
 			{
 				FScopeLock Lock(&WritePacer);
-				const FString Verify = ProjectSavedDir / TEXT("VerifyGeneratedCode") / CleanFilename;
+				const FString Verify = ProjectSavedDir / TEXT("VerifyGeneratedCode") / (PathPrefix + TEXT("_") + CleanFilename);
 				bool Written = FFileHelper::SaveStringToFile(InNewHeaderContents, *Verify);
 				check(Written);
 			}
