Improved the detection of class methods.  Resolves issue of incorrectly matching method names in comments.

#rnx

diff --git a/Engine/Source/Programs/UnrealHeaderTool/Private/BaseParser.cpp b/Engine/Source/Programs/UnrealHeaderTool/Private/BaseParser.cpp
--- a/Engine/Source/Programs/UnrealHeaderTool/Private/BaseParser.cpp
+++ b/Engine/Source/Programs/UnrealHeaderTool/Private/BaseParser.cpp
@@ -430,6 +430,7 @@ bool FBaseParser::GetToken(FToken& Token, bool bNoConsts/*=false*/, ESymbolParse
 		// If const values are allowed, determine whether the identifier represents a constant
 		Token.TokenType = ETokenType::Identifier;
 		Token.Value = FStringView(Start, int32(Pos - Start));
+		Token.InputLine = InputLine;
 		if (!bNoConsts)
 		{
 			// See if the identifier is part of a vector, rotation or other struct constant.
@@ -451,6 +452,10 @@ bool FBaseParser::GetToken(FToken& Token, bool bNoConsts/*=false*/, ESymbolParse
 		}
 
 		InputPos = int32(Pos - Input);
+		if (bRecordTokens)
+		{
+			RecordedTokens.Push(Token);
+		}
 		return true;
 	}
 
@@ -481,6 +486,7 @@ bool FBaseParser::GetToken(FToken& Token, bool bNoConsts/*=false*/, ESymbolParse
 
 		Token.TokenType = bIsFloat ? ETokenType::FloatConst : (bIsHex ? ETokenType::HexConst : ETokenType::DecimalConst);
 		Token.Value = FStringView(Start, int32(Pos - Start));
+		Token.InputLine = InputLine;
 
 		if (Token.Value.Len() >= NAME_SIZE)
 		{
@@ -488,6 +494,10 @@ bool FBaseParser::GetToken(FToken& Token, bool bNoConsts/*=false*/, ESymbolParse
 		}
 
 		InputPos = int32(Pos - Input);
+		if (bRecordTokens)
+		{
+			RecordedTokens.Push(Token);
+		}
 		return true;
 	}
 
@@ -506,7 +516,12 @@ bool FBaseParser::GetToken(FToken& Token, bool bNoConsts/*=false*/, ESymbolParse
 
 		Token.TokenType = ETokenType::CharConst;
 		Token.Value = FStringView(Start, int32(Pos - Start));
+		Token.InputLine = InputLine;
 		InputPos = int32(Pos - Input);
+		if (bRecordTokens)
+		{
+			RecordedTokens.Push(Token);
+		}
 		return true;
 	}
 
@@ -541,6 +556,7 @@ bool FBaseParser::GetToken(FToken& Token, bool bNoConsts/*=false*/, ESymbolParse
 
 		Token.TokenType = ETokenType::StringConst;
 		Token.Value = FStringView(Start, int32(Pos - Start));
+		Token.InputLine = InputLine;
 
 		if (!Token.Value.Len() - EscapeCount >= MAX_STRING_CONST_SIZE)
 		{
@@ -548,6 +564,10 @@ bool FBaseParser::GetToken(FToken& Token, bool bNoConsts/*=false*/, ESymbolParse
 		}
 
 		InputPos = int32(Pos - Input);
+		if (bRecordTokens)
+		{
+			RecordedTokens.Push(Token);
+		}
 		return true;
 	}
 	else
@@ -589,7 +609,12 @@ bool FBaseParser::GetToken(FToken& Token, bool bNoConsts/*=false*/, ESymbolParse
 
 		Token.TokenType = ETokenType::Symbol;
 		Token.Value = FStringView(Start, int32(Pos - Start));
+		Token.InputLine = InputLine;
 		InputPos = int32(Pos - Input);
+		if (bRecordTokens)
+		{
+			RecordedTokens.Push(Token);
+		}
 		return true;
 	}
 }
@@ -902,12 +927,26 @@ void FBaseParser::UngetToken( const FToken& Token )
 {
 	InputPos = Token.StartPos;
 	InputLine = Token.StartLine;
+	if (bRecordTokens)
+	{
+		while (!RecordedTokens.IsEmpty() && RecordedTokens.Last().StartPos >= InputPos)
+		{
+			RecordedTokens.RemoveAt(RecordedTokens.Num() - 1);
+		}
+	}
 }
 
 void FBaseParser::UngetToken(int32 StartLine, int32 StartPos)
 {
 	InputLine = StartLine;
 	InputPos = StartPos;
+	if (bRecordTokens)
+	{
+		while (!RecordedTokens.IsEmpty() && RecordedTokens.Last().StartPos >= InputPos)
+		{
+			RecordedTokens.RemoveAt(RecordedTokens.Num() - 1);
+		}
+	}
 }
 
 //
diff --git a/Engine/Source/Programs/UnrealHeaderTool/Private/BaseParser.h b/Engine/Source/Programs/UnrealHeaderTool/Private/BaseParser.h
--- a/Engine/Source/Programs/UnrealHeaderTool/Private/BaseParser.h
+++ b/Engine/Source/Programs/UnrealHeaderTool/Private/BaseParser.h
@@ -122,6 +122,9 @@ public:
 	/** Starting line in script. */
 	int32 StartLine = 0;
 
+	/** Input line of the token */
+	int32 InputLine = 0;
+
 	/** The text of the token */
 	FStringView Value;
 
@@ -587,6 +590,9 @@ public:
 
 	// Initialize the metadata keywords prior to parsing
 	static void InitMetadataKeywords();
+
+	TArray<FToken> RecordedTokens;
+	bool bRecordTokens = false;
 };
 
 template <typename Lambda>
@@ -667,3 +673,72 @@ void FBaseParser::ParseInheritance(const TCHAR* What, Lambda&& InLambda)
 		InLambda(*InterfaceName, false);
 	}
 }
+
+class FTokenReplay
+{
+public:
+	explicit FTokenReplay(const TArray<FToken>& InTokens)
+		: Tokens(InTokens)
+	{
+	}
+
+	bool GetToken(FToken& Token)
+	{
+		if (CurrentIndex < Tokens.Num())
+		{
+			Token = Tokens[CurrentIndex++];
+			return true;
+		}
+		else
+		{
+			Token = FToken();
+			return false;
+		}
+	}
+
+	void UngetToken(const FToken& Token)
+	{
+		while (CurrentIndex > 0)
+		{
+			--CurrentIndex;
+			if (Tokens[CurrentIndex].StartPos == Token.StartPos)
+			{
+				break;
+			}
+		}
+	}
+
+	bool MatchIdentifier(const TCHAR* Match, ESearchCase::Type SearchCase)
+	{
+		if (CurrentIndex < Tokens.Num() && Tokens[CurrentIndex].IsIdentifier(Match, SearchCase))
+		{
+			++CurrentIndex;
+			return true;
+		}
+		return false;
+	}
+
+	bool MatchSymbol(const TCHAR Match)
+	{
+		if (CurrentIndex < Tokens.Num() && Tokens[CurrentIndex].IsSymbol(Match))
+		{
+			++CurrentIndex;
+			return true;
+		}
+		return false;
+	}
+
+	bool MatchSymbol(const TCHAR* Match)
+	{
+		if (CurrentIndex < Tokens.Num() && Tokens[CurrentIndex].IsSymbol(Match, ESearchCase::CaseSensitive))
+		{
+			++CurrentIndex;
+			return true;
+		}
+		return false;
+	}
+
+private:
+	const TArray<FToken>& Tokens;
+	int CurrentIndex = 0;
+};
\ No newline at end of file
diff --git a/Engine/Source/Programs/UnrealHeaderTool/Private/CodeGenerator.cpp b/Engine/Source/Programs/UnrealHeaderTool/Private/CodeGenerator.cpp
--- a/Engine/Source/Programs/UnrealHeaderTool/Private/CodeGenerator.cpp
+++ b/Engine/Source/Programs/UnrealHeaderTool/Private/CodeGenerator.cpp
@@ -83,13 +83,30 @@ void ProcessParsedStruct(FUnrealScriptStructDefinitionInfo& ScriptStructDef);
 // Array of all the temporary header async file tasks so we can ensure they have completed before issuing our timings
 static FGraphEventArray GAsyncFileTasks;
 
-bool HasIdentifierExactMatch(const TCHAR* StringBegin, const TCHAR* StringEnd, const FString& Find);
-
 // Globals for common class definitions
 FUnrealClassDefinitionInfo* GUObjectDef = nullptr;
 FUnrealClassDefinitionInfo* GUClassDef = nullptr;
 FUnrealClassDefinitionInfo* GUInterfaceDef = nullptr;
 
+struct FFindDelcarationResults
+{
+	bool bVirtualFound = false;
+	const FDeclaration* Declaration = nullptr;
+	int FunctionNameTokenIndex = -1;
+
+	bool WasFound() const
+	{
+		return Declaration != nullptr;
+	}
+
+	bool IsVirtual() const
+	{
+		return bVirtualFound;
+	}
+};
+
+FFindDelcarationResults FindDeclaration(const FUnrealStructDefinitionInfo& StructDef, const FString& Identifier);
+
 namespace
 {
 	static const FName NAME_SerializeToFArchive("SerializeToFArchive");
@@ -268,7 +285,7 @@ namespace
 		const FUnrealSourceFile& SourceFile,
 		EExportClassOutFlags& OutFlags)
 	{
-		const bool bHasGetLifetimeReplicatedProps = HasIdentifierExactMatch(ClassDef.GetDefinitionRange().Start, ClassDef.GetDefinitionRange().End, STRING_GetLifetimeReplicatedPropsStr);
+		const bool bHasGetLifetimeReplicatedProps = FindDeclaration(ClassDef, STRING_GetLifetimeReplicatedPropsStr).WasFound();
 
 		if (!bHasGetLifetimeReplicatedProps)
 		{
@@ -363,110 +380,36 @@ void FGeneratedCPP::AddExportTaskRef(FGraphEventArray& Events) const
 #define BEGIN_WRAP_EDITOR_ONLY(DoWrap) DoWrap ? BeginEditorOnlyGuard : TEXT("")
 #define END_WRAP_EDITOR_ONLY(DoWrap) DoWrap ? EndEditorOnlyGuard : TEXT("")
 
-/**
- * Finds exact match of Identifier in string. Returns nullptr if none is found.
- *
- * @param StringBegin Start of string to search.
- * @param StringEnd End of string to search.
- * @param Identifier Identifier to find.
- * @return Pointer to Identifier match within string. nullptr if none found.
- */
-const TCHAR* FindIdentifierExactMatch(const TCHAR* StringBegin, const TCHAR* StringEnd, const FString& Identifier)
+FFindDelcarationResults FindDeclaration(const FUnrealStructDefinitionInfo& StructDef, const FString& Identifier)
 {
-	int32 StringLen = UE_PTRDIFF_TO_INT32(StringEnd - StringBegin);
+	FFindDelcarationResults Results;
 
-	// Check for exact match first.
-	if (FCString::Strncmp(StringBegin, *Identifier, StringLen) == 0)
+	if (Identifier.IsEmpty())
 	{
-		return StringBegin;
+		return Results;
 	}
 
-	int32        FindLen        = Identifier.Len();
-	const TCHAR* StringToSearch = StringBegin;
-
-	for (;;)
+	for (const FDeclaration& Declaration : StructDef.GetDeclarations())
 	{
-		const TCHAR* IdentifierStart = FCString::Strstr(StringToSearch, *Identifier);
-		if (IdentifierStart == nullptr)
-		{
-			// Not found.
-			return nullptr;
-		}
-
-		if (IdentifierStart > StringEnd || IdentifierStart + FindLen + 1 > StringEnd)
+		for (int32 Index = 0, EIndex = Declaration.Tokens.Num(); Index != EIndex; ++Index)
 		{
-			// Found match is out of string range.
-			return nullptr;
-		}
-
-		if (IdentifierStart == StringBegin && !FChar::IsIdentifier(*(IdentifierStart + FindLen + 1)))
-		{
-			// Found match is at the beginning of string.
-			return IdentifierStart;
-		}
-
-		if (IdentifierStart + FindLen == StringEnd && !FChar::IsIdentifier(*(IdentifierStart - 1)))
-		{
-			// Found match ends with end of string.
-			return IdentifierStart;
-		}
-
-		if (!FChar::IsIdentifier(*(IdentifierStart + FindLen)) && !FChar::IsIdentifier(*(IdentifierStart - 1)))
-		{
-			// Found match is in the middle of string
-			return IdentifierStart;
+			const FToken& Token = Declaration.Tokens[Index];
+			if (Token.IsIdentifier())
+			{
+				if (Token.IsValue(TEXT("virtual"), ESearchCase::CaseSensitive))
+				{
+					Results.bVirtualFound = true;
+				}
+				else if (Token.IsValue(*Identifier, ESearchCase::CaseSensitive))
+				{
+					Results.Declaration = &Declaration;
+					Results.FunctionNameTokenIndex = Index;
+					return Results;
+				}
+			}
 		}
-
-		// Didn't find exact match, nor got to end of search string. Keep on searching.
-		StringToSearch = IdentifierStart + FindLen;
 	}
-
-	// We should never get here.
-	checkNoEntry();
-	return nullptr;
-}
-
-/**
- * Finds exact match of Identifier in string. Returns nullptr if none is found.
- *
- * @param String String to search.
- * @param Identifier Identifier to find.
- * @return Index to Identifier match within String. INDEX_NONE if none found.
- */
-int32 FindIdentifierExactMatch(const FString& String, const FString& Identifier)
-{
-	const TCHAR* IdentifierPtr = FindIdentifierExactMatch(*String, *String + String.Len(), Identifier);
-	if (IdentifierPtr == nullptr)
-	{
-		return INDEX_NONE;
-	}
-
-	return UE_PTRDIFF_TO_INT32(IdentifierPtr - *String);
-}
-
-/**
-* Checks if exact match of Identifier is in String.
-*
-* @param StringBegin Start of string to search.
-* @param StringEnd End of string to search.
-* @param Identifier Identifier to find.
-* @return true if Identifier is within string, false otherwise.
-*/
-bool HasIdentifierExactMatch(const TCHAR* StringBegin, const TCHAR* StringEnd, const FString& Find)
-{
-	return FindIdentifierExactMatch(StringBegin, StringEnd, Find) != nullptr;
-}
-
-/**
-* Checks if exact match of Identifier is in String.
-*
-* @param String String to search.
-* @param Identifier Identifier to find.
-* @return true if Identifier is within String, false otherwise.
-*/
-bool HasIdentifierExactMatch(const FString &String, const FString& Identifier)
-{
-	return FindIdentifierExactMatch(String, Identifier) != INDEX_NONE;
+	return Results;
 }
 
 void ConvertToBuildIncludePath(const FManifestModule& Module, FString& LocalPath)
@@ -2794,8 +2737,6 @@ void FNativeClassHeaderGenerator::ExportClassFromSourceFileInner(
 
 	FString PPOMacroName;
 
-	ClassDef.ValidateDefinitionRange();
-
 	FString GeneratedSerializeFunctionCPP;
 	FString GeneratedSerializeFunctionHeaderMacroName;
 
@@ -4489,22 +4430,10 @@ FString FNativeClassHeaderGenerator::GetFunctionReturnString(FUnrealFunctionDefi
 	return Result;
 }
 
-bool FNativeClassHeaderGenerator::IsMissingVirtualSpecifier(const FString& SourceFile, int32 FunctionNamePosition)
-{
-	auto IsEndOfSearchChar = [](TCHAR C) { return (C == TEXT('}')) || (C == TEXT('{')) || (C == TEXT(';')); };
-
-	// Find first occurrence of "}", ";", "{" going backwards from ImplementationPosition.
-	int32 EndOfSearchCharIndex = SourceFile.FindLastCharByPredicate(IsEndOfSearchChar, FunctionNamePosition);
-	check(EndOfSearchCharIndex != INDEX_NONE);
-
-	// Then find if there is "virtual" keyword starting from position of found character to ImplementationPosition
-	return !HasIdentifierExactMatch(&SourceFile[EndOfSearchCharIndex], &SourceFile[FunctionNamePosition], TEXT("virtual"));
-}
-
-void FNativeClassHeaderGenerator::CheckRPCFunctions(FReferenceGatherers& OutReferenceGatherers, FUnrealFunctionDefinitionInfo& FunctionDef, const FString& ClassName, int32 ImplementationPosition, int32 ValidatePosition, const FUnrealSourceFile& SourceFile) const
+void FNativeClassHeaderGenerator::CheckRPCFunctions(FReferenceGatherers& OutReferenceGatherers, FUnrealFunctionDefinitionInfo& FunctionDef, const FString& ClassName, const FFindDelcarationResults& Implementation, const FFindDelcarationResults& Validation, const FUnrealSourceFile& SourceFile) const
 {
-	bool bHasImplementation = ImplementationPosition != INDEX_NONE;
-	bool bHasValidate = ValidatePosition != INDEX_NONE;
+	bool bHasImplementation = Implementation.WasFound();
+	bool bHasValidate = Validation.WasFound();
 
 	const FFuncInfo& FunctionData = FunctionDef.GetFunctionData();
 	FString FunctionReturnType = GetFunctionReturnString(FunctionDef, OutReferenceGatherers);
@@ -4566,13 +4495,13 @@ void FNativeClassHeaderGenerator::CheckRPCFunctions(FReferenceGatherers& OutRefe
 	//
 	// If all needed functions are declared, check if they have virtual specifiers.
 	//
-	if (bNeedsImplementation && bHasImplementation && IsMissingVirtualSpecifier(FileContent, ImplementationPosition))
+	if (bNeedsImplementation && bHasImplementation && !Implementation.IsVirtual())
 	{
 		FString FunctionDecl = FString::Printf(TEXT("%s %s::%s(%s) %s"), *FunctionReturnType, *ClassName, *FunctionData.CppImplName, *ParameterString, *ConstModifier);
 		FunctionDef.Throwf(TEXT("Declared function %sis not marked as virtual."), *FunctionDecl);
 	}
 
-	if (bNeedsValidate && bHasValidate && IsMissingVirtualSpecifier(FileContent, ValidatePosition))
+	if (bNeedsValidate && bHasValidate && !Validation.IsVirtual())
 	{
 		FString FunctionDecl = FString::Printf(TEXT("bool %s::%s(%s) %s"), *ClassName, *FunctionData.CppValidationImplName, *ParameterString, *ConstModifier);
 		FunctionDef.Throwf(TEXT("Declared function %sis not marked as virtual."), *FunctionDecl);
@@ -4883,15 +4812,10 @@ void FNativeClassHeaderGenerator::ExportFunctionThunk(FUHTStringBuilder& RPCWrap
 	FUnrealClassDefinitionInfo* OwnerClassDef = FunctionDef.GetOwnerClass();
 	check(OwnerClassDef);
 
-	OwnerClassDef->ValidateDefinitionRange();
-	const TCHAR* ClassStart = OwnerClassDef->GetDefinitionRange().Start;
-	const TCHAR* ClassEnd = OwnerClassDef->GetDefinitionRange().End;
-	FString      ClassName  = OwnerClassDef->GetName();
+	FString ClassName = OwnerClassDef->GetName();
 
-	FString ClassDefinition(UE_PTRDIFF_TO_INT32(ClassEnd - ClassStart), ClassStart);
-
-	bool bHasImplementation = HasIdentifierExactMatch(ClassDefinition, FunctionData.CppImplName);
-	bool bHasValidate = HasIdentifierExactMatch(ClassDefinition, FunctionData.CppValidationImplName);
+	bool bHasImplementation = FindDeclaration(*OwnerClassDef, FunctionData.CppImplName).WasFound();
+	bool bHasValidate = FindDeclaration(*OwnerClassDef, FunctionData.CppValidationImplName).WasFound();
 
 	bool bShouldEnableImplementationDeprecation =
 		// Enable deprecation warnings only if GENERATED_BODY is used inside class or interface (not GENERATED_UCLASS_BODY etc.)
@@ -4996,10 +4920,6 @@ void FNativeClassHeaderGenerator::ExportNativeFunctions(FOutputDevice& OutGenera
 
 	const FString ClassCPPName = ClassDef.GetAlternateNameCPP(ClassDef.HasAnyClassFlags(CLASS_Interface));
 
-	ClassDef.ValidateDefinitionRange();
-	const TCHAR* ClassStart = ClassDef.GetDefinitionRange().Start;
-	const TCHAR* ClassEnd = ClassDef.GetDefinitionRange().End;
-
 	// gather static class data
 	TArray<FString> SparseClassDataTypes;
 	ClassDef.GetSparseClassDataTypes(SparseClassDataTypes);
@@ -5093,24 +5013,10 @@ void FNativeClassHeaderGenerator::ExportNativeFunctions(FOutputDevice& OutGenera
 
 		if (!bWillBeProgrammerTyped)
 		{
-			FString ClassDefinition(UE_PTRDIFF_TO_INT32(ClassEnd - ClassStart), ClassStart);
-
 			FString FunctionName = FunctionDef->GetName();
-			int32 ClassDefinitionStartPosition = UE_PTRDIFF_TO_INT32(ClassStart - *SourceFile.GetContent());
-
-			int32 ImplementationPosition = FindIdentifierExactMatch(ClassDefinition, FunctionData.CppImplName);
-			bool bHasImplementation = ImplementationPosition != INDEX_NONE;
-			if (bHasImplementation)
-			{
-				ImplementationPosition += ClassDefinitionStartPosition;
-			}
 
-			int32 ValidatePosition = FindIdentifierExactMatch(ClassDefinition, FunctionData.CppValidationImplName);
-			bool bHasValidate = ValidatePosition != INDEX_NONE;
-			if (bHasValidate)
-			{
-				ValidatePosition += ClassDefinitionStartPosition;
-			}
+			FFindDelcarationResults Implementation = FindDeclaration(ClassDef, FunctionData.CppImplName);
+			FFindDelcarationResults Validation = FindDeclaration(ClassDef, FunctionData.CppValidationImplName);
 
 			//Emit warning here if necessary
 			FUHTStringBuilder FunctionDeclaration;
@@ -5126,14 +5032,14 @@ void FNativeClassHeaderGenerator::ExportNativeFunctions(FOutputDevice& OutGenera
 				FStringOutputDevice ValidDecl;
 				ValidDecl.Logf(TEXT("\t%s bool %s(%s);\r\n"), Virtual, *FunctionData.CppValidationImplName, *ParameterList);
 				FuncStringBuilders.AutogeneratedBlueprintFunctionDeclarations.Log(*ValidDecl);
-				if (!bHasValidate)
+				if (!Validation.WasFound())
 				{
 					FuncStringBuilders.AutogeneratedBlueprintFunctionDeclarationsOnlyNotDeclared.Logf(TEXT("%s"), *ValidDecl);
 				}
 			}
 
 			FuncStringBuilders.AutogeneratedBlueprintFunctionDeclarations.Log(*FunctionDeclaration);
-			if (!bHasImplementation && FunctionData.CppImplName != FunctionName)
+			if (!Implementation.WasFound() && FunctionData.CppImplName != FunctionName)
 			{
 				FuncStringBuilders.AutogeneratedBlueprintFunctionDeclarationsOnlyNotDeclared.Log(*FunctionDeclaration);
 			}
@@ -5141,7 +5047,7 @@ void FNativeClassHeaderGenerator::ExportNativeFunctions(FOutputDevice& OutGenera
 			// Versions that skip function autodeclaration throw an error when a function is missing.
 			if (ClassDef.HasGeneratedBody() && (ClassDef.GetGeneratedCodeVersion() > EGeneratedCodeVersion::V1))
 			{
-				CheckRPCFunctions(OutReferenceGatherers, *FunctionDef, ClassCPPName, ImplementationPosition, ValidatePosition, SourceFile);
+				CheckRPCFunctions(OutReferenceGatherers, *FunctionDef, ClassCPPName, Implementation, Validation, SourceFile);
 			}
 		}
 
diff --git a/Engine/Source/Programs/UnrealHeaderTool/Private/HeaderParser.cpp b/Engine/Source/Programs/UnrealHeaderTool/Private/HeaderParser.cpp
--- a/Engine/Source/Programs/UnrealHeaderTool/Private/HeaderParser.cpp
+++ b/Engine/Source/Programs/UnrealHeaderTool/Private/HeaderParser.cpp
@@ -4649,7 +4649,7 @@ bool FHeaderParser::CompileDeclaration(TArray<FUnrealFunctionDefinitionInfo*>& D
 
 	if (Token.IsIdentifier(TEXT("UDELEGATE"), ESearchCase::CaseSensitive))
 	{
-		FUnrealFunctionDefinitionInfo &DelegateDef = CompileDelegateDeclaration(Token.Value, EDelegateSpecifierAction::Parse);
+		FUnrealFunctionDefinitionInfo& DelegateDef = CompileDelegateDeclaration(Token.Value, EDelegateSpecifierAction::Parse);
 		DelegatesToFixup.Add(&DelegateDef);
 		return true;
 	}
@@ -4754,152 +4754,281 @@ bool FHeaderParser::CompileDeclaration(TArray<FUnrealFunctionDefinitionInfo*>& D
 		}
 	}
 
-	if (bEncounteredNewStyleClass_UnmatchedBrackets && IsInAClass())
+	// Skip anything that looks like a macro followed by no bracket that we don't know about
+	if (ProbablyAnUnknownObjectLikeMacro(*this, Token))
+	{
+		return true;
+	}
+
+	FRecordTokens RecordTokens(*this, bEncounteredNewStyleClass_UnmatchedBrackets && IsInAClass() ? &GetCurrentClassDef() : nullptr, &Token);
+	bool Result = SkipDeclaration(Token);
+	if (RecordTokens.Stop())
 	{
 		FUnrealStructDefinitionInfo& StructDef = GetCurrentClassDef();
-		FToken ConstructorToken = Token;
+		const FDeclaration& Declaration = StructDef.GetDeclarations().Last();
+		if (CheckForConstructor(StructDef, Declaration))
+		{
+		}
+		else if (TopNest->NestType == ENestType::Class)
+		{
+			if (CheckForSerialize(StructDef, Declaration))
+			{
+			}
+		}
+	}
+	return Result;
+}
+
+bool FHeaderParser::CheckForConstructor(FUnrealStructDefinitionInfo& StructDef, const FDeclaration& Declaration)
+{
+	FTokenReplay Tokens(Declaration.Tokens);
+	FUnrealClassDefinitionInfo& ClassDef = StructDef.AsClassChecked();
 
-		// Allow explicit constructors
-		bool bFoundExplicit = ConstructorToken.IsIdentifier(TEXT("explicit"), ESearchCase::CaseSensitive);
-		if (bFoundExplicit)
+	FToken Token;
+	if (!Tokens.GetToken(Token))
+	{
+		return false;
+	}
+
+	// Allow explicit constructors
+	bool bFoundExplicit = Token.IsIdentifier(TEXT("explicit"), ESearchCase::CaseSensitive);
+	if (bFoundExplicit)
+	{
+		Tokens.GetToken(Token);
+	}
+
+	bool bSkippedAPIToken = false;
+	if (Token.Value.EndsWith(TEXT("_API"), ESearchCase::CaseSensitive))
+	{
+		if (!bFoundExplicit)
 		{
-			GetToken(ConstructorToken);
+			// Explicit can come before or after an _API
+			Tokens.MatchIdentifier(TEXT("explicit"), ESearchCase::CaseSensitive);
 		}
 
-		bool bSkippedAPIToken = false;
-		if (ConstructorToken.Value.EndsWith(TEXT("_API"), ESearchCase::CaseSensitive))
+		Tokens.GetToken(Token);
+		bSkippedAPIToken = true;
+	}
+
+	if (!Token.IsIdentifier(*ClassDef.GetAlternateNameCPP(), ESearchCase::IgnoreCase))
+	{
+		return false;
+	}
+
+	Tokens.GetToken(Token);
+	if (!Token.IsSymbol(TEXT('(')))
+	{
+		return false;
+	}
+
+	bool bOICtor = false;
+	bool bVTCtor = false;
+
+	if (!ClassDef.IsDefaultConstructorDeclared() && Tokens.MatchSymbol(TEXT(')')))
+	{
+		ClassDef.MarkDefaultConstructorDeclared();
+	}
+	else if (!ClassDef.IsObjectInitializerConstructorDeclared()
+		|| !ClassDef.IsCustomVTableHelperConstructorDeclared())
+	{
+		bool bIsConst = false;
+		bool bIsRef = false;
+		int32 ParenthesesNestingLevel = 1;
+
+		while (ParenthesesNestingLevel && Tokens.GetToken(Token))
 		{
-			if (!bFoundExplicit)
+			// Template instantiation or additional parameter excludes ObjectInitializer constructor.
+			if (Token.IsSymbol(TEXT(',')) || Token.IsSymbol(TEXT('<')))
 			{
-				// Explicit can come before or after an _API
-				MatchIdentifier(TEXT("explicit"), ESearchCase::CaseSensitive);
+				bOICtor = false;
+				bVTCtor = false;
+				break;
 			}
 
-			GetToken(ConstructorToken);
-			bSkippedAPIToken = true;
+			if (Token.IsSymbol(TEXT('(')))
+			{
+				ParenthesesNestingLevel++;
+				continue;
+			}
+
+			if (Token.IsSymbol(TEXT(')')))
+			{
+				ParenthesesNestingLevel--;
+				continue;
+			}
+
+			if (Token.IsIdentifier(TEXT("const"), ESearchCase::CaseSensitive))
+			{
+				bIsConst = true;
+				continue;
+			}
+
+			if (Token.IsSymbol(TEXT('&')))
+			{
+				bIsRef = true;
+				continue;
+			}
+
+			if (Token.IsIdentifier(TEXT("FObjectInitializer"), ESearchCase::CaseSensitive)
+				|| Token.IsIdentifier(TEXT("FPostConstructInitializeProperties"), ESearchCase::CaseSensitive) // Deprecated, but left here, so it won't break legacy code.
+				)
+			{
+				bOICtor = true;
+			}
+
+			if (Token.IsIdentifier(TEXT("FVTableHelper"), ESearchCase::CaseSensitive))
+			{
+				bVTCtor = true;
+			}
 		}
 
-		if (ConstructorToken.IsIdentifier(*StructDef.GetAlternateNameCPP(), ESearchCase::IgnoreCase))
+		// Parse until finish.
+		while (ParenthesesNestingLevel && Tokens.GetToken(Token))
 		{
-			if (TryToMatchConstructorParameterList(ConstructorToken))
+			if (Token.IsSymbol(TEXT('(')))
 			{
-				return true;
+				ParenthesesNestingLevel++;
+				continue;
 			}
+
+			if (Token.IsSymbol(TEXT(')')))
+			{
+				ParenthesesNestingLevel--;
+				continue;
+			}
+		}
+
+		if (bOICtor && bIsRef && bIsConst)
+		{
+			ClassDef.MarkObjectInitializerConstructorDeclared();
 		}
-		else if (bSkippedAPIToken)
+		if (bVTCtor && bIsRef)
 		{
-			// We skipped over an _API token, but this wasn't a constructor so we need to unget so that subsequent code and still process it
-			UngetToken(ConstructorToken);
+			ClassDef.MarkCustomVTableHelperConstructorDeclared();
 		}
 	}
 
-	// Skip anything that looks like a macro followed by no bracket that we don't know about
-	if (ProbablyAnUnknownObjectLikeMacro(*this, Token))
+	if (!bVTCtor)
 	{
-		return true;
+		ClassDef.MarkConstructorDeclared();
 	}
 
-	// Determine if this statement is a serialize function declaration
-	if (bEncounteredNewStyleClass_UnmatchedBrackets && IsInAClass() && TopNest->NestType == ENestType::Class)
+	return false;
+}
+
+bool FHeaderParser::CheckForSerialize(FUnrealStructDefinitionInfo& StructDef, const FDeclaration& Declaration)
+{
+	FTokenReplay Tokens(Declaration.Tokens);
+	FUnrealClassDefinitionInfo& ClassDef = StructDef.AsClassChecked();
+
+	FToken Token;
+	if (!Tokens.GetToken(Token))
 	{
-		while (Token.IsIdentifier(TEXT("virtual"), ESearchCase::CaseSensitive) || Token.Value.EndsWith(TEXT("_API"), ESearchCase::CaseSensitive))
-		{
-			GetToken(Token);
-		}
+		return false;
+	}
 
-		if (Token.IsIdentifier(TEXT("void"), ESearchCase::CaseSensitive))
-		{
-			GetToken(Token);
-			if (Token.IsIdentifier(TEXT("Serialize"), ESearchCase::CaseSensitive))
-			{
-				GetToken(Token);
-				if (Token.IsSymbol(TEXT('(')))
-				{
-					GetToken(Token);
+	while (Token.IsIdentifier(TEXT("virtual"), ESearchCase::CaseSensitive) || Token.Value.EndsWith(TEXT("_API"), ESearchCase::CaseSensitive))
+	{
+		Tokens.GetToken(Token);
+	}
 
-					ESerializerArchiveType ArchiveType = ESerializerArchiveType::None;
-					if (Token.IsIdentifier(TEXT("FArchive"), ESearchCase::CaseSensitive))
-					{
-						GetToken(Token);
-						if (Token.IsSymbol(TEXT('&')))
-						{
-							GetToken(Token);
+	if (!Token.IsIdentifier(TEXT("void"), ESearchCase::CaseSensitive))
+	{
+		return false;
+	}
 
-							// Allow the declaration to not define a name for the archive parameter
-							if (!Token.IsSymbol(TEXT(')')))
-							{
-								GetToken(Token);
-							}
+	Tokens.GetToken(Token);
+	if (!Token.IsIdentifier(TEXT("Serialize"), ESearchCase::CaseSensitive))
+	{
+		return false;
+	}
 
-							if (Token.IsSymbol(TEXT(')')))
-							{
-								ArchiveType = ESerializerArchiveType::Archive;
-							}
-						}
-					}
-					else if (Token.IsIdentifier(TEXT("FStructuredArchive"), ESearchCase::CaseSensitive))
-					{
-						GetToken(Token);
-						if (Token.IsSymbol(TEXT("::"), ESearchCase::CaseSensitive))
-						{
-							GetToken(Token);
+	Tokens.GetToken(Token);
+	if (!Token.IsSymbol(TEXT('(')))
+	{
+		return false;
+	}
 
-							if (Token.IsIdentifier(TEXT("FRecord"), ESearchCase::CaseSensitive))
-							{
-								GetToken(Token);
+	Tokens.GetToken(Token);
 
-								// Allow the declaration to not define a name for the slot parameter
-								if (!Token.IsSymbol(TEXT(')')))
-								{
-									GetToken(Token);
-								}
+	ESerializerArchiveType ArchiveType = ESerializerArchiveType::None;
+	if (Token.IsIdentifier(TEXT("FArchive"), ESearchCase::CaseSensitive))
+	{
+		Tokens.GetToken(Token);
+		if (Token.IsSymbol(TEXT('&')))
+		{
+			Tokens.GetToken(Token);
 
-								if (Token.IsSymbol(TEXT(')')))
-								{
-									ArchiveType = ESerializerArchiveType::StructuredArchiveRecord;
-								}
-							}
-						}
-					}
-					else if (Token.IsIdentifier(TEXT("FStructuredArchiveRecord"), ESearchCase::CaseSensitive))
-					{
-						GetToken(Token);
+			// Allow the declaration to not define a name for the archive parameter
+			if (!Token.IsSymbol(TEXT(')')))
+			{
+				Tokens.GetToken(Token);
+			}
 
-						// Allow the declaration to not define a name for the slot parameter
-						if (!Token.IsSymbol(TEXT(')')))
-						{
-							GetToken(Token);
-						}
+			if (Token.IsSymbol(TEXT(')')))
+			{
+				ArchiveType = ESerializerArchiveType::Archive;
+			}
+		}
+	}
+	else if (Token.IsIdentifier(TEXT("FStructuredArchive"), ESearchCase::CaseSensitive))
+	{
+		Tokens.GetToken(Token);
+		if (Token.IsSymbol(TEXT("::"), ESearchCase::CaseSensitive))
+		{
+			Tokens.GetToken(Token);
 
-						if (Token.IsSymbol(TEXT(')')))
-						{
-							ArchiveType = ESerializerArchiveType::StructuredArchiveRecord;
-						}
-					}
+			if (Token.IsIdentifier(TEXT("FRecord"), ESearchCase::CaseSensitive))
+			{
+				Tokens.GetToken(Token);
 
-					if (ArchiveType != ESerializerArchiveType::None)
-					{
-						// Found what we want!
-						if (uint32 CurrentCompilerDirective = GetCurrentCompilerDirective(); CurrentCompilerDirective == 0 || CurrentCompilerDirective == ECompilerDirective::WithEditorOnlyData)
-						{
-							FString EnclosingDefine = CurrentCompilerDirective != 0 ? TEXT("WITH_EDITORONLY_DATA") : TEXT("");
+				// Allow the declaration to not define a name for the slot parameter
+				if (!Token.IsSymbol(TEXT(')')))
+				{
+					Tokens.GetToken(Token);
+				}
 
-							FUnrealClassDefinitionInfo& ClassDef = GetCurrentClassDef();
-							ClassDef.AddArchiveType(ArchiveType);
-							ClassDef.SetEnclosingDefine(MoveTemp(EnclosingDefine));
-						}
-						else
-						{
-							Throwf(TEXT("Serialize functions must not be inside preprocessor blocks, except for WITH_EDITORONLY_DATA"));
-						}
-					}
+				if (Token.IsSymbol(TEXT(')')))
+				{
+					ArchiveType = ESerializerArchiveType::StructuredArchiveRecord;
 				}
 			}
 		}
 	}
+	else if (Token.IsIdentifier(TEXT("FStructuredArchiveRecord"), ESearchCase::CaseSensitive))
+	{
+		Tokens.GetToken(Token);
 
-	// Ignore C++ declaration / function definition. 
-	return SkipDeclaration(Token);
+		// Allow the declaration to not define a name for the slot parameter
+		if (!Token.IsSymbol(TEXT(')')))
+		{
+			Tokens.GetToken(Token);
+		}
+
+		if (Token.IsSymbol(TEXT(')')))
+		{
+			ArchiveType = ESerializerArchiveType::StructuredArchiveRecord;
+		}
+	}
+
+	if (ArchiveType != ESerializerArchiveType::None)
+	{
+		// Found what we want!
+		if (Declaration.CurrentCompilerDirective == 0 || Declaration.CurrentCompilerDirective == ECompilerDirective::WithEditorOnlyData)
+		{
+			FString EnclosingDefine = Declaration.CurrentCompilerDirective != 0 ? TEXT("WITH_EDITORONLY_DATA") : TEXT("");
+
+			ClassDef.AddArchiveType(ArchiveType);
+			ClassDef.SetEnclosingDefine(MoveTemp(EnclosingDefine));
+		}
+		else
+		{
+			FUHTMessage(StructDef.GetUnrealSourceFile(), Token.InputLine).Throwf(TEXT("Serialize functions must not be inside preprocessor blocks, except for WITH_EDITORONLY_DATA"));
+		}
+		return true;
+	}
+	
+	return false;
 }
 
 bool FHeaderParser::SkipDeclaration(FToken& Token)
@@ -6111,6 +6240,9 @@ FUnrealFunctionDefinitionInfo& FHeaderParser::CompileFunctionDeclaration()
 		bAutomaticallyFinal = false;
 	}
 
+	// Record the tokens so we can detect this function as a declaration later (i.e. RPC)
+	FRecordTokens RecordTokens(*this, &GetCurrentClassDef(), nullptr);
+
 	bool bSawVirtual = false;
 
 	if (MatchIdentifier(TEXT("virtual"), ESearchCase::CaseSensitive))
@@ -8759,3 +8891,36 @@ FUnrealFunctionDefinitionInfo& FHeaderParser::CreateFunction(const TCHAR* FuncNa
 
 	return FuncDef;
 }
+
+FRecordTokens::FRecordTokens(FHeaderParser& InParser, FUnrealStructDefinitionInfo* InStructDef, FToken* InToken)
+	: Parser(InParser)
+	, StructDef(InStructDef)
+	, CurrentCompilerDirective(Parser.GetCurrentCompilerDirective())
+{
+	if (StructDef != nullptr)
+	{
+		Parser.bRecordTokens = true;
+		if (InToken != nullptr)
+		{
+			Parser.RecordedTokens.Push(*InToken);
+		}
+	}
+}
+
+FRecordTokens::~FRecordTokens()
+{
+	Stop();
+}
+
+bool FRecordTokens::Stop()
+{
+	if (StructDef != nullptr)
+	{
+		Parser.bRecordTokens = false;
+		StructDef->AddDeclaration(CurrentCompilerDirective, MoveTemp(Parser.RecordedTokens));
+		StructDef = nullptr;
+		return true;
+	}
+	return false;
+}
+
diff --git a/Engine/Source/Programs/UnrealHeaderTool/Private/HeaderParser.h b/Engine/Source/Programs/UnrealHeaderTool/Private/HeaderParser.h
--- a/Engine/Source/Programs/UnrealHeaderTool/Private/HeaderParser.h
+++ b/Engine/Source/Programs/UnrealHeaderTool/Private/HeaderParser.h
@@ -20,6 +20,7 @@ class FStringOutputDevice;
 class FUnrealSourceFile;
 class FScope;
 class FHeaderProvider;
+struct FDeclaration;
 
 /*-----------------------------------------------------------------------------
 	Constants & types.
@@ -130,6 +131,8 @@ struct FDocumentationPolicy
 //
 class FHeaderParser : public FBaseParser
 {
+	friend class FRecordTokens;
+
 public:
 	// Performs a preliminary parse of the text in the specified buffer, pulling out:
 	//   Class name and parent class name
@@ -603,6 +606,12 @@ private:
 	// Checks if a valid range has been found on the provided metadata
 	void ConditionalLogPointerUsage(EPointerMemberBehavior PointerMemberBehavior, const TCHAR* PointerTypeDesc, FString&& PointerTypeDecl);
 
+	// Check to see if the declaration is a constructor
+	static bool CheckForConstructor(FUnrealStructDefinitionInfo& StructDef, const FDeclaration& Declaration);
+
+	// Check to see if the declaration is a serialize
+	static bool CheckForSerialize(FUnrealStructDefinitionInfo& StructDef, const FDeclaration& Declaration);
+
 	// Names that cannot be used enums, UStructs, or UClasses
 	static TArray<FString> ReservedTypeNames;
 
@@ -658,3 +667,16 @@ public:
 		int32 InLineNumber
 	);
 };
+
+class FRecordTokens
+{
+public:
+	explicit FRecordTokens(FHeaderParser& InParser, FUnrealStructDefinitionInfo* InStructDef, FToken* InToken);
+	~FRecordTokens();
+	bool Stop();
+
+private:
+	FHeaderParser& Parser;
+	FUnrealStructDefinitionInfo* StructDef;
+	uint32 CurrentCompilerDirective = 0;
+};
\ No newline at end of file
diff --git a/Engine/Source/Programs/UnrealHeaderTool/Private/NativeClassExporter.h b/Engine/Source/Programs/UnrealHeaderTool/Private/NativeClassExporter.h
--- a/Engine/Source/Programs/UnrealHeaderTool/Private/NativeClassExporter.h
+++ b/Engine/Source/Programs/UnrealHeaderTool/Private/NativeClassExporter.h
@@ -11,6 +11,7 @@ class FUnrealPackageDefinitionInfo;
 class FUnrealPropertyDefinitionInfo;
 struct FFuncInfo;
 struct FGeneratedFileInfo;
+struct FFindDelcarationResults;
 
 //
 //	FNativeClassHeaderGenerator
@@ -507,7 +508,7 @@ private:
 	* @param	ValidatePosition		Position in source file of _Validate function for function described by FunctionData.
 	* @param	SourceFile				Currently analyzed source file.
 	*/
-	void CheckRPCFunctions(FReferenceGatherers& OutReferenceGatherers, FUnrealFunctionDefinitionInfo& FunctionDef, const FString& ClassName, int32 ImplementationPosition, int32 ValidatePosition, const FUnrealSourceFile& SourceFile) const;
+	void CheckRPCFunctions(FReferenceGatherers& OutReferenceGatherers, FUnrealFunctionDefinitionInfo& FunctionDef, const FString& ClassName, const FFindDelcarationResults& Implementation, const FFindDelcarationResults& Validation, const FUnrealSourceFile& SourceFile) const;
 
 	/**
 	 * Exports the native stubs for the list of functions specified
@@ -669,15 +670,6 @@ private:
 	*/
 	static FString GetFunctionParameterString(FUnrealFunctionDefinitionInfo& FunctionDef, FReferenceGatherers& OutReferenceGatherers);
 
-	/**
-	 * Checks if function is missing "virtual" specifier.
-	 *
-	 * @param SourceFile SourceFile where function is declared.
-	 * @param FunctionNamePosition Position of name of function in SourceFile.
-	 * @return true if function misses "virtual" specifier, false otherwise.
-	 */
-	static bool IsMissingVirtualSpecifier(const FString& SourceFile, int32 FunctionNamePosition);
-
 public:
 
 	/** 
diff --git a/Engine/Source/Programs/UnrealHeaderTool/Private/UnrealTypeDefinitionInfo.h b/Engine/Source/Programs/UnrealHeaderTool/Private/UnrealTypeDefinitionInfo.h
--- a/Engine/Source/Programs/UnrealHeaderTool/Private/UnrealTypeDefinitionInfo.h
+++ b/Engine/Source/Programs/UnrealHeaderTool/Private/UnrealTypeDefinitionInfo.h
@@ -92,6 +92,12 @@ enum class EPostParseFinalizePhase : uint8
 	MAX_Phases,
 };
 
+struct FDeclaration
+{
+	uint32 CurrentCompilerDirective;
+	TArray<FToken> Tokens;
+};
+
 /**
  * UHT Specific implementation of meta data support
  */
@@ -2285,6 +2291,22 @@ public:
 		return RigVMInfo;
 	}
 
+	/**
+	 * Add a declaration to the struct
+	 */
+	void AddDeclaration(uint32 CurrentCompilerDirective, TArray<FToken>&& Tokens)
+	{
+		Declarations.Emplace(FDeclaration{ CurrentCompilerDirective, MoveTemp(Tokens) });
+	}
+
+	/** 
+	 * Get all of the declarations
+	 */
+	const TArray<FDeclaration>& GetDeclarations() const
+	{
+		return Declarations;
+	}
+
 protected:
 	explicit FUnrealStructDefinitionInfo(UObject* Object)
 		: FUnrealFieldDefinitionInfo(Object)
@@ -2311,6 +2333,9 @@ private:
 	/** Functions of the structure */
 	TArray<TSharedRef<FUnrealFunctionDefinitionInfo>> Functions;
 
+	/** Other declarations */
+	TArray<FDeclaration> Declarations;
+
 	FBaseStructInfo SuperStructInfo;
 	TArray<FBaseStructInfo> BaseStructInfos;
 
