1) Removed the need to create a UScriptStruct to test to see if it has a NoOp constructor.
2) Added a quick method to query the deferred CppStructOps until we decide on a long term solution
3) Added method to test to see if a source file is the NoExports include file.

#rnx

diff --git a/Engine/Source/Programs/UnrealHeaderTool/Private/HeaderParser.cpp b/Engine/Source/Programs/UnrealHeaderTool/Private/HeaderParser.cpp
--- a/Engine/Source/Programs/UnrealHeaderTool/Private/HeaderParser.cpp
+++ b/Engine/Source/Programs/UnrealHeaderTool/Private/HeaderParser.cpp
@@ -3438,7 +3438,7 @@ void FHeaderParser::GetVarType(
 	}
 	else if (VarType.IsValue(TEXT("FLargeWorldCoordinatesReal"), ESearchCase::CaseSensitive))		// LWC_TODO: Remove with UE_LARGE_WORLD_COORDINATES_DISABLED toggle.
 	{
-		if (!SourceFile.GetFilename().EndsWith(TEXT("NoExportTypes.h")))
+		if (!SourceFile.IsNoExportTypes())
 		{
 			Throwf(TEXT("FLargeWorldCoordinatesReal is intended for LWC support only and should not be used outside of NoExportTypes.h"));
 		}
diff --git a/Engine/Source/Programs/UnrealHeaderTool/Private/UnrealSourceFile.h b/Engine/Source/Programs/UnrealHeaderTool/Private/UnrealSourceFile.h
--- a/Engine/Source/Programs/UnrealHeaderTool/Private/UnrealSourceFile.h
+++ b/Engine/Source/Programs/UnrealHeaderTool/Private/UnrealSourceFile.h
@@ -40,7 +40,8 @@ public:
 		, Filename             (InFilename)
 		, PackageDef           (InPackageDef)
 	{
-		if (GetStrippedFilename() != "NoExportTypes")
+		bIsNoExportTypes = GetStrippedFilename() == "NoExportTypes";
+		if (!bIsNoExportTypes)
 		{
 			Includes.Emplace(FHeaderProvider(EHeaderProviderSourceType::FileName, "NoExportTypes.h"));
 		}
@@ -467,6 +468,14 @@ public:
 		return bIsReferenced || GetScope()->ContainsTypes();
 	}
 
+	/**
+	 * Return true if this is the NoExportTypes.h file 
+	 */
+	bool IsNoExportTypes() const
+	{
+		return bIsNoExportTypes;
+	}
+
 private:
 
 	// File scope.
@@ -520,6 +529,9 @@ private:
 	// Tells if this file is referenced by another
 	bool bIsReferenced = false;
 
+	// True if this is the NoExportTypes.h file
+	bool bIsNoExportTypes = false;
+
 	// Current topological sort state
 	ETopologicalState TopologicalState = ETopologicalState::Unmarked;
 
diff --git a/Engine/Source/Programs/UnrealHeaderTool/Private/UnrealTypeDefinitionInfo.cpp b/Engine/Source/Programs/UnrealHeaderTool/Private/UnrealTypeDefinitionInfo.cpp
--- a/Engine/Source/Programs/UnrealHeaderTool/Private/UnrealTypeDefinitionInfo.cpp
+++ b/Engine/Source/Programs/UnrealHeaderTool/Private/UnrealTypeDefinitionInfo.cpp
@@ -791,6 +791,11 @@ struct FUHTDisplayNameHelper
 	}
 };
 
+bool FUnrealPropertyDefinitionInfo::HasNoOpConstructor() const
+{
+	return PropertyBase.IsStructOrStructStaticArray() && PropertyBase.ScriptStructDef->HasNoOpConstructor();
+}
+
 FUnrealObjectDefinitionInfo* FUnrealPropertyDefinitionInfo::GetOwnerObject() const
 {
 	for (FUnrealTypeDefinitionInfo* TypeDef = GetOuter(); TypeDef; TypeDef = TypeDef->GetOuter())
@@ -1846,6 +1851,12 @@ void FUnrealScriptStructDefinitionInfo::SetHasInstancedReference()
 	PostParseFinalizeReferencedProperties();
 }
 
+bool FUnrealScriptStructDefinitionInfo::HasNoOpConstructor() const
+{
+	UScriptStruct::ICppStructOps* CppStructOps = UScriptStruct::FindDeferredCppStructOps(GetFName());
+	return CppStructOps ? CppStructOps->HasNoopConstructor() : false;
+}
+
 FUnrealClassDefinitionInfo::FUnrealClassDefinitionInfo(FUnrealSourceFile& InSourceFile, int32 InLineNumber, FString&& InNameCPP, FName InName, bool bInIsInterface)
 	: FUnrealStructDefinitionInfo(InSourceFile, InLineNumber, MoveTemp(InNameCPP), InName, InSourceFile.GetPackageDef())
 	, bIsInterface(bInIsInterface)
diff --git a/Engine/Source/Programs/UnrealHeaderTool/Private/UnrealTypeDefinitionInfo.h b/Engine/Source/Programs/UnrealHeaderTool/Private/UnrealTypeDefinitionInfo.h
--- a/Engine/Source/Programs/UnrealHeaderTool/Private/UnrealTypeDefinitionInfo.h
+++ b/Engine/Source/Programs/UnrealHeaderTool/Private/UnrealTypeDefinitionInfo.h
@@ -1031,14 +1031,7 @@ public:
 	/**
 	 * Return true if this is a structure and it has no op constructor
 	 */
-	bool HasNoOpConstructor() const
-	{
-		if (FStructProperty* StructProperty = CastField<FStructProperty>(GetProperty()))
-		{
-			return StructProperty->HasNoOpConstructor();
-		}
-		return false;
-	}
+	bool HasNoOpConstructor() const;
 
 	/**
 	 * Get the string that represents the array dimensions.  A nullptr is returned if the property doesn't have any dimensions
@@ -2506,6 +2499,8 @@ public:
 		MacroDeclaredLineNumber = InMacroDeclaredLineNumber;
 	}
 
+	bool HasNoOpConstructor() const;
+
 	/**
 	 * Handles the propagation of property flags
 	 */
diff --git a/Engine/Source/Runtime/CoreUObject/Private/UObject/Class.cpp b/Engine/Source/Runtime/CoreUObject/Private/UObject/Class.cpp
--- a/Engine/Source/Runtime/CoreUObject/Private/UObject/Class.cpp
+++ b/Engine/Source/Runtime/CoreUObject/Private/UObject/Class.cpp
@@ -2530,6 +2530,13 @@ void UScriptStruct::DeferCppStructOps(FName Target, ICppStructOps* InCppStructOp
 	DeferredStructOps.Add(Target,InCppStructOps);
 }
 
+#if HACK_HEADER_GENERATOR
+UScriptStruct::ICppStructOps* UScriptStruct::FindDeferredCppStructOps(FName StructName)
+{
+	return GetDeferredCppStructOps().FindRef(StructName);
+}
+#endif
+
 /** Look for the CppStructOps if we don't already have it and set the property size **/
 void UScriptStruct::PrepareCppStructOps()
 {
diff --git a/Engine/Source/Runtime/CoreUObject/Public/UObject/Class.h b/Engine/Source/Runtime/CoreUObject/Public/UObject/Class.h
--- a/Engine/Source/Runtime/CoreUObject/Public/UObject/Class.h
+++ b/Engine/Source/Runtime/CoreUObject/Public/UObject/Class.h
@@ -1606,6 +1606,8 @@ public:
 
 	// Required by UHT makefiles for internal data serialization.
 	friend struct FScriptStructArchiveProxy;
+
+	static COREUOBJECT_API ICppStructOps* FindDeferredCppStructOps(FName StructName);
 #endif
 	bool CanSerializeAsAlias(const struct FPropertyTag& Tag) const;
 
