Replaced two UProperty formatting routines with UHT specific code.

#rnx

diff --git a/Engine/Source/Programs/UnrealHeaderTool/Private/ParserHelper.h b/Engine/Source/Programs/UnrealHeaderTool/Private/ParserHelper.h
--- a/Engine/Source/Programs/UnrealHeaderTool/Private/ParserHelper.h
+++ b/Engine/Source/Programs/UnrealHeaderTool/Private/ParserHelper.h
@@ -293,6 +293,32 @@ public:
 		return Type == CPT_Delegate && IsPrimitiveOrPrimitiveStaticArray();
 	}
 
+	bool PassCPPArgsByRef() const
+	{
+		if (ArrayType == EArrayType::Dynamic || ArrayType == EArrayType::Set || MapKeyProp.IsValid())
+		{
+			return true;
+		}
+		switch (Type)
+		{
+		//case CPT_Struct:
+		case CPT_String:
+		case CPT_Text:
+		case CPT_Delegate:
+		case CPT_MulticastDelegate:
+		case CPT_WeakObjectReference:
+		case CPT_LazyObjectReference:
+		case CPT_ObjectPtrReference:
+		case CPT_SoftObjectReference:
+		case CPT_Map:
+		case CPT_Set:
+			return true;
+
+		default:
+			return false;
+		}
+	}
+
 	FUnrealEnumDefinitionInfo* AsEnum() const;
 
 	bool IsEnum() const;
diff --git a/Engine/Source/Programs/UnrealHeaderTool/Private/PropertyTypes.cpp b/Engine/Source/Programs/UnrealHeaderTool/Private/PropertyTypes.cpp
--- a/Engine/Source/Programs/UnrealHeaderTool/Private/PropertyTypes.cpp
+++ b/Engine/Source/Programs/UnrealHeaderTool/Private/PropertyTypes.cpp
@@ -326,6 +326,15 @@ namespace
 		}
 	};
 
+	template<typename TraitsType>
+	struct GetCPPMacroTypeDispatch
+	{
+		FString operator()(const FUnrealPropertyDefinitionInfo& PropDef, FString& ExtendedTypeText)
+		{
+			return TraitsType::GetCPPMacroType(PropDef, ExtendedTypeText);
+		}
+	};
+
 	template<typename TraitsType>
 	struct PostParseFinalizeDispatch
 	{
@@ -432,6 +441,8 @@ struct FPropertyTypeTraitsBase
 
 	//static FString GetCPPTypeForwardDeclaration(const FUnrealPropertyDefinitionInfo& PropDef);
 
+	//static FString GetCPPMacroType(const FUnrealPropertyDefinitionInfo& PropDef, FString& ExtendedTypeText);
+
 	/**
 	 * Handle any final initialization after parsing has completed
 	 *
@@ -457,6 +468,13 @@ struct FPropertyTypeTraitsNumericBase : public FPropertyTypeTraitsBase
 	{
 		return FString();
 	}
+
+	static FString GetCPPMacroType(const FUnrealPropertyDefinitionInfo& PropDef, FString& ExtendedTypeText)
+	{
+		ExtendedTypeText = TEXT("F");
+		ExtendedTypeText += PropDef.GetEngineClassName();
+		return TEXT("PROPERTY");
+	}
 };
 
 struct FPropertyTypeTraitsByte : public FPropertyTypeTraitsNumericBase
@@ -949,6 +967,11 @@ struct FPropertyTypeTraitsBool : public FPropertyTypeTraitsBooleanBase
 	{
 		return GetCPPTypeHelper(PropDef, ExtendedTypeText, CPPExportFlags, true, TEXT("bool"));
 	}
+
+	static FString GetCPPMacroType(const FUnrealPropertyDefinitionInfo& PropDef, FString& ExtendedTypeText)
+	{
+		return TEXT("UBOOL");
+	}
 };
 
 struct FPropertyTypeTraitsBool8 : public FPropertyTypeTraitsBooleanBase
@@ -962,6 +985,11 @@ struct FPropertyTypeTraitsBool8 : public FPropertyTypeTraitsBooleanBase
 	{
 		return GetCPPTypeHelper(PropDef, ExtendedTypeText, CPPExportFlags, false, TEXT("uint8"));
 	}
+
+	static FString GetCPPMacroType(const FUnrealPropertyDefinitionInfo& PropDef, FString& ExtendedTypeText)
+	{
+		return TEXT("UBOOL8");
+	}
 };
 
 struct FPropertyTypeTraitsBool16 : public FPropertyTypeTraitsBooleanBase
@@ -975,6 +1003,11 @@ struct FPropertyTypeTraitsBool16 : public FPropertyTypeTraitsBooleanBase
 	{
 		return GetCPPTypeHelper(PropDef, ExtendedTypeText, CPPExportFlags, false, TEXT("uint16"));
 	}
+
+	static FString GetCPPMacroType(const FUnrealPropertyDefinitionInfo& PropDef, FString& ExtendedTypeText)
+	{
+		return TEXT("UBOOL16");
+	}
 };
 
 struct FPropertyTypeTraitsBool32 : public FPropertyTypeTraitsBooleanBase
@@ -988,6 +1021,11 @@ struct FPropertyTypeTraitsBool32 : public FPropertyTypeTraitsBooleanBase
 	{
 		return GetCPPTypeHelper(PropDef, ExtendedTypeText, CPPExportFlags, false, TEXT("uint32"));
 	}
+
+	static FString GetCPPMacroType(const FUnrealPropertyDefinitionInfo& PropDef, FString& ExtendedTypeText)
+	{
+		return TEXT("UBOOL32");
+	}
 };
 
 struct FPropertyTypeTraitsBool64 : public FPropertyTypeTraitsBooleanBase
@@ -1001,6 +1039,11 @@ struct FPropertyTypeTraitsBool64 : public FPropertyTypeTraitsBooleanBase
 	{
 		return GetCPPTypeHelper(PropDef, ExtendedTypeText, CPPExportFlags, false, TEXT("uint64"));
 	}
+
+	static FString GetCPPMacroType(const FUnrealPropertyDefinitionInfo& PropDef, FString& ExtendedTypeText)
+	{
+		return TEXT("UBOOL64");
+	}
 };
 
 //----------------------------------------------------------------------------------------------------------------------------------------------------------
@@ -1149,6 +1192,18 @@ struct FPropertyTypeTraitsEnum : public FPropertyTypeTraitsBase
 
 		return FString::Printf(TEXT("enum class %s : %s;"), *VarProperty.EnumDef->GetName(), *PropDef.GetValuePropDef().GetCPPType());
 	}
+
+	static FString GetCPPMacroType(const FUnrealPropertyDefinitionInfo& PropDef, FString& ExtendedTypeText)
+	{
+		const FPropertyBase& VarProperty = PropDef.GetPropertyBase();
+		FUnrealEnumDefinitionInfo* EnumDef = VarProperty.EnumDef;
+		if (EnumDef->GetCppForm() != UEnum::ECppForm::EnumClass)
+		{
+			return FPropertyTypeTraitsByte::GetCPPMacroType(PropDef, ExtendedTypeText);
+		}
+		ExtendedTypeText = EnumDef->GetName();
+		return TEXT("ENUM");
+	}
 };
 
 //----------------------------------------------------------------------------------------------------------------------------------------------------------
@@ -1292,6 +1347,21 @@ struct FPropertyTypeTraitsObjectReference : public FPropertyTypeTraitsObjectBase
 			return FString::Printf(TEXT("class %s%s;"), VarProperty.ClassDef->GetPrefixCPP(), *VarProperty.ClassDef->GetName());
 		}
 	}
+
+	static FString GetCPPMacroType(const FUnrealPropertyDefinitionInfo& PropDef, FString& ExtendedTypeText)
+	{
+		const FPropertyBase& VarProperty = PropDef.GetPropertyBase();
+		check(VarProperty.ClassDef);
+		if (VarProperty.ClassDef->IsChildOf(*GUClassDef))
+		{
+			ExtendedTypeText = TEXT("UClass");
+		}
+		else
+		{
+			ExtendedTypeText =  FString::Printf(TEXT("%s%s"), VarProperty.ClassDef->GetPrefixCPP(), *VarProperty.ClassDef->GetName());
+		}
+		return TEXT("OBJECT");
+	}
 };
 
 struct FPropertyTypeTraitsWeakObjectReference : public FPropertyTypeTraitsObjectBase
@@ -1333,6 +1403,19 @@ struct FPropertyTypeTraitsWeakObjectReference : public FPropertyTypeTraitsObject
 		check(VarProperty.ClassDef);
 		return FString::Printf(TEXT("class %s%s;"), VarProperty.ClassDef->GetPrefixCPP(), *VarProperty.ClassDef->GetName());
 	}
+
+	static FString GetCPPMacroType(const FUnrealPropertyDefinitionInfo& PropDef, FString& ExtendedTypeText)
+	{
+		const FPropertyBase& VarProperty = PropDef.GetPropertyBase();
+		check(VarProperty.ClassDef);
+		if (PropDef.HasAnyPropertyFlags(CPF_AutoWeak))
+		{
+			ExtendedTypeText = FString::Printf(TEXT("TAutoWeakObjectPtr<%s%s>"), VarProperty.ClassDef->GetPrefixCPP(), *VarProperty.ClassDef->GetName());
+			return TEXT("AUTOWEAKOBJECT");
+		}
+		ExtendedTypeText = FString::Printf(TEXT("TWeakObjectPtr<%s%s>"), VarProperty.ClassDef->GetPrefixCPP(), *VarProperty.ClassDef->GetName());
+		return TEXT("WEAKOBJECT");
+	}
 };
 
 struct FPropertyTypeTraitsLazyObjectReference : public FPropertyTypeTraitsObjectBase
@@ -1364,6 +1447,14 @@ struct FPropertyTypeTraitsLazyObjectReference : public FPropertyTypeTraitsObject
 		check(VarProperty.ClassDef);
 		return FString::Printf(TEXT("class %s%s;"), VarProperty.ClassDef->GetPrefixCPP(), *VarProperty.ClassDef->GetName());
 	}
+
+	static FString GetCPPMacroType(const FUnrealPropertyDefinitionInfo& PropDef, FString& ExtendedTypeText)
+	{
+		const FPropertyBase& VarProperty = PropDef.GetPropertyBase();
+		check(VarProperty.ClassDef);
+		ExtendedTypeText = FString::Printf(TEXT("TLazyObjectPtr<%s%s>"), VarProperty.ClassDef->GetPrefixCPP(), *VarProperty.ClassDef->GetName());
+		return TEXT("LAZYOBJECT");
+	}
 };
 
 struct FPropertyTypeTraitsObjectPtrReference : public FPropertyTypeTraitsObjectBase
@@ -1446,13 +1537,28 @@ struct FPropertyTypeTraitsObjectPtrReference : public FPropertyTypeTraitsObjectB
 		check(VarProperty.ClassDef);
 		if (VarProperty.ClassDef->IsChildOf(*GUClassDef))
 		{
-			return FString::Printf(TEXT("class %s%s;"), VarProperty.MetaClassDef->GetPrefixCPP(), *VarProperty.MetaClassDef->GetName());
+			return FString::Printf(TEXT("class %s%s;"), VarProperty.ClassDef->GetPrefixCPP(), *VarProperty.ClassDef->GetName());
 		}
 		else
 		{
 			return FString::Printf(TEXT("class %s%s;"), VarProperty.ClassDef->GetPrefixCPP(), *VarProperty.ClassDef->GetName());
 		}
 	}
+
+	static FString GetCPPMacroType(const FUnrealPropertyDefinitionInfo& PropDef, FString& ExtendedTypeText)
+	{
+		const FPropertyBase& VarProperty = PropDef.GetPropertyBase();
+		check(VarProperty.ClassDef);
+		if (VarProperty.ClassDef->IsChildOf(*GUClassDef))
+		{
+			ExtendedTypeText = FString::Printf(TEXT("TObjectPtr<%s%s>"), VarProperty.ClassDef->GetPrefixCPP(), *VarProperty.ClassDef->GetName());
+		}
+		else
+		{
+			ExtendedTypeText = FString::Printf(TEXT("TObjectPtr<%s%s>"), VarProperty.ClassDef->GetPrefixCPP(), *VarProperty.ClassDef->GetName());
+		}
+		return TEXT("OBJECTPTR");
+	}
 };
 
 struct FPropertyTypeTraitsSoftObjectReference : public FPropertyTypeTraitsObjectBase
@@ -1501,7 +1607,6 @@ struct FPropertyTypeTraitsSoftObjectReference : public FPropertyTypeTraitsObject
 		check(VarProperty.ClassDef);
 		if (VarProperty.ClassDef->IsChildOf(*GUClassDef))
 		{
-			UClass* MetaClass = VarProperty.MetaClassDef->GetClass();
 			return FString::Printf(TEXT("TSoftClassPtr<%s%s> "), VarProperty.MetaClassDef->GetPrefixCPP(), *VarProperty.MetaClassDef->GetName()); //@TODO - This has an extra space in it
 		}
 		else
@@ -1524,6 +1629,22 @@ struct FPropertyTypeTraitsSoftObjectReference : public FPropertyTypeTraitsObject
 			return FString::Printf(TEXT("class %s%s;"), VarProperty.ClassDef->GetPrefixCPP(), *VarProperty.ClassDef->GetName());
 		}
 	}
+
+	static FString GetCPPMacroType(const FUnrealPropertyDefinitionInfo& PropDef, FString& ExtendedTypeText)
+	{
+		const FPropertyBase& VarProperty = PropDef.GetPropertyBase();
+		check(VarProperty.ClassDef);
+		if (VarProperty.ClassDef->IsChildOf(*GUClassDef))
+		{
+			ExtendedTypeText = FString::Printf(TEXT("TSoftClassPtr<%s%s> "), VarProperty.MetaClassDef->GetPrefixCPP(), *VarProperty.MetaClassDef->GetName()); //@TODO - This has an extra space in it
+			return TEXT("SOFTCLASS");
+		}
+		else
+		{
+			ExtendedTypeText = FString::Printf(TEXT("TSoftObjectPtr<%s%s>"), VarProperty.ClassDef->GetPrefixCPP(), *VarProperty.ClassDef->GetName());
+			return TEXT("SOFTOBJECT");
+		}
+	}
 };
 
 struct FPropertyTypeTraitsInterface : public FPropertyTypeTraitsBase
@@ -1590,6 +1711,20 @@ struct FPropertyTypeTraitsInterface : public FPropertyTypeTraitsBase
 
 		return FString::Printf(TEXT("class I%s;"), *ExportClass->GetName());
 	}
+
+	static FString GetCPPMacroType(const FUnrealPropertyDefinitionInfo& PropDef, FString& ExtendedTypeText)
+	{
+		FUnrealClassDefinitionInfo* ExportClass = PropDef.GetPropertyBase().ClassDef;
+		while (ExportClass && !ExportClass->HasAnyClassFlags(CLASS_Native))
+		{
+			ExportClass = ExportClass->GetSuperClass();
+		}
+		check(ExportClass);
+		check(ExportClass->HasAnyClassFlags(CLASS_Interface));
+
+		ExtendedTypeText = FString::Printf(TEXT("I%s"), *ExportClass->GetName());
+		return TEXT("TINTERFACE");
+	}
 };
 
 //----------------------------------------------------------------------------------------------------------------------------------------------------------
@@ -1639,6 +1774,13 @@ struct FPropertyTypeTraitsName : public FPropertyTypeTraitsBase
 	{
 		return FString();
 	}
+
+	static FString GetCPPMacroType(const FUnrealPropertyDefinitionInfo& PropDef, FString& ExtendedTypeText)
+	{
+		ExtendedTypeText = TEXT("F");
+		ExtendedTypeText += PropDef.GetEngineClassName();
+		return TEXT("PROPERTY");
+	}
 };
 
 struct FPropertyTypeTraitsString : public FPropertyTypeTraitsBase
@@ -1677,6 +1819,13 @@ struct FPropertyTypeTraitsString : public FPropertyTypeTraitsBase
 	{
 		return FString();
 	}
+
+	static FString GetCPPMacroType(const FUnrealPropertyDefinitionInfo& PropDef, FString& ExtendedTypeText)
+	{
+		ExtendedTypeText = TEXT("F");
+		ExtendedTypeText += PropDef.GetEngineClassName();
+		return TEXT("PROPERTY");
+	}
 };
 
 struct FPropertyTypeTraitsText : public FPropertyTypeTraitsBase
@@ -1752,6 +1901,13 @@ struct FPropertyTypeTraitsText : public FPropertyTypeTraitsBase
 	{
 		return FString();
 	}
+
+	static FString GetCPPMacroType(const FUnrealPropertyDefinitionInfo& PropDef, FString& ExtendedTypeText)
+	{
+		ExtendedTypeText = TEXT("F");
+		ExtendedTypeText += PropDef.GetEngineClassName();
+		return TEXT("PROPERTY");
+	}
 };
 
 struct FPropertyTypeTraitsStruct : public FPropertyTypeTraitsBase
@@ -2020,6 +2176,12 @@ struct FPropertyTypeTraitsStruct : public FPropertyTypeTraitsBase
 			
 		return FString::Printf(TEXT("struct F%s;"), *PropDef.GetPropertyBase().ScriptStructDef->GetName());
 	}
+
+	static FString GetCPPMacroType(const FUnrealPropertyDefinitionInfo& PropDef, FString& ExtendedTypeText)
+	{
+		ExtendedTypeText = GetCPPType(PropDef, nullptr, CPPF_None);
+		return TEXT("STRUCT");
+	}
 };
 
 struct FPropertyTypeTraitsDelegate : public FPropertyTypeTraitsBase
@@ -2102,6 +2264,13 @@ struct FPropertyTypeTraitsDelegate : public FPropertyTypeTraitsBase
 	{
 		return FString();
 	}
+
+	static FString GetCPPMacroType(const FUnrealPropertyDefinitionInfo& PropDef, FString& ExtendedTypeText)
+	{
+		ExtendedTypeText = TEXT("F");
+		ExtendedTypeText += PropDef.GetEngineClassName();
+		return TEXT("PROPERTY");
+	}
 };
 
 struct FPropertyTypeTraitsMulticastDelegate : public FPropertyTypeTraitsBase
@@ -2207,6 +2376,13 @@ struct FPropertyTypeTraitsMulticastDelegate : public FPropertyTypeTraitsBase
 	{
 		return FString();
 	}
+
+	static FString GetCPPMacroType(const FUnrealPropertyDefinitionInfo& PropDef, FString& ExtendedTypeText)
+	{
+		ExtendedTypeText = TEXT("F");
+		ExtendedTypeText += PropDef.GetEngineClassName();
+		return TEXT("PROPERTY");
+	}
 };
 
 struct FPropertyTypeTraitsFieldPath : public FPropertyTypeTraitsBase
@@ -2252,6 +2428,13 @@ struct FPropertyTypeTraitsFieldPath : public FPropertyTypeTraitsBase
 	{
 		return FString::Printf(TEXT("class F%s;"), *PropDef.GetPropertyBase().FieldClassName.ToString());
 	}
+
+	static FString GetCPPMacroType(const FUnrealPropertyDefinitionInfo& PropDef, FString& ExtendedTypeText)
+	{
+		const FPropertyBase& VarProperty = PropDef.GetPropertyBase();
+		ExtendedTypeText = FString::Printf(TEXT("TFieldPath<F%s>"), *VarProperty.FieldClassName.ToString());
+		return TEXT("STRUCT");
+	}
 };
 
 //----------------------------------------------------------------------------------------------------------------------------------------------------------
@@ -2362,6 +2545,11 @@ struct FPropertyTypeTraitsStaticArray : public FPropertyTypeTraitsBase
 	{
 		return PropertyTypeDispatch<GetCPPTypeForwardDeclarationDispatch, false, FString>(PropDef);
 	}
+
+	static FString GetCPPMacroType(const FUnrealPropertyDefinitionInfo& PropDef, FString& ExtendedTypeText)
+	{
+		return PropertyTypeDispatch<GetCPPMacroTypeDispatch, false, FString>(PropDef, ExtendedTypeText);
+	}
 };
 
 struct FPropertyTypeTraitsDynamicArray : public FPropertyTypeTraitsBase
@@ -2443,6 +2631,12 @@ struct FPropertyTypeTraitsDynamicArray : public FPropertyTypeTraitsBase
 	{
 		return PropDef.GetValuePropDef().GetCPPTypeForwardDeclaration();
 	}
+
+	static FString GetCPPMacroType(const FUnrealPropertyDefinitionInfo& PropDef, FString& ExtendedTypeText)
+	{
+		ExtendedTypeText = PropDef.GetValuePropDef().GetCPPType();
+		return TEXT("TARRAY");
+	}
 };
 
 struct FPropertyTypeTraitsSet : public FPropertyTypeTraitsBase
@@ -2521,6 +2715,12 @@ struct FPropertyTypeTraitsSet : public FPropertyTypeTraitsBase
 	{
 		return PropDef.GetValuePropDef().GetCPPTypeForwardDeclaration();
 	}
+
+	static FString GetCPPMacroType(const FUnrealPropertyDefinitionInfo& PropDef, FString& ExtendedTypeText)
+	{
+		ExtendedTypeText = PropDef.GetValuePropDef().GetCPPType();
+		return TEXT("TSET");
+	}
 };
 
 struct FPropertyTypeTraitsMap : public FPropertyTypeTraitsBase
@@ -2630,6 +2830,12 @@ struct FPropertyTypeTraitsMap : public FPropertyTypeTraitsBase
 	{
 		return FString::Printf(TEXT("%s %s"), *PropDef.GetKeyPropDef().GetCPPTypeForwardDeclaration(), *PropDef.GetValuePropDef().GetCPPTypeForwardDeclaration());
 	}
+
+	static FString GetCPPMacroType(const FUnrealPropertyDefinitionInfo& PropDef, FString& ExtendedTypeText)
+	{
+		ExtendedTypeText = FString::Printf(TEXT("%s,%s"), *PropDef.GetKeyPropDef().GetCPPType(), *PropDef.GetValuePropDef().GetCPPType());
+		return TEXT("TMAP");
+	}
 };
 
 //----------------------------------------------------------------------------------------------------------------------------------------------------------
@@ -2740,6 +2946,11 @@ FString FPropertyTraits::GetCPPTypeForwardDeclaration(const FUnrealPropertyDefin
 	return PropertyTypeDispatch<GetCPPTypeForwardDeclarationDispatch, true, FString>(PropDef);
 }
 
+FString FPropertyTraits::GetCPPMacroType(const FUnrealPropertyDefinitionInfo& PropDef, FString& ExtendedTypeText)
+{
+	return PropertyTypeDispatch<GetCPPMacroTypeDispatch, true, FString>(PropDef, ExtendedTypeText);
+}
+
 bool FPropertyTraits::SameType(const FUnrealPropertyDefinitionInfo& Lhs, const FUnrealPropertyDefinitionInfo& Rhs)
 {
 	auto AdjustTypes = [](EUHTPropertyType PropertyType)
diff --git a/Engine/Source/Programs/UnrealHeaderTool/Private/PropertyTypes.h b/Engine/Source/Programs/UnrealHeaderTool/Private/PropertyTypes.h
--- a/Engine/Source/Programs/UnrealHeaderTool/Private/PropertyTypes.h
+++ b/Engine/Source/Programs/UnrealHeaderTool/Private/PropertyTypes.h
@@ -75,6 +75,8 @@ struct FPropertyTraits
 
 	static FString GetCPPTypeForwardDeclaration(const FUnrealPropertyDefinitionInfo& PropDef);
 
+	static FString GetCPPMacroType(const FUnrealPropertyDefinitionInfo& PropDef, FString& ExtendedTypeText);
+
 	/**
 	 * Tests to see if the two types are the same types
 	 */
diff --git a/Engine/Source/Programs/UnrealHeaderTool/Private/UnrealTypeDefinitionInfo.cpp b/Engine/Source/Programs/UnrealHeaderTool/Private/UnrealTypeDefinitionInfo.cpp
--- a/Engine/Source/Programs/UnrealHeaderTool/Private/UnrealTypeDefinitionInfo.cpp
+++ b/Engine/Source/Programs/UnrealHeaderTool/Private/UnrealTypeDefinitionInfo.cpp
@@ -935,6 +935,146 @@ FText FUnrealPropertyDefinitionInfo::GetToolTipText(bool bShortTooltip) const
 	return LocalizedToolTip;
 }
 
+void FUnrealPropertyDefinitionInfo::ExportCppDeclaration(FOutputDevice& Out, EExportedDeclaration::Type DeclarationType, const TCHAR* ArrayDimOverride/* = NULL*/, uint32 AdditionalExportCPPFlags/* = 0*/, bool bSkipParameterName/* = false*/) const
+{
+	auto FnOld = [this](FOutputDevice& Out, EExportedDeclaration::Type DeclarationType, const TCHAR* ArrayDimOverride, uint32 AdditionalExportCPPFlags, bool bSkipParameterName)
+	{
+		GetProperty()->ExportCppDeclaration(Out, DeclarationType, ArrayDimOverride, AdditionalExportCPPFlags, bSkipParameterName);
+	};
+
+	auto FnNew = [this](FOutputDevice& Out, EExportedDeclaration::Type DeclarationType, const TCHAR* ArrayDimOverride, uint32 AdditionalExportCPPFlags, bool bSkipParameterName)
+	{
+		const FPropertyBase& VarProperty = GetPropertyBase();
+		const bool bIsParameter = (DeclarationType == EExportedDeclaration::Parameter) || (DeclarationType == EExportedDeclaration::MacroParameter);
+		const bool bIsInterfaceProp = VarProperty.IsInterfaceOrInterfaceStaticArray();
+
+		// export the property type text (e.g. FString; int32; TArray, etc.)
+		FString ExtendedTypeText;
+		const uint32 ExportCPPFlags = AdditionalExportCPPFlags | (bIsParameter ? CPPF_ArgumentOrReturnValue : 0);
+		FString TypeText = GetCPPType(&ExtendedTypeText, ExportCPPFlags);
+
+		const bool bCanHaveRef = 0 == (AdditionalExportCPPFlags & CPPF_NoRef);
+		const bool bCanHaveConst = 0 == (AdditionalExportCPPFlags & CPPF_NoConst);
+		if (!VarProperty.IsBooleanOrBooleanStaticArray() && bCanHaveConst) // can't have const bitfields because then we cannot determine their offset and mask from the compiler
+		{
+			// export 'const' for parameters
+			const bool bIsConstParam = bIsParameter && (HasAnyPropertyFlags(CPF_ConstParm) || (bIsInterfaceProp && !HasAllPropertyFlags(CPF_OutParm)));
+			const bool bIsOnConstClass = PropertyBase.IsObjectRefOrObjectRefStaticArray() && PropertyBase.ClassDef->HasAnyClassFlags(CLASS_Const);
+			const bool bShouldHaveRef = bCanHaveRef && HasAnyPropertyFlags(CPF_OutParm | CPF_ReferenceParm);
+
+			const bool bConstAtTheBeginning = bIsOnConstClass || (bIsConstParam && !bShouldHaveRef);
+			if (bConstAtTheBeginning)
+			{
+				TypeText = FString::Printf(TEXT("const %s"), *TypeText);
+			}
+
+			const FUnrealClassDefinitionInfo* const MyPotentialConstClass = (DeclarationType == EExportedDeclaration::Member) ? UHTCast<FUnrealClassDefinitionInfo>(GetOwnerObject()) : nullptr;
+			const bool bFromConstClass = MyPotentialConstClass && MyPotentialConstClass->HasAnyClassFlags(CLASS_Const);
+			const bool bConstAtTheEnd = bFromConstClass || (bIsConstParam && bShouldHaveRef);
+			if (bConstAtTheEnd)
+			{
+				ExtendedTypeText += TEXT(" const");
+			}
+		}
+
+		FString NameCpp;
+		if (!bSkipParameterName)
+		{
+			NameCpp = HasAnyPropertyFlags(CPF_Deprecated) ? GetNameCPP() + TEXT("_DEPRECATED") : GetNameCPP();
+		}
+		if (DeclarationType == EExportedDeclaration::MacroParameter)
+		{
+			NameCpp = FString(TEXT(", ")) + NameCpp;
+		}
+
+		TCHAR ArrayStr[MAX_SPRINTF] = TEXT("");
+		const bool bExportStaticArray = 0 == (CPPF_NoStaticArray & AdditionalExportCPPFlags);
+		if (ArrayDimOverride != nullptr && bExportStaticArray)
+		{
+			FCString::Sprintf(ArrayStr, TEXT("[%s]"), ArrayDimOverride);
+		}
+
+		if (VarProperty.IsBooleanOrBooleanStaticArray())
+		{
+			// if this is a member variable, export it as a bitfield
+			if (ArrayDimOverride == nullptr && DeclarationType == EExportedDeclaration::Member)
+			{
+				bool bCanUseBitfield = VarProperty.Type != CPT_Bool && GetVariableCategory() != EVariableCategory::Return;
+				// export as a uint32 member....bad to hardcode, but this is a special case that won't be used anywhere else
+				Out.Logf(TEXT("%s%s %s%s%s"), *TypeText, *ExtendedTypeText, *NameCpp, ArrayStr, bCanUseBitfield ? TEXT(":1") : TEXT(""));
+			}
+
+			//@todo we currently can't have out bools.. so this isn't really necessary, but eventually out bools may be supported, so leave here for now
+			else if (bIsParameter && HasAnyPropertyFlags(CPF_OutParm))
+			{
+				// export as a reference
+				Out.Logf(TEXT("%s%s%s %s%s"), *TypeText, *ExtendedTypeText
+					, bCanHaveRef ? TEXT("&") : TEXT("")
+					, *NameCpp, ArrayStr);
+			}
+
+			else
+			{
+				Out.Logf(TEXT("%s%s %s%s"), *TypeText, *ExtendedTypeText, *NameCpp, ArrayStr);
+			}
+		}
+		else
+		{
+			if (bIsParameter)
+			{
+				if (ArrayDimOverride != nullptr)
+				{
+					// don't export as a pointer
+					Out.Logf(TEXT("%s%s %s%s"), *TypeText, *ExtendedTypeText, *NameCpp, ArrayStr);
+				}
+				else
+				{
+					if (VarProperty.PassCPPArgsByRef())
+					{
+						// export as a reference (const ref if it isn't an out parameter)
+						Out.Logf(TEXT("%s%s%s%s %s"),
+							(bCanHaveConst && !HasAnyPropertyFlags(CPF_OutParm | CPF_ConstParm)) ? TEXT("const ") : TEXT(""),
+							*TypeText, *ExtendedTypeText,
+							bCanHaveRef ? TEXT("&") : TEXT(""),
+							*NameCpp);
+					}
+					else
+					{
+						// export as a pointer if this is an optional out parm, reference if it's just an out parm, standard otherwise...
+						TCHAR ModifierString[2] = { 0,0 };
+						if (bCanHaveRef && (HasAnyPropertyFlags(CPF_OutParm | CPF_ReferenceParm) || bIsInterfaceProp))
+						{
+							ModifierString[0] = TEXT('&');
+						}
+						Out.Logf(TEXT("%s%s%s %s%s"), *TypeText, *ExtendedTypeText, ModifierString, *NameCpp, ArrayStr);
+					}
+				}
+			}
+			else
+			{
+				Out.Logf(TEXT("%s%s %s%s"), *TypeText, *ExtendedTypeText, *NameCpp, ArrayStr);
+			}
+		}
+	};
+
+	FUHTStringBuilder OldSB, NewSB;
+	FnOld(OldSB, DeclarationType, ArrayDimOverride, AdditionalExportCPPFlags, bSkipParameterName);
+	FnNew(NewSB, DeclarationType, ArrayDimOverride, AdditionalExportCPPFlags, bSkipParameterName);
+	check(OldSB.Compare(NewSB, ESearchCase::CaseSensitive) == 0);
+	if (OldSB.Compare(NewSB, ESearchCase::CaseSensitive) != 0)
+	{
+		FUHTStringBuilder OldSB2, NewSB2;
+		FnOld(OldSB2, DeclarationType, ArrayDimOverride, AdditionalExportCPPFlags, bSkipParameterName);
+		FnNew(NewSB2, DeclarationType, ArrayDimOverride, AdditionalExportCPPFlags, bSkipParameterName);
+	}
+	Out.Log(NewSB);
+}
+
+FString FUnrealPropertyDefinitionInfo::GetCPPMacroType(FString& ExtendedTypeText) const
+{
+	return FPropertyTraits::GetCPPMacroType(*this, ExtendedTypeText);
+}
+
 FUnrealPackageDefinitionInfo& FUnrealPropertyDefinitionInfo::GetPackageDef() const
 {
 	if (HasSource())
diff --git a/Engine/Source/Programs/UnrealHeaderTool/Private/UnrealTypeDefinitionInfo.h b/Engine/Source/Programs/UnrealHeaderTool/Private/UnrealTypeDefinitionInfo.h
--- a/Engine/Source/Programs/UnrealHeaderTool/Private/UnrealTypeDefinitionInfo.h
+++ b/Engine/Source/Programs/UnrealHeaderTool/Private/UnrealTypeDefinitionInfo.h
@@ -674,16 +674,9 @@ public:
 		Property = InProperty;
 	}
 
-	void ExportCppDeclaration(FOutputDevice& Out, EExportedDeclaration::Type DeclarationType, const TCHAR* ArrayDimOverride = NULL, uint32 AdditionalExportCPPFlags = 0
-		, bool bSkipParameterName = false, const FString* ActualCppType = nullptr, const FString* ActualExtendedType = nullptr, const FString* ActualParameterName = nullptr) const
-	{
-		GetProperty()->ExportCppDeclaration(Out, DeclarationType, ArrayDimOverride, AdditionalExportCPPFlags, bSkipParameterName, ActualCppType, ActualExtendedType, ActualParameterName);
-	}
+	void ExportCppDeclaration(FOutputDevice& Out, EExportedDeclaration::Type DeclarationType, const TCHAR* ArrayDimOverride = NULL, uint32 AdditionalExportCPPFlags = 0, bool bSkipParameterName = false) const;
 
-	FString GetCPPMacroType(FString& ExtendedTypeText) const
-	{
-		return GetProperty()->GetCPPMacroType(ExtendedTypeText);
-	}
+	FString GetCPPMacroType(FString& ExtendedTypeText) const;
 
 	/**
 	 * Returns the name of the property
