Removing remaining places where GetScriptStruct is needed during code generation.

#rnx

diff --git a/Engine/Source/Programs/UnrealHeaderTool/Private/PropertyTypes.cpp b/Engine/Source/Programs/UnrealHeaderTool/Private/PropertyTypes.cpp
--- a/Engine/Source/Programs/UnrealHeaderTool/Private/PropertyTypes.cpp
+++ b/Engine/Source/Programs/UnrealHeaderTool/Private/PropertyTypes.cpp
@@ -2168,8 +2168,7 @@ struct FPropertyTypeTraitsStruct : public FPropertyTypeTraitsBase
 	static FString GetCPPTypeForwardDeclaration(const FUnrealPropertyDefinitionInfo& PropDef)
 	{
 		// Core type structs don't need to forward declare in UHT as every generated.h indirectly includes CoreMinimal.h
-		const UScriptStruct* Struct = PropDef.GetPropertyBase().ScriptStructDef->GetScriptStruct();
-		if (Struct->GetCppStructOps() && Struct->GetCppStructOps()->IsUECoreType())
+		if (UScriptStruct::ICppStructOps* CppStructOps = UScriptStruct::FindDeferredCppStructOps(PropDef.GetPropertyBase().ScriptStructDef->GetFName()); CppStructOps != nullptr && CppStructOps->IsUECoreType())
 		{
 			return FString();
 		}
diff --git a/Engine/Source/Programs/UnrealHeaderTool/Private/UnrealTypeDefinitionInfo.h b/Engine/Source/Programs/UnrealHeaderTool/Private/UnrealTypeDefinitionInfo.h
--- a/Engine/Source/Programs/UnrealHeaderTool/Private/UnrealTypeDefinitionInfo.h
+++ b/Engine/Source/Programs/UnrealHeaderTool/Private/UnrealTypeDefinitionInfo.h
@@ -2478,7 +2478,7 @@ public:
 	 */
 	bool HasDefaults() const
 	{
-		return !!GetScriptStruct()->GetCppStructOps();
+		return UScriptStruct::FindDeferredCppStructOps(GetFName()) != nullptr;
 	}
 
 
