Remove the need for GetFunction during code generation

#rnx

diff --git a/Engine/Source/Programs/UnrealHeaderTool/Private/CodeGenerator.cpp b/Engine/Source/Programs/UnrealHeaderTool/Private/CodeGenerator.cpp
--- a/Engine/Source/Programs/UnrealHeaderTool/Private/CodeGenerator.cpp
+++ b/Engine/Source/Programs/UnrealHeaderTool/Private/CodeGenerator.cpp
@@ -1474,7 +1474,7 @@ void FNativeClassHeaderGenerator::OutputProperty(FOutputDevice& DeclOut, FOutput
 
 	case EUHTPropertyType::MulticastDelegate:
 	{
-		bool bIsSparse = PropertyBase.FunctionDef->GetFunction()->IsA<USparseDelegateFunction>();
+		bool bIsSparse = PropertyBase.FunctionDef->GetFunctionType() == EFunctionType::SparseDelegate;
 
 		FString MetaDataParams = OutputMetaDataCodeForObject(DeclOut, Out, PropertyDef, *FString::Printf(TEXT("%s_MetaData"), *Name), DeclSpaces, Spaces);
 
@@ -2375,7 +2375,7 @@ void FNativeClassHeaderGenerator::ExportFunction(FOutputDevice& Out, FReferenceG
 		StructureSize = TEXT("0");
 	}
 
-	USparseDelegateFunction* SparseDelegateFunction = Cast<USparseDelegateFunction>(FunctionDef.GetFunction());
+	bool bIsSparse = FunctionDef.GetFunctionType() == EFunctionType::SparseDelegate;
 	const TCHAR* UFunctionObjectFlags = FunctionDef.IsOwnedByDynamicType() ? TEXT("RF_Public|RF_Transient") : TEXT("RF_Public|RF_Transient|RF_MarkAsNative");
 
 	TTuple<FString, FString> PropertyRange = OutputProperties(CurrentFunctionText, StaticDefinitions, OutReferenceGatherers, *FString::Printf(TEXT("%s::"), *StaticsStructName), FunctionDef, TEXT("\t\t"), TEXT("\t"));
@@ -2393,8 +2393,8 @@ void FNativeClassHeaderGenerator::ExportFunction(FOutputDevice& Out, FReferenceG
 		*OuterFunc,
 		*GetSingletonNameFuncAddr(SuperFunctionDef, OutReferenceGatherers.UniqueCrossModuleReferences),
 		*CreateUTF8LiteralString(FNativeClassHeaderGenerator::GetOverriddenName(FunctionDef)),
-		(SparseDelegateFunction ? *CreateUTF8LiteralString(SparseDelegateFunction->OwningClassName.ToString()) : TEXT("nullptr")),
-		(SparseDelegateFunction ? *CreateUTF8LiteralString(SparseDelegateFunction->DelegateName.ToString()) : TEXT("nullptr")),
+		(bIsSparse ? *CreateUTF8LiteralString(FunctionDef.GetSparseOwningClassName().ToString()) : TEXT("nullptr")),
+		(bIsSparse ? *CreateUTF8LiteralString(FunctionDef.GetSparseDelegateName().ToString()) : TEXT("nullptr")),
 		*StructureSize,
 		*PropertyRange.Get<0>(),
 		*PropertyRange.Get<1>(),
diff --git a/Engine/Source/Programs/UnrealHeaderTool/Private/PropertyTypes.cpp b/Engine/Source/Programs/UnrealHeaderTool/Private/PropertyTypes.cpp
--- a/Engine/Source/Programs/UnrealHeaderTool/Private/PropertyTypes.cpp
+++ b/Engine/Source/Programs/UnrealHeaderTool/Private/PropertyTypes.cpp
@@ -2222,23 +2222,21 @@ struct FPropertyTypeTraitsDelegate : public FPropertyTypeTraitsBase
 	{
 		const FPropertyBase& VarProperty = PropDef.GetPropertyBase();
 		check(VarProperty.FunctionDef);
-		UFunction* SignatureFunction = VarProperty.FunctionDef->GetFunction();
-		check(SignatureFunction);
 
-		FString UnmangledFunctionName = SignatureFunction->GetName().LeftChop(FString(HEADER_GENERATED_DELEGATE_SIGNATURE_SUFFIX).Len());
+		FString UnmangledFunctionName = VarProperty.FunctionDef->GetName().LeftChop(FString(HEADER_GENERATED_DELEGATE_SIGNATURE_SUFFIX).Len());
 		const bool bBlueprintCppBackend = (0 != (CPPExportFlags & EPropertyExportCPPFlags::CPPF_BlueprintCppBackend));
-		const bool bNative = SignatureFunction->IsNative();
+		const bool bNative = VarProperty.FunctionDef->IsNative();
 		if (bBlueprintCppBackend && bNative)
 		{
-			UStruct* StructOwner = Cast<UStruct>(SignatureFunction->GetOuter());
-			if (StructOwner)
+			if (FUnrealStructDefinitionInfo* StructOwnerDef = UHTCast<FUnrealStructDefinitionInfo>(VarProperty.FunctionDef->GetOuter()))
 			{
-				return FString::Printf(TEXT("%s%s::F%s"), StructOwner->GetPrefixCPP(), *StructOwner->GetName(), *UnmangledFunctionName);
+				return FString::Printf(TEXT("%s%s::F%s"), StructOwnerDef->GetPrefixCPP(), *StructOwnerDef->GetName(), *UnmangledFunctionName);
 			}
 		}
 		else
 		{
-			const bool NonNativeClassOwner = (SignatureFunction->GetOwnerClass() && !SignatureFunction->GetOwnerClass()->HasAnyClassFlags(CLASS_Native));
+			FUnrealClassDefinitionInfo* OwnerClass = VarProperty.FunctionDef->GetOwnerClass();
+			const bool NonNativeClassOwner = OwnerClass && !OwnerClass->HasAnyClassFlags(CLASS_Native);
 			if (bBlueprintCppBackend && NonNativeClassOwner)
 			{
 				// The name must be valid, this removes spaces, ?, etc from the user's function name. It could
@@ -2247,7 +2245,7 @@ struct FPropertyTypeTraitsDelegate : public FPropertyTypeTraitsBase
 				// identifier and collide:
 				UnmangledFunctionName = UnicodeToCPPIdentifier(UnmangledFunctionName, false, TEXT(""));
 				// the name must be unique
-				const FString OwnerName = UnicodeToCPPIdentifier(SignatureFunction->GetOwnerClass()->GetName(), false, TEXT(""));
+				const FString OwnerName = UnicodeToCPPIdentifier(OwnerClass->GetName(), false, TEXT(""));
 				const FString NewUnmangledFunctionName = FString::Printf(TEXT("%s__%s"), *UnmangledFunctionName, *OwnerName);
 				UnmangledFunctionName = NewUnmangledFunctionName;
 			}
@@ -2290,18 +2288,15 @@ struct FPropertyTypeTraitsMulticastDelegate : public FPropertyTypeTraitsBase
 	static FProperty* CreateEngineType(FUnrealPropertyDefinitionInfo& PropDef, FFieldVariant Scope, const FName& Name, EObjectFlags ObjectFlags)
 	{
 		const FPropertyBase& VarProperty = PropDef.GetPropertyBase();
-		UFunction* Function = PropDef.GetPropertyBase().FunctionDef->GetFunction();
 
-		FMulticastDelegateProperty* Result;
-		if (Function->IsA<USparseDelegateFunction>())
+		if (VarProperty.FunctionDef->GetFunctionType() == EFunctionType::SparseDelegate)
 		{
-			Result = new FMulticastSparseDelegateProperty(Scope, Name, ObjectFlags);
+			return new FMulticastSparseDelegateProperty(Scope, Name, ObjectFlags);
 		}
 		else
 		{
-			Result = new FMulticastInlineDelegateProperty(Scope, Name, ObjectFlags);
+			return new FMulticastInlineDelegateProperty(Scope, Name, ObjectFlags);
 		}
-		return Result;
 	}
 
 	static bool IsSupportedByBlueprint(const FUnrealPropertyDefinitionInfo& PropDef, bool bMemberVariable)
@@ -2312,8 +2307,7 @@ struct FPropertyTypeTraitsMulticastDelegate : public FPropertyTypeTraitsBase
 	static FString GetEngineClassName(const FUnrealPropertyDefinitionInfo& PropDef)
 	{
 		const FPropertyBase& VarProperty = PropDef.GetPropertyBase();
-		UFunction* Function = PropDef.GetPropertyBase().FunctionDef->GetFunction();
-		if (Function->IsA<USparseDelegateFunction>())
+		if (VarProperty.FunctionDef->GetFunctionType() == EFunctionType::SparseDelegate)
 		{
 			return FMulticastSparseDelegateProperty::StaticClass()->GetName();
 		}
@@ -2327,31 +2321,29 @@ struct FPropertyTypeTraitsMulticastDelegate : public FPropertyTypeTraitsBase
 	{
 		const FPropertyBase& VarProperty = PropDef.GetPropertyBase();
 		check(VarProperty.FunctionDef);
-		UFunction* SignatureFunction = VarProperty.FunctionDef->GetFunction();
 
 		// We have this test because sometimes the delegate hasn't been set up by FixupDelegateProperties at the time
 		// we need the type for an error message.  We deliberately format it so that it's unambiguously not CPP code, but is still human-readable.
-		if (!SignatureFunction)
+		if (!VarProperty.FunctionDef)
 		{
 			return FString(TEXT("{multicast delegate type}"));
 		}
 
-		FString UnmangledFunctionName = SignatureFunction->GetName().LeftChop(FString(HEADER_GENERATED_DELEGATE_SIGNATURE_SUFFIX).Len());
-		const UClass* OwnerClass = SignatureFunction->GetOwnerClass();
+		FString UnmangledFunctionName = VarProperty.FunctionDef->GetName().LeftChop(FString(HEADER_GENERATED_DELEGATE_SIGNATURE_SUFFIX).Len());
+		FUnrealClassDefinitionInfo* OwnerClassDef = VarProperty.FunctionDef->GetOwnerClass();
 
 		const bool bBlueprintCppBackend = (0 != (CPPExportFlags & EPropertyExportCPPFlags::CPPF_BlueprintCppBackend));
-		const bool bNative = SignatureFunction->IsNative();
+		const bool bNative = VarProperty.FunctionDef->IsNative();
 		if (bBlueprintCppBackend && bNative)
 		{
-			UStruct* StructOwner = Cast<UStruct>(SignatureFunction->GetOuter());
-			if (StructOwner)
+			if (FUnrealStructDefinitionInfo* StructOwnerDef = UHTCast<FUnrealStructDefinitionInfo>(VarProperty.FunctionDef->GetOuter()))
 			{
-				return FString::Printf(TEXT("%s%s::F%s"), StructOwner->GetPrefixCPP(), *StructOwner->GetName(), *UnmangledFunctionName);
+				return FString::Printf(TEXT("%s%s::F%s"), StructOwnerDef->GetPrefixCPP(), *StructOwnerDef->GetName(), *UnmangledFunctionName);
 			}
 		}
 		else
 		{
-			if ((0 != (CPPExportFlags & EPropertyExportCPPFlags::CPPF_BlueprintCppBackend)) && OwnerClass && !OwnerClass->HasAnyClassFlags(CLASS_Native))
+			if ((0 != (CPPExportFlags & EPropertyExportCPPFlags::CPPF_BlueprintCppBackend)) && OwnerClassDef && !OwnerClassDef->HasAnyClassFlags(CLASS_Native))
 			{
 				// The name must be valid, this removes spaces, ?, etc from the user's function name. It could
 				// be slightly shorter because the postfix ("__pf") is not needed here because we further post-
@@ -2359,7 +2351,7 @@ struct FPropertyTypeTraitsMulticastDelegate : public FPropertyTypeTraitsBase
 				// identifier and collide:
 				UnmangledFunctionName = UnicodeToCPPIdentifier(UnmangledFunctionName, false, TEXT(""));
 				// the name must be unique
-				const FString OwnerName = UnicodeToCPPIdentifier(OwnerClass->GetName(), false, TEXT(""));
+				const FString OwnerName = UnicodeToCPPIdentifier(OwnerClassDef->GetName(), false, TEXT(""));
 				const FString NewUnmangledFunctionName = FString::Printf(TEXT("%s__%s"), *UnmangledFunctionName, *OwnerName);
 				UnmangledFunctionName = NewUnmangledFunctionName;
 			}
diff --git a/Engine/Source/Programs/UnrealHeaderTool/Private/UnrealTypeDefinitionInfo.h b/Engine/Source/Programs/UnrealHeaderTool/Private/UnrealTypeDefinitionInfo.h
--- a/Engine/Source/Programs/UnrealHeaderTool/Private/UnrealTypeDefinitionInfo.h
+++ b/Engine/Source/Programs/UnrealHeaderTool/Private/UnrealTypeDefinitionInfo.h
@@ -2684,6 +2684,14 @@ public:
 	 */
 	FUnrealFunctionDefinitionInfo* GetSuperFunction() const;
 
+	/**
+	 * Get the name of the sparse owning class name
+	 */
+	FName GetSparseOwningClassName() const
+	{
+		return SparseOwningClassName;
+	}
+
 	/**
 	 * Set the name of the sparse owning class name
 	 */
@@ -2692,6 +2700,14 @@ public:
 		SparseOwningClassName = InSparseOwningClassName;
 	}
 
+	/**
+	 * Get the name of the sparse delegate name
+	 */
+	FName GetSparseDelegateName() const
+	{
+		return SparseDelegateName;
+	}
+
 	/**
 	 * Set the name of the sparse delegate name
 	 */
