Remove the need for GetClass during code generation

#rnx

diff --git a/Engine/Source/Programs/UnrealHeaderTool/Private/CodeGenerator.cpp b/Engine/Source/Programs/UnrealHeaderTool/Private/CodeGenerator.cpp
--- a/Engine/Source/Programs/UnrealHeaderTool/Private/CodeGenerator.cpp
+++ b/Engine/Source/Programs/UnrealHeaderTool/Private/CodeGenerator.cpp
@@ -280,7 +280,6 @@ namespace
 		const TCHAR* ClassCPPName,
 		const TCHAR* API,
 		FUnrealClassDefinitionInfo& ClassDef,
-		UClass* SuperClass,
 		FOutputDevice& Writer,
 		const FUnrealSourceFile& SourceFile,
 		EExportClassOutFlags& OutFlags)
@@ -1088,8 +1087,7 @@ void FNativeClassHeaderGenerator::OutputProperty(FOutputDevice& DeclOut, FOutput
 	{
 		FString MetaDataParams = OutputMetaDataCodeForObject(DeclOut, Out, PropertyDef, *FString::Printf(TEXT("%s_MetaData"), *Name), DeclSpaces, Spaces);
 
-		UClass* PropertyClass = PropertyBase.ClassDef->GetClass();
-		if (PropertyClass->IsChildOf(UClass::StaticClass()))
+		if (PropertyBase.ClassDef->IsChildOf(*GUClassDef))
 		{
 			DeclOut.Logf(TEXT("%sstatic const UECodeGen_Private::FClassPropertyParams %s;\r\n"), DeclSpaces, *NameWithoutScope);
 
@@ -1135,8 +1133,7 @@ void FNativeClassHeaderGenerator::OutputProperty(FOutputDevice& DeclOut, FOutput
 	{
 		FString MetaDataParams = OutputMetaDataCodeForObject(DeclOut, Out, PropertyDef, *FString::Printf(TEXT("%s_MetaData"), *Name), DeclSpaces, Spaces);
 
-		UClass* PropertyClass = PropertyBase.ClassDef->GetClass();
-		if (PropertyClass->IsChildOf(UClass::StaticClass()))
+		if (PropertyBase.ClassDef->IsChildOf(*GUClassDef))
 		{
 			DeclOut.Logf(TEXT("%sstatic const UECodeGen_Private::FClassPtrPropertyParams %s;\r\n"), DeclSpaces, *NameWithoutScope);
 
@@ -1182,8 +1179,7 @@ void FNativeClassHeaderGenerator::OutputProperty(FOutputDevice& DeclOut, FOutput
 	{
 		FString MetaDataParams = OutputMetaDataCodeForObject(DeclOut, Out, PropertyDef, *FString::Printf(TEXT("%s_MetaData"), *Name), DeclSpaces, Spaces);
 
-		UClass* PropertyClass = PropertyBase.ClassDef->GetClass();
-		if (PropertyClass->IsChildOf(UClass::StaticClass()))
+		if (PropertyBase.ClassDef->IsChildOf(*GUClassDef))
 		{
 			DeclOut.Logf(TEXT("%sstatic const UECodeGen_Private::FSoftClassPropertyParams %s;\r\n"), DeclSpaces, *NameWithoutScope);
 
@@ -2058,7 +2054,7 @@ void FNativeClassHeaderGenerator::ExportNativeGeneratedInitCode(FOutputDevice& O
 		bool bHasInterfaces = false;
 		for (FUnrealStructDefinitionInfo::FBaseStructInfo& BaseStruct : ClassDef.GetBaseStructInfos())
 		{
-			if (BaseStruct.Struct != nullptr && UHTCastChecked<FUnrealClassDefinitionInfo>(BaseStruct.Struct).IsInterface())
+			if (FUnrealClassDefinitionInfo* BaseClass = UHTCast<FUnrealClassDefinitionInfo>(BaseStruct.Struct); BaseClass != nullptr && BaseClass->IsInterface())
 			{
 				bHasInterfaces = true;
 				break;
@@ -2072,7 +2068,7 @@ void FNativeClassHeaderGenerator::ExportNativeGeneratedInitCode(FOutputDevice& O
 			StaticDefinitions.Logf(TEXT("\t\tconst UECodeGen_Private::FImplementedInterfaceParams %s::InterfaceParams[] = {\r\n"), *StaticsStructName);
 			for (FUnrealStructDefinitionInfo::FBaseStructInfo& BaseStruct : ClassDef.GetBaseStructInfos())
 			{
-				if (BaseStruct.Struct == nullptr || !UHTCastChecked<FUnrealClassDefinitionInfo>(BaseStruct.Struct).IsInterface())
+				if (FUnrealClassDefinitionInfo* BaseClass = UHTCast<FUnrealClassDefinitionInfo>(BaseStruct.Struct); BaseClass == nullptr || !BaseClass->IsInterface())
 				{
 					continue;
 				}
@@ -2661,10 +2657,8 @@ void FNativeClassHeaderGenerator::ExportClassFromSourceFileInner(
 	FUHTStringBuilder StandardUObjectConstructorsMacroCall;
 	FUHTStringBuilder EnhancedUObjectConstructorsMacroCall;
 
-	UClass* Class = ClassDef.GetClass();
 
 	FUnrealClassDefinitionInfo* SuperClassDef = ClassDef.GetSuperClass();
-	UClass* SuperClass = SuperClassDef ? SuperClassDef->GetClass() : nullptr;
 
 	// C++ -> VM stubs (native function execs)
 	FUHTStringBuilder ClassMacroCalls;
@@ -2790,7 +2784,7 @@ void FNativeClassHeaderGenerator::ExportClassFromSourceFileInner(
 		Boilerplate.Log(*FriendText);
 		Boilerplate.Logf(TEXT("public:\r\n"));
 
-		const bool bCastedClass = Class->HasAnyCastFlag(CASTCLASS_AllFlags) && SuperClassDef && Class->ClassCastFlags != SuperClass->ClassCastFlags;
+		const bool bCastedClass = ClassDef.HasAnyCastFlags(CASTCLASS_AllFlags) && SuperClassDef && ClassDef.GetClassCastFlags() != SuperClassDef->GetClassCastFlags();
 
 		Boilerplate.Logf(TEXT("\tDECLARE_CLASS(%s, %s, COMPILED_IN_FLAGS(%s%s), %s, TEXT(\"%s\"), %s_API)\r\n"),
 			*ClassCPPName,
@@ -2871,7 +2865,7 @@ void FNativeClassHeaderGenerator::ExportClassFromSourceFileInner(
 
 			// we'll need a way to get to the UObject portion of a native interface, so that we can safely pass native interfaces
 			// to script VM functions
-			if (SuperClass->IsChildOf(UInterface::StaticClass()))
+			if (SuperClassDef->IsChildOf(*GUInterfaceDef))
 			{
 				// Note: This used to be declared as a pure virtual function, but it was changed here in order to allow the Blueprint nativization process
 				// to detect C++ interface classes that explicitly declare pure virtual functions via type traits. This code will no longer trigger that check.
@@ -2880,7 +2874,7 @@ void FNativeClassHeaderGenerator::ExportClassFromSourceFileInner(
 
 			if (ClassHasReplicatedProperties(ClassDef))
 			{
-				WriteReplicatedMacroData(*ClassCPPName, *APIArg, ClassDef, SuperClass, InterfaceBoilerplate, SourceFile, OutFlags);
+				WriteReplicatedMacroData(*ClassCPPName, *APIArg, ClassDef, InterfaceBoilerplate, SourceFile, OutFlags);
 			}
 
 			FString NoPureDeclsMacroName = SourceFile.GetGeneratedMacroName(ClassDef.GetGeneratedBodyLine(), TEXT("_INCLASS_IINTERFACE_NO_PURE_DECLS"));
@@ -2900,14 +2894,18 @@ void FNativeClassHeaderGenerator::ExportClassFromSourceFileInner(
 			}
 
 			// export implementation of _getUObject for classes that implement interfaces
-			if (Class->Interfaces.Num() > 0)
+			for (FUnrealStructDefinitionInfo::FBaseStructInfo& BaseStruct : ClassDef.GetBaseStructInfos())
 			{
-				Boilerplate.Logf(TEXT("\tvirtual UObject* _getUObject() const override { return const_cast<%s*>(this); }\r\n"), *ClassCPPName);
+				if (FUnrealClassDefinitionInfo* BaseClass = UHTCast<FUnrealClassDefinitionInfo>(BaseStruct.Struct); BaseClass != nullptr && BaseClass->IsInterface())
+				{
+					Boilerplate.Logf(TEXT("\tvirtual UObject* _getUObject() const override { return const_cast<%s*>(this); }\r\n"), *ClassCPPName);
+					break;
+				}
 			}
 
 			if (ClassHasReplicatedProperties(ClassDef))
 			{
-				WriteReplicatedMacroData(*ClassCPPName, *APIArg, ClassDef, SuperClass, Boilerplate, SourceFile, OutFlags);
+				WriteReplicatedMacroData(*ClassCPPName, *APIArg, ClassDef, Boilerplate, SourceFile, OutFlags);
 			}
 
 			{
diff --git a/Engine/Source/Programs/UnrealHeaderTool/Private/HeaderParser.cpp b/Engine/Source/Programs/UnrealHeaderTool/Private/HeaderParser.cpp
--- a/Engine/Source/Programs/UnrealHeaderTool/Private/HeaderParser.cpp
+++ b/Engine/Source/Programs/UnrealHeaderTool/Private/HeaderParser.cpp
@@ -2380,7 +2380,7 @@ void FHeaderParser::ValidateClassFlags(const FUnrealClassDefinitionInfo& ToValid
 	if (ToValidate.HasAnyClassFlags(CLASS_NeedsDeferredDependencyLoading) && !ToValidate.IsChildOf(*GUClassDef))
 	{
 		// CLASS_NeedsDeferredDependencyLoading can only be set on classes derived from UClass
-		ToValidate.Throwf(TEXT("NeedsDeferredDependencyLoading is set on %s but the flag can only be used with classes derived from UClass."), *ToValidate.GetClass()->GetName());
+		ToValidate.Throwf(TEXT("NeedsDeferredDependencyLoading is set on %s but the flag can only be used with classes derived from UClass."), *ToValidate.GetName());
 	}
 }
 
diff --git a/Engine/Source/Programs/UnrealHeaderTool/Private/PropertyTypes.cpp b/Engine/Source/Programs/UnrealHeaderTool/Private/PropertyTypes.cpp
--- a/Engine/Source/Programs/UnrealHeaderTool/Private/PropertyTypes.cpp
+++ b/Engine/Source/Programs/UnrealHeaderTool/Private/PropertyTypes.cpp
@@ -1621,7 +1621,6 @@ struct FPropertyTypeTraitsSoftObjectReference : public FPropertyTypeTraitsObject
 		check(VarProperty.ClassDef);
 		if (VarProperty.ClassDef->IsChildOf(*GUClassDef))
 		{
-			UClass* MetaClass = VarProperty.MetaClassDef->GetClass();
 			return FString::Printf(TEXT("class %s%s;"), VarProperty.MetaClassDef->GetPrefixCPP(), *VarProperty.MetaClassDef->GetName());
 		}
 		else
@@ -1683,18 +1682,18 @@ struct FPropertyTypeTraitsInterface : public FPropertyTypeTraitsBase
 	{
 		if (ExtendedTypeText != NULL)
 		{
-			UClass* ExportClass = PropDef.GetPropertyBase().ClassDef->GetClass();
+			FUnrealClassDefinitionInfo* ExportClassDef = PropDef.GetPropertyBase().ClassDef;
 			if (0 == (CPPF_BlueprintCppBackend & CPPExportFlags))
 			{
-				while (ExportClass && !ExportClass->HasAnyClassFlags(CLASS_Native))
+				while (ExportClassDef && !ExportClassDef->HasAnyClassFlags(CLASS_Native))
 				{
-					ExportClass = ExportClass->GetSuperClass();
+					ExportClassDef = ExportClassDef->GetSuperClass();
 				}
 			}
-			check(ExportClass);
-			check(ExportClass->HasAnyClassFlags(CLASS_Interface) || 0 != (CPPF_BlueprintCppBackend & CPPExportFlags));
+			check(ExportClassDef);
+			check(ExportClassDef->HasAnyClassFlags(CLASS_Interface) || 0 != (CPPF_BlueprintCppBackend & CPPExportFlags));
 
-			*ExtendedTypeText = FString::Printf(TEXT("<I%s>"), *ExportClass->GetName());
+			*ExtendedTypeText = FString::Printf(TEXT("<I%s>"), *ExportClassDef->GetName());
 		}
 		return TEXT("TScriptInterface");
 	}
@@ -2235,8 +2234,8 @@ struct FPropertyTypeTraitsDelegate : public FPropertyTypeTraitsBase
 		}
 		else
 		{
-			FUnrealClassDefinitionInfo* OwnerClass = VarProperty.FunctionDef->GetOwnerClass();
-			const bool NonNativeClassOwner = OwnerClass && !OwnerClass->HasAnyClassFlags(CLASS_Native);
+			FUnrealClassDefinitionInfo* OwnerClassDef = VarProperty.FunctionDef->GetOwnerClass();
+			const bool NonNativeClassOwner = OwnerClassDef && !OwnerClassDef->HasAnyClassFlags(CLASS_Native);
 			if (bBlueprintCppBackend && NonNativeClassOwner)
 			{
 				// The name must be valid, this removes spaces, ?, etc from the user's function name. It could
@@ -2245,7 +2244,7 @@ struct FPropertyTypeTraitsDelegate : public FPropertyTypeTraitsBase
 				// identifier and collide:
 				UnmangledFunctionName = UnicodeToCPPIdentifier(UnmangledFunctionName, false, TEXT(""));
 				// the name must be unique
-				const FString OwnerName = UnicodeToCPPIdentifier(OwnerClass->GetName(), false, TEXT(""));
+				const FString OwnerName = UnicodeToCPPIdentifier(OwnerClassDef->GetName(), false, TEXT(""));
 				const FString NewUnmangledFunctionName = FString::Printf(TEXT("%s__%s"), *UnmangledFunctionName, *OwnerName);
 				UnmangledFunctionName = NewUnmangledFunctionName;
 			}
diff --git a/Engine/Source/Programs/UnrealHeaderTool/Private/UnrealTypeDefinitionInfo.h b/Engine/Source/Programs/UnrealHeaderTool/Private/UnrealTypeDefinitionInfo.h
--- a/Engine/Source/Programs/UnrealHeaderTool/Private/UnrealTypeDefinitionInfo.h
+++ b/Engine/Source/Programs/UnrealHeaderTool/Private/UnrealTypeDefinitionInfo.h
@@ -3011,6 +3011,18 @@ public:
 		}
 	}
 
+	/**
+	 * Used to safely check whether the passed in cast flag is set.
+	 *
+	 * @param	FlagsToCheck		Class flag(s) to check for
+	 * @return	true if the passed in flag is set, false otherwise
+	 *			(including no flag passed in, unless the FlagsToCheck is CLASS_AllFlags)
+	 */
+	bool HasAnyCastFlags(EClassCastFlags FlagsToCheck) const
+	{
+		return EnumHasAnyFlags(ClassCastFlags, FlagsToCheck);
+	}
+
 	/**
 	 * Return the flags that were parsed as part of the pre-parse phase
 	 */
