Remove some diagnostic code to test for version correctness.

#rnx

diff --git a/Engine/Source/Programs/UnrealHeaderTool/Private/UnrealTypeDefinitionInfo.cpp b/Engine/Source/Programs/UnrealHeaderTool/Private/UnrealTypeDefinitionInfo.cpp
--- a/Engine/Source/Programs/UnrealHeaderTool/Private/UnrealTypeDefinitionInfo.cpp
+++ b/Engine/Source/Programs/UnrealHeaderTool/Private/UnrealTypeDefinitionInfo.cpp
@@ -937,137 +937,117 @@ FText FUnrealPropertyDefinitionInfo::GetToolTipText(bool bShortTooltip) const
 
 void FUnrealPropertyDefinitionInfo::ExportCppDeclaration(FOutputDevice& Out, EExportedDeclaration::Type DeclarationType, const TCHAR* ArrayDimOverride/* = NULL*/, uint32 AdditionalExportCPPFlags/* = 0*/, bool bSkipParameterName/* = false*/) const
 {
-	auto FnOld = [this](FOutputDevice& Out, EExportedDeclaration::Type DeclarationType, const TCHAR* ArrayDimOverride, uint32 AdditionalExportCPPFlags, bool bSkipParameterName)
-	{
-		GetProperty()->ExportCppDeclaration(Out, DeclarationType, ArrayDimOverride, AdditionalExportCPPFlags, bSkipParameterName);
-	};
+	const FPropertyBase& VarProperty = GetPropertyBase();
+	const bool bIsParameter = (DeclarationType == EExportedDeclaration::Parameter) || (DeclarationType == EExportedDeclaration::MacroParameter);
+	const bool bIsInterfaceProp = VarProperty.IsInterfaceOrInterfaceStaticArray();
+
+	// export the property type text (e.g. FString; int32; TArray, etc.)
+	FString ExtendedTypeText;
+	const uint32 ExportCPPFlags = AdditionalExportCPPFlags | (bIsParameter ? CPPF_ArgumentOrReturnValue : 0);
+	FString TypeText = GetCPPType(&ExtendedTypeText, ExportCPPFlags);
 
-	auto FnNew = [this](FOutputDevice& Out, EExportedDeclaration::Type DeclarationType, const TCHAR* ArrayDimOverride, uint32 AdditionalExportCPPFlags, bool bSkipParameterName)
+	const bool bCanHaveRef = 0 == (AdditionalExportCPPFlags & CPPF_NoRef);
+	const bool bCanHaveConst = 0 == (AdditionalExportCPPFlags & CPPF_NoConst);
+	if (!VarProperty.IsBooleanOrBooleanStaticArray() && bCanHaveConst) // can't have const bitfields because then we cannot determine their offset and mask from the compiler
 	{
-		const FPropertyBase& VarProperty = GetPropertyBase();
-		const bool bIsParameter = (DeclarationType == EExportedDeclaration::Parameter) || (DeclarationType == EExportedDeclaration::MacroParameter);
-		const bool bIsInterfaceProp = VarProperty.IsInterfaceOrInterfaceStaticArray();
+		// export 'const' for parameters
+		const bool bIsConstParam = bIsParameter && (HasAnyPropertyFlags(CPF_ConstParm) || (bIsInterfaceProp && !HasAllPropertyFlags(CPF_OutParm)));
+		const bool bIsOnConstClass = PropertyBase.IsObjectRefOrObjectRefStaticArray() && PropertyBase.ClassDef->HasAnyClassFlags(CLASS_Const);
+		const bool bShouldHaveRef = bCanHaveRef && HasAnyPropertyFlags(CPF_OutParm | CPF_ReferenceParm);
 
-		// export the property type text (e.g. FString; int32; TArray, etc.)
-		FString ExtendedTypeText;
-		const uint32 ExportCPPFlags = AdditionalExportCPPFlags | (bIsParameter ? CPPF_ArgumentOrReturnValue : 0);
-		FString TypeText = GetCPPType(&ExtendedTypeText, ExportCPPFlags);
+		const bool bConstAtTheBeginning = bIsOnConstClass || (bIsConstParam && !bShouldHaveRef);
+		if (bConstAtTheBeginning)
+		{
+			TypeText = FString::Printf(TEXT("const %s"), *TypeText);
+		}
 
-		const bool bCanHaveRef = 0 == (AdditionalExportCPPFlags & CPPF_NoRef);
-		const bool bCanHaveConst = 0 == (AdditionalExportCPPFlags & CPPF_NoConst);
-		if (!VarProperty.IsBooleanOrBooleanStaticArray() && bCanHaveConst) // can't have const bitfields because then we cannot determine their offset and mask from the compiler
+		const FUnrealClassDefinitionInfo* const MyPotentialConstClass = (DeclarationType == EExportedDeclaration::Member) ? UHTCast<FUnrealClassDefinitionInfo>(GetOwnerObject()) : nullptr;
+		const bool bFromConstClass = MyPotentialConstClass && MyPotentialConstClass->HasAnyClassFlags(CLASS_Const);
+		const bool bConstAtTheEnd = bFromConstClass || (bIsConstParam && bShouldHaveRef);
+		if (bConstAtTheEnd)
 		{
-			// export 'const' for parameters
-			const bool bIsConstParam = bIsParameter && (HasAnyPropertyFlags(CPF_ConstParm) || (bIsInterfaceProp && !HasAllPropertyFlags(CPF_OutParm)));
-			const bool bIsOnConstClass = PropertyBase.IsObjectRefOrObjectRefStaticArray() && PropertyBase.ClassDef->HasAnyClassFlags(CLASS_Const);
-			const bool bShouldHaveRef = bCanHaveRef && HasAnyPropertyFlags(CPF_OutParm | CPF_ReferenceParm);
+			ExtendedTypeText += TEXT(" const");
+		}
+	}
 
-			const bool bConstAtTheBeginning = bIsOnConstClass || (bIsConstParam && !bShouldHaveRef);
-			if (bConstAtTheBeginning)
-			{
-				TypeText = FString::Printf(TEXT("const %s"), *TypeText);
-			}
+	FString NameCpp;
+	if (!bSkipParameterName)
+	{
+		NameCpp = HasAnyPropertyFlags(CPF_Deprecated) ? GetNameCPP() + TEXT("_DEPRECATED") : GetNameCPP();
+	}
+	if (DeclarationType == EExportedDeclaration::MacroParameter)
+	{
+		NameCpp = FString(TEXT(", ")) + NameCpp;
+	}
 
-			const FUnrealClassDefinitionInfo* const MyPotentialConstClass = (DeclarationType == EExportedDeclaration::Member) ? UHTCast<FUnrealClassDefinitionInfo>(GetOwnerObject()) : nullptr;
-			const bool bFromConstClass = MyPotentialConstClass && MyPotentialConstClass->HasAnyClassFlags(CLASS_Const);
-			const bool bConstAtTheEnd = bFromConstClass || (bIsConstParam && bShouldHaveRef);
-			if (bConstAtTheEnd)
-			{
-				ExtendedTypeText += TEXT(" const");
-			}
-		}
+	TCHAR ArrayStr[MAX_SPRINTF] = TEXT("");
+	const bool bExportStaticArray = 0 == (CPPF_NoStaticArray & AdditionalExportCPPFlags);
+	if (ArrayDimOverride != nullptr && bExportStaticArray)
+	{
+		FCString::Sprintf(ArrayStr, TEXT("[%s]"), ArrayDimOverride);
+	}
 
-		FString NameCpp;
-		if (!bSkipParameterName)
+	if (VarProperty.IsBooleanOrBooleanStaticArray())
+	{
+		// if this is a member variable, export it as a bitfield
+		if (ArrayDimOverride == nullptr && DeclarationType == EExportedDeclaration::Member)
 		{
-			NameCpp = HasAnyPropertyFlags(CPF_Deprecated) ? GetNameCPP() + TEXT("_DEPRECATED") : GetNameCPP();
+			bool bCanUseBitfield = VarProperty.Type != CPT_Bool && GetVariableCategory() != EVariableCategory::Return;
+			// export as a uint32 member....bad to hardcode, but this is a special case that won't be used anywhere else
+			Out.Logf(TEXT("%s%s %s%s%s"), *TypeText, *ExtendedTypeText, *NameCpp, ArrayStr, bCanUseBitfield ? TEXT(":1") : TEXT(""));
 		}
-		if (DeclarationType == EExportedDeclaration::MacroParameter)
+
+		//@todo we currently can't have out bools.. so this isn't really necessary, but eventually out bools may be supported, so leave here for now
+		else if (bIsParameter && HasAnyPropertyFlags(CPF_OutParm))
 		{
-			NameCpp = FString(TEXT(", ")) + NameCpp;
+			// export as a reference
+			Out.Logf(TEXT("%s%s%s %s%s"), *TypeText, *ExtendedTypeText
+				, bCanHaveRef ? TEXT("&") : TEXT("")
+				, *NameCpp, ArrayStr);
 		}
 
-		TCHAR ArrayStr[MAX_SPRINTF] = TEXT("");
-		const bool bExportStaticArray = 0 == (CPPF_NoStaticArray & AdditionalExportCPPFlags);
-		if (ArrayDimOverride != nullptr && bExportStaticArray)
+		else
 		{
-			FCString::Sprintf(ArrayStr, TEXT("[%s]"), ArrayDimOverride);
+			Out.Logf(TEXT("%s%s %s%s"), *TypeText, *ExtendedTypeText, *NameCpp, ArrayStr);
 		}
-
-		if (VarProperty.IsBooleanOrBooleanStaticArray())
+	}
+	else
+	{
+		if (bIsParameter)
 		{
-			// if this is a member variable, export it as a bitfield
-			if (ArrayDimOverride == nullptr && DeclarationType == EExportedDeclaration::Member)
-			{
-				bool bCanUseBitfield = VarProperty.Type != CPT_Bool && GetVariableCategory() != EVariableCategory::Return;
-				// export as a uint32 member....bad to hardcode, but this is a special case that won't be used anywhere else
-				Out.Logf(TEXT("%s%s %s%s%s"), *TypeText, *ExtendedTypeText, *NameCpp, ArrayStr, bCanUseBitfield ? TEXT(":1") : TEXT(""));
-			}
-
-			//@todo we currently can't have out bools.. so this isn't really necessary, but eventually out bools may be supported, so leave here for now
-			else if (bIsParameter && HasAnyPropertyFlags(CPF_OutParm))
-			{
-				// export as a reference
-				Out.Logf(TEXT("%s%s%s %s%s"), *TypeText, *ExtendedTypeText
-					, bCanHaveRef ? TEXT("&") : TEXT("")
-					, *NameCpp, ArrayStr);
-			}
-
-			else
+			if (ArrayDimOverride != nullptr)
 			{
+				// don't export as a pointer
 				Out.Logf(TEXT("%s%s %s%s"), *TypeText, *ExtendedTypeText, *NameCpp, ArrayStr);
 			}
-		}
-		else
-		{
-			if (bIsParameter)
+			else
 			{
-				if (ArrayDimOverride != nullptr)
+				if (VarProperty.PassCPPArgsByRef())
 				{
-					// don't export as a pointer
-					Out.Logf(TEXT("%s%s %s%s"), *TypeText, *ExtendedTypeText, *NameCpp, ArrayStr);
+					// export as a reference (const ref if it isn't an out parameter)
+					Out.Logf(TEXT("%s%s%s%s %s"),
+						(bCanHaveConst && !HasAnyPropertyFlags(CPF_OutParm | CPF_ConstParm)) ? TEXT("const ") : TEXT(""),
+						*TypeText, *ExtendedTypeText,
+						bCanHaveRef ? TEXT("&") : TEXT(""),
+						*NameCpp);
 				}
 				else
 				{
-					if (VarProperty.PassCPPArgsByRef())
-					{
-						// export as a reference (const ref if it isn't an out parameter)
-						Out.Logf(TEXT("%s%s%s%s %s"),
-							(bCanHaveConst && !HasAnyPropertyFlags(CPF_OutParm | CPF_ConstParm)) ? TEXT("const ") : TEXT(""),
-							*TypeText, *ExtendedTypeText,
-							bCanHaveRef ? TEXT("&") : TEXT(""),
-							*NameCpp);
-					}
-					else
+					// export as a pointer if this is an optional out parm, reference if it's just an out parm, standard otherwise...
+					TCHAR ModifierString[2] = { 0,0 };
+					if (bCanHaveRef && (HasAnyPropertyFlags(CPF_OutParm | CPF_ReferenceParm) || bIsInterfaceProp))
 					{
-						// export as a pointer if this is an optional out parm, reference if it's just an out parm, standard otherwise...
-						TCHAR ModifierString[2] = { 0,0 };
-						if (bCanHaveRef && (HasAnyPropertyFlags(CPF_OutParm | CPF_ReferenceParm) || bIsInterfaceProp))
-						{
-							ModifierString[0] = TEXT('&');
-						}
-						Out.Logf(TEXT("%s%s%s %s%s"), *TypeText, *ExtendedTypeText, ModifierString, *NameCpp, ArrayStr);
+						ModifierString[0] = TEXT('&');
 					}
+					Out.Logf(TEXT("%s%s%s %s%s"), *TypeText, *ExtendedTypeText, ModifierString, *NameCpp, ArrayStr);
 				}
 			}
-			else
-			{
-				Out.Logf(TEXT("%s%s %s%s"), *TypeText, *ExtendedTypeText, *NameCpp, ArrayStr);
-			}
 		}
-	};
-
-	FUHTStringBuilder OldSB, NewSB;
-	FnOld(OldSB, DeclarationType, ArrayDimOverride, AdditionalExportCPPFlags, bSkipParameterName);
-	FnNew(NewSB, DeclarationType, ArrayDimOverride, AdditionalExportCPPFlags, bSkipParameterName);
-	check(OldSB.Compare(NewSB, ESearchCase::CaseSensitive) == 0);
-	if (OldSB.Compare(NewSB, ESearchCase::CaseSensitive) != 0)
-	{
-		FUHTStringBuilder OldSB2, NewSB2;
-		FnOld(OldSB2, DeclarationType, ArrayDimOverride, AdditionalExportCPPFlags, bSkipParameterName);
-		FnNew(NewSB2, DeclarationType, ArrayDimOverride, AdditionalExportCPPFlags, bSkipParameterName);
+		else
+		{
+			Out.Logf(TEXT("%s%s %s%s"), *TypeText, *ExtendedTypeText, *NameCpp, ArrayStr);
+		}
 	}
-	Out.Log(NewSB);
 }
 
 FString FUnrealPropertyDefinitionInfo::GetCPPMacroType(FString& ExtendedTypeText) const
