Removed a dereference that wasn't needed.

#rnx

diff --git a/Engine/Source/Programs/UnrealHeaderTool/Private/BaseParser.cpp b/Engine/Source/Programs/UnrealHeaderTool/Private/BaseParser.cpp
--- a/Engine/Source/Programs/UnrealHeaderTool/Private/BaseParser.cpp
+++ b/Engine/Source/Programs/UnrealHeaderTool/Private/BaseParser.cpp
@@ -504,9 +504,11 @@ bool FBaseParser::GetToken(FToken& Token, bool bNoConsts/*=false*/, ESymbolParse
 	// Escaped character constant
 	else if (c == '\'')
 	{
+
+		// We try to skip the character constant value. But if it is backslash, we have to skip another character
 		if (*Pos++ == '\\')
 		{
-			*Pos++;
+			++Pos;
 		}
 
 		if (*Pos++ != '\'')
