1) Make the creation of engine types not required for post parsing.
2) Don't create engine types unless a script plugin is found

#rnx

diff --git a/Engine/Source/Programs/UnrealHeaderTool/Private/ClassMaps.h b/Engine/Source/Programs/UnrealHeaderTool/Private/ClassMaps.h
--- a/Engine/Source/Programs/UnrealHeaderTool/Private/ClassMaps.h
+++ b/Engine/Source/Programs/UnrealHeaderTool/Private/ClassMaps.h
@@ -80,6 +80,22 @@ struct FTypeDefinitionInfoMap
 		}
 	}
 
+	template <typename Type, typename Lambda>
+	Type* Find(Lambda&& InLambda)
+	{
+		for (const TPair<FName, TSharedRef<FUnrealTypeDefinitionInfo>>& KVP : DefinitionsByName)
+		{
+			if (Type* TypeDef = UHTCast<Type>(*KVP.Value))
+			{
+				if (InLambda(*TypeDef))
+				{
+					return TypeDef;
+				}
+			}
+		}
+		return nullptr;
+	}
+
 private:
 	TMap<FName, TSharedRef<FUnrealTypeDefinitionInfo>> DefinitionsByName;
 };
diff --git a/Engine/Source/Programs/UnrealHeaderTool/Private/CodeGenerator.cpp b/Engine/Source/Programs/UnrealHeaderTool/Private/CodeGenerator.cpp
--- a/Engine/Source/Programs/UnrealHeaderTool/Private/CodeGenerator.cpp
+++ b/Engine/Source/Programs/UnrealHeaderTool/Private/CodeGenerator.cpp
@@ -5607,6 +5607,11 @@ void FNativeClassHeaderGenerator::GenerateSourceFiles(
 	{
 		GeneratedCPP.AddExportTaskRef(ExportSourceTasks);
 	}
+
+	// This is strange, but flushing the log can take a long time.  Without an explicit flush, we were getting a long stall in the UE_LOG message
+	// during the detection of script plugins.  By doing it here, the flush should easily complete during code generation.
+	GLog->FlushThreadedLogs();
+
 	FTaskGraphInterface::Get().WaitUntilTasksComplete(ExportSourceTasks);
 #else
 	for (FGeneratedCPP& GeneratedCPP : GeneratedCPPs)
@@ -6112,6 +6117,11 @@ void FNativeClassHeaderGenerator::ExportUpdatedHeaders(FString&& PackageName, TA
 /** Get all script plugins based on ini setting */
 void GetScriptPlugins(TArray<IScriptGeneratorPluginInterface*>& ScriptPlugins)
 {
+	if (!GManifest.IsGameTarget)
+	{
+		UE_LOG(LogCompile, Log, TEXT("Script generator plugins only enabled in game targets."));
+	}
+
 	ScriptPlugins = IModularFeatures::Get().GetModularFeatureImplementations<IScriptGeneratorPluginInterface>(TEXT("ScriptGenerator"));
 	UE_LOG(LogCompile, Log, TEXT("Found %d script generator plugins."), ScriptPlugins.Num());
 
@@ -6740,7 +6750,7 @@ void ParseSourceFiles(TArray<FUnrealSourceFile*>& OrderedSourceFiles)
 #endif
 }
 
-void PrepareTypesForExport(TArray<FUnrealPackageDefinitionInfo*>& PackageDefs)
+void PostParseFinalize(TArray<FUnrealPackageDefinitionInfo*>& PackageDefs)
 {
 	// Until we fix the issue with types not having sources, we need to do this
 	TArray<FUnrealTypeDefinitionInfo*> NoSourceTypeDefs;
@@ -6752,56 +6762,60 @@ void PrepareTypesForExport(TArray<FUnrealPackageDefinitionInfo*>& PackageDefs)
 		}
 	});
 
-	for (FUnrealTypeDefinitionInfo* TypeDef : NoSourceTypeDefs)
-	{
-		TypeDef->PostParseFinalize(EPostParseFinalizePhase::PreCreateEngineTypes);
-	}
-
-	for (FUnrealPackageDefinitionInfo* PackageDef : PackageDefs)
+	auto PostParseFinalize = [&PackageDefs, &NoSourceTypeDefs](EPostParseFinalizePhase Phase)
 	{
-		FResults::Try([PackageDef]()
+		for (FUnrealTypeDefinitionInfo* TypeDef : NoSourceTypeDefs)
 		{
-			PackageDef->PostParseFinalize(EPostParseFinalizePhase::PreCreateEngineTypes);
-		});
-	}
+			TypeDef->PostParseFinalize(Phase);
+		}
 
-	for (FUnrealPackageDefinitionInfo* PackageDef : PackageDefs)
-	{
-		FResults::Try([PackageDef]()
+		for (FUnrealPackageDefinitionInfo* PackageDef : PackageDefs)
 		{
-			PackageDef->CreateUObjectEngineTypes(ECreateEngineTypesPhase::Phase1);
-		});
-	}
+			FResults::Try([PackageDef, Phase]()
+				{
+					PackageDef->PostParseFinalize(Phase);
+				});
+		}
+	};
 
-	for (FUnrealPackageDefinitionInfo* PackageDef : PackageDefs)
-	{
-		FResults::Try([PackageDef]()
-		{
-			PackageDef->CreateUObjectEngineTypes(ECreateEngineTypesPhase::Phase2);
-		});
-	}
+	PostParseFinalize(EPostParseFinalizePhase::Phase1);
+	PostParseFinalize(EPostParseFinalizePhase::Phase2);
+	FResults::WaitForErrorTasks();
 
 	for (FUnrealTypeDefinitionInfo* TypeDef : NoSourceTypeDefs)
 	{
-		TypeDef->PostParseFinalize(EPostParseFinalizePhase::PostCreateEngineTypes);
+		if (FUnrealObjectDefinitionInfo* ObjectDef = UHTCast<FUnrealObjectDefinitionInfo>(TypeDef))
+		{
+			check(ObjectDef->GetObject());
+		}
 	}
+}
 
-	for (FUnrealPackageDefinitionInfo* PackageDef : PackageDefs)
+void CreateEngineTypes(TArray<FUnrealPackageDefinitionInfo*>& PackageDefs)
+{
+	auto CreateEngineTypes = [&PackageDefs](ECreateEngineTypesPhase Phase)
 	{
-		FResults::Try([PackageDef]()
+		for (FUnrealPackageDefinitionInfo* PackageDef : PackageDefs)
 		{
-			PackageDef->PostParseFinalize(EPostParseFinalizePhase::PostCreateEngineTypes);
-		});
-	}
-	FResults::WaitForErrorTasks();
+			FResults::Try([PackageDef, Phase]()
+				{
+					PackageDef->CreateUObjectEngineTypes(Phase);
+				});
+		}
 
-	for (FUnrealTypeDefinitionInfo* TypeDef : NoSourceTypeDefs)
-	{
-		if (FUnrealObjectDefinitionInfo* ObjectDef = UHTCast<FUnrealObjectDefinitionInfo>(TypeDef))
+		for (FUnrealPackageDefinitionInfo* PackageDef : PackageDefs)
 		{
-			check(ObjectDef->GetObject());
+			FResults::Try([PackageDef, Phase]()
+				{
+					PackageDef->CreateUObjectEngineTypes(Phase);
+				});
 		}
-	}
+	};
+
+	CreateEngineTypes(ECreateEngineTypesPhase::Phase1);
+	CreateEngineTypes(ECreateEngineTypesPhase::Phase2);
+	CreateEngineTypes(ECreateEngineTypesPhase::Phase3);
+	FResults::WaitForErrorTasks();
 }
 
 void Export(TArray<FUnrealPackageDefinitionInfo*>& PackageDefs, TArray<FUnrealSourceFile*>& OrderedSourceFiles)
@@ -6855,19 +6869,8 @@ void ExportClassTreeToScriptPlugins(const TMap<UClass*, FUnrealSourceFile*>& Sou
 	}
 }
 
-void ExportToScriptPlugins(TArray<FUnrealPackageDefinitionInfo*>& PackageDefs, FString& ExternalDependencies)
+void ExportToScriptPlugins(TArray<IScriptGeneratorPluginInterface*>& ScriptPlugins, TArray<FUnrealPackageDefinitionInfo*>& PackageDefs, FString& ExternalDependencies)
 {
-	TArray<IScriptGeneratorPluginInterface*> ScriptPlugins;
-	// Can only export scripts for game targets
-	if (GManifest.IsGameTarget)
-	{
-		GetScriptPlugins(ScriptPlugins);
-	}
-	if (ScriptPlugins.Num() == 0)
-	{
-		return;
-	}
-
 	TMap<UClass*, FUnrealSourceFile*> SourceFileLookup;
 	for (FUnrealPackageDefinitionInfo* PackageDef : PackageDefs)
 	{
@@ -6921,10 +6924,8 @@ void ExportToScriptPlugins(TArray<FUnrealPackageDefinitionInfo*>& PackageDefs, F
 	return;
 }
 
-void ExportToScriptPlugins(TArray<FUnrealPackageDefinitionInfo*>& PackageDefs)
+void WriteExternalDependencies(const FString& ExternalDependencies)
 {
-	FString ExternalDependencies;
-	ExportToScriptPlugins(PackageDefs, ExternalDependencies);
 	FFileHelper::SaveStringToFile(ExternalDependencies, *GManifest.ExternalDependenciesFile);
 }
 
@@ -6988,16 +6989,22 @@ ECompilationResult::Type UnrealHeaderTool_Main(const FString& ModuleInfoFilename
 		GEngineClasses.Add(ClassDefRef);
 	}
 
+	FString ExternalDependencies;
+	TArray<IScriptGeneratorPluginInterface*> ScriptPlugins;
+
 	double TotalPreparseTime = FResults::TimedTry([&PackageDefs, &ModuleInfoPath]() { PreparseSources(PackageDefs, ModuleInfoPath); });
 	double TotalDefineTypesTime = FResults::TimedTry([&PackageDefs]() { DefineTypes(PackageDefs); });
 	double TotalResolveParentsTime = FResults::TimedTry([&PackageDefs]() { ResolveParents(PackageDefs); });
 	double TotalPrepareTypesForParsingTime = FResults::TimedTry([&PackageDefs]() { PrepareTypesForParsing(PackageDefs); });
 	double TotalTopologicalSortTime = FResults::TimedTry([&OrderedSourceFiles]() { TopologicalSort(OrderedSourceFiles); });
 	double TotalParseTime = FResults::TimedTry([&OrderedSourceFiles]() { ParseSourceFiles(OrderedSourceFiles); });
-	double TotalPrepareTypesForExportTime = FResults::TimedTry([&PackageDefs]() { PrepareTypesForExport(PackageDefs); });
+	double TotalPostParseFinalizeTime = FResults::TimedTry([&PackageDefs]() { PostParseFinalize(PackageDefs); });
 	TotalTopologicalSortTime += FResults::TimedTry([&OrderedSourceFiles]() { TopologicalSort(OrderedSourceFiles); }); // Sort again to include new dependencies
 	double TotalCodeGenTime = FResults::TimedTry([&PackageDefs, &OrderedSourceFiles]() { Export(PackageDefs, OrderedSourceFiles); });
-	double TotalPluginTime = FResults::TimedTry([&PackageDefs]() { ExportToScriptPlugins(PackageDefs); });
+	double TotalCheckForScriptPluginsTime = FResults::TimedTry([&ScriptPlugins]() { GetScriptPlugins(ScriptPlugins); });
+	double TotalCreateEngineTypesTime = ScriptPlugins.IsEmpty() ? 0.0 : FResults::TimedTry([&PackageDefs]() { CreateEngineTypes(PackageDefs); });
+	double TotalPluginTime = ScriptPlugins.IsEmpty() ? 0.0 : FResults::TimedTry([&ScriptPlugins, &PackageDefs, &ExternalDependencies]() { ExportToScriptPlugins(ScriptPlugins, PackageDefs, ExternalDependencies); });
+	double TotalWriteExternalDependenciesTime = FResults::TimedTry([&ExternalDependencies]() { WriteExternalDependencies(ExternalDependencies); });
 	double TotalSummaryTime = FResults::TimedTry([&PackageDefs]() { GenerateSummary(PackageDefs); });
 
 	// Finish all async file tasks before stopping the clock
@@ -7011,7 +7018,6 @@ ECompilationResult::Type UnrealHeaderTool_Main(const FString& ModuleInfoFilename
 	{
 		NumSources += Module.PublicUObjectClassesHeaders.Num() + Module.PublicUObjectHeaders.Num() + Module.PrivateUObjectHeaders.Num();
 	}
-	
 	UE_LOG(LogCompile, Log, TEXT("Preparing %d modules took %.3f seconds"), GManifest.Modules.Num(), TotalPrepareModuleTime);
 	UE_LOG(LogCompile, Log, TEXT("Preparsing %d sources took %.3f seconds"), NumSources, TotalPreparseTime);
 	UE_LOG(LogCompile, Log, TEXT("Defining types took %.3f seconds"), TotalDefineTypesTime);
@@ -7019,11 +7025,14 @@ ECompilationResult::Type UnrealHeaderTool_Main(const FString& ModuleInfoFilename
 	UE_LOG(LogCompile, Log, TEXT("Preparing types for parsing took %.3f seconds"), TotalPrepareTypesForParsingTime);
 	UE_LOG(LogCompile, Log, TEXT("Sorting files by dependencies took %.3f seconds"), TotalTopologicalSortTime);
 	UE_LOG(LogCompile, Log, TEXT("Parsing took %.3f seconds"), TotalParseTime);
-	UE_LOG(LogCompile, Log, TEXT("Preparing types for export took %.3f seconds"), TotalPrepareTypesForExportTime);
+	UE_LOG(LogCompile, Log, TEXT("Post parse finalization took %.3f seconds"), TotalPostParseFinalizeTime);
 	UE_LOG(LogCompile, Log, TEXT("Code generation took %.3f seconds"), TotalCodeGenTime);
+	UE_LOG(LogCompile, Log, TEXT("Check for script plugins took % .3f seconds"), TotalCheckForScriptPluginsTime);
+	UE_LOG(LogCompile, Log, TEXT("Create engine types took % .3f seconds"), TotalCreateEngineTypesTime);
 	UE_LOG(LogCompile, Log, TEXT("ScriptPlugin overhead was %.3f seconds"), TotalPluginTime);
+	UE_LOG(LogCompile, Log, TEXT("Write external dependencies overhead was %.3f seconds"), TotalWriteExternalDependenciesTime);
 	UE_LOG(LogCompile, Log, TEXT("Summary generation took %.3f seconds"), TotalSummaryTime);
-	UE_LOG(LogCompile, Log, TEXT("Macroize time was %.3f seconds"), GMacroizeTime);
+	UE_LOG(LogCompile, Log, TEXT("Macroize time was %.3f CPU seconds"), GMacroizeTime);
 
 	FUnrealHeaderToolStats& Stats = FUnrealHeaderToolStats::Get();
 	for (const TPair<FName, double>& Pair : Stats.Counters)
diff --git a/Engine/Source/Programs/UnrealHeaderTool/Private/PropertyTypes.cpp b/Engine/Source/Programs/UnrealHeaderTool/Private/PropertyTypes.cpp
--- a/Engine/Source/Programs/UnrealHeaderTool/Private/PropertyTypes.cpp
+++ b/Engine/Source/Programs/UnrealHeaderTool/Private/PropertyTypes.cpp
@@ -132,6 +132,74 @@ namespace
 		}
 	};
 
+	class FEnumLookup
+	{
+	public:
+		FEnumLookup()
+		{
+			GTypeDefinitionInfoMap.ForAllTypesByName([this](FUnrealTypeDefinitionInfo& TypeDef) 
+				{
+					if (FUnrealEnumDefinitionInfo* EnumDef = UHTCast<FUnrealEnumDefinitionInfo>(TypeDef))
+					{
+						for (const TPair<FName, int64>& Kvp : EnumDef->GetEnums())
+						{
+							EnumValueMap.Add(Kvp.Key, FEnumAndValue{ EnumDef, Kvp.Value });
+						}
+					}
+				}
+			);
+		}
+
+		// This code emulates UEnum::LookupEnumNameSlow.  There are some cases where the Value check for INDEX_NONE is 
+		// rejecting matches.
+		FUnrealEnumDefinitionInfo* Find(const TCHAR* EnumValueName)
+		{
+			const FEnumAndValue* Value = FindEnumAndValue(EnumValueName);
+			if (Value != nullptr && Value->Value != INDEX_NONE)
+			{
+				return Value->Enum;
+			}
+
+			FString TestShortName = FString(TEXT("::")) + EnumValueName;
+			return GTypeDefinitionInfoMap.Find<FUnrealEnumDefinitionInfo>([&TestShortName](FUnrealEnumDefinitionInfo& EnumDef)
+				{
+					for (const TPair<FName, int64>& NameAndValue : EnumDef.GetEnums())
+					{
+						if (NameAndValue.Key.ToString().Contains(TestShortName))
+						{
+							return true;
+						}
+					}
+					return false;
+				});
+		}
+
+	private:
+		struct FEnumAndValue
+		{
+			FUnrealEnumDefinitionInfo* Enum;
+			int64 Value;
+		};
+
+		const FEnumAndValue* FindEnumAndValue(const TCHAR* EnumValueName)
+		{
+			FName Key(EnumValueName, FNAME_Find);
+			if (Key == NAME_None)
+			{
+				return nullptr;
+			}
+			return EnumValueMap.Find(Key);
+		}
+
+		TMap<FName, FEnumAndValue> EnumValueMap;
+	};
+
+	FEnumLookup& GetEnumLookup()
+	{
+		static FEnumLookup EnumLookup;
+		return EnumLookup;
+	}
+
 	//----------------------------------------------------------------------------------------------------------------------------------------------------------
 	//----------------------------------------------------------------------------------------------------------------------------------------------------------
 	// Field class bootstrap system
@@ -629,13 +697,14 @@ struct FPropertyTypeTraitsInt : public FPropertyTypeTraitsNumericBase
 
 	static void CreateProperty(FUnrealPropertyDefinitionInfo& PropDef)
 	{
+		const FPropertyBase& VarProperty = PropDef.GetPropertyBase();
+		PropDef.SetUnsized(VarProperty.IntType == EIntType::Unsized);
 	}
 
 	static FProperty* CreateEngineType(FUnrealPropertyDefinitionInfo& PropDef, FFieldVariant Scope, const FName& Name, EObjectFlags ObjectFlags)
 	{
 		const FPropertyBase& VarProperty = PropDef.GetPropertyBase();
 		FIntProperty* Result = new FIntProperty(Scope, Name, ObjectFlags);
-		PropDef.SetUnsized(VarProperty.IntType == EIntType::Unsized);
 		return Result;
 	}
 
@@ -725,13 +794,13 @@ struct FPropertyTypeTraitsUInt32 : public FPropertyTypeTraitsNumericBase
 {
 	static void CreateProperty(FUnrealPropertyDefinitionInfo& PropDef)
 	{
+		const FPropertyBase& VarProperty = PropDef.GetPropertyBase();
+		PropDef.SetUnsized(VarProperty.IntType == EIntType::Unsized);
 	}
 
 	static FProperty* CreateEngineType(FUnrealPropertyDefinitionInfo& PropDef, FFieldVariant Scope, const FName& Name, EObjectFlags ObjectFlags)
 	{
-		const FPropertyBase& VarProperty = PropDef.GetPropertyBase();
 		FUInt32Property* Result = new FUInt32Property(Scope, Name, ObjectFlags);
-		PropDef.SetUnsized(VarProperty.IntType == EIntType::Unsized);
 		return Result;
 	}
 
@@ -2442,16 +2511,8 @@ struct FPropertyTypeTraitsStaticArray : public FPropertyTypeTraitsBase
 
 	static FPropertyFlagsChanges PostParseFinalize(FUnrealPropertyDefinitionInfo& PropDef, EPropertyFlags& AddedFlags)
 	{
-		return PostParseFinalizeHelper<false>(PropDef, AddedFlags);
-	}
-
-	static FProperty* CreateEngineType(FUnrealPropertyDefinitionInfo& PropDef, FFieldVariant Scope, const FName& Name, EObjectFlags ObjectFlags)
-	{
-
 		if (PropDef.GetArrayDimensions())
 		{
-			UEnum* Enum = nullptr;
-
 			FString Temp = PropDef.GetArrayDimensions();
 
 			bool bAgain;
@@ -2496,21 +2557,29 @@ struct FPropertyTypeTraitsStaticArray : public FPropertyTypeTraitsBase
 				}
 			} while (bAgain);
 
-			UEnum::LookupEnumNameSlow(*Temp, &Enum);
-
-			if (!Enum)
+			if (Temp.Len() > 0 && !FChar::IsDigit(Temp[0]))
 			{
-				// If the enum wasn't declared in this scope, then try to find it anywhere we can
-				Enum = FEngineAPI::FindObject<UEnum>(ANY_PACKAGE, PropDef.GetArrayDimensions());
-			}
+				FUnrealEnumDefinitionInfo* EnumDef = GetEnumLookup().Find(*Temp);
 
-			if (Enum)
-			{
-				// set the ArraySizeEnum if applicable
-				PropDef.GetPropertyBase().MetaData.Add(NAME_ArraySizeEnum, Enum->GetPathName());
+				if (!EnumDef)
+				{
+					// If the enum wasn't declared in this scope, then try to find it anywhere we can
+					EnumDef = GTypeDefinitionInfoMap.FindByName<FUnrealEnumDefinitionInfo>(PropDef.GetArrayDimensions());
+				}
+
+				if (EnumDef)
+				{
+					// set the ArraySizeEnum if applicable
+					PropDef.GetPropertyBase().MetaData.Add(NAME_ArraySizeEnum, EnumDef->GetPathName());
+				}
 			}
 		}
 
+		return PostParseFinalizeHelper<false>(PropDef, AddedFlags);
+	}
+
+	static FProperty* CreateEngineType(FUnrealPropertyDefinitionInfo& PropDef, FFieldVariant Scope, const FName& Name, EObjectFlags ObjectFlags)
+	{
 		FProperty* Property = CreateEngineTypeHelper<false>(PropDef, Scope, Name, ObjectFlags);
 		Property->ArrayDim = 2;
 		return Property;
diff --git a/Engine/Source/Programs/UnrealHeaderTool/Private/UnrealTypeDefinitionInfo.cpp b/Engine/Source/Programs/UnrealHeaderTool/Private/UnrealTypeDefinitionInfo.cpp
--- a/Engine/Source/Programs/UnrealHeaderTool/Private/UnrealTypeDefinitionInfo.cpp
+++ b/Engine/Source/Programs/UnrealHeaderTool/Private/UnrealTypeDefinitionInfo.cpp
@@ -563,11 +563,11 @@
 {
 	switch (Phase)
 	{
-	case EPostParseFinalizePhase::PreCreateEngineTypes:
+	case EPostParseFinalizePhase::Phase1:
 		FPropertyTraits::PostParseFinalize(*this);
 		break;
 
-	case EPostParseFinalizePhase::PostCreateEngineTypes:
+	case EPostParseFinalizePhase::Phase2:
 
 		// Due to structures not being fully parsed at parse time, we have to do blueprint checks here
 		if (HasAnyPropertyFlags(CPF_BlueprintAssignable | CPF_BlueprintCallable) && GetPropertyBase().IsMulticastDelegateOrMulticastDelegateStaticArray())
@@ -622,12 +622,28 @@
 			}
 		}
 
-		//@TODO: The creation of engine type should be part of the engine type creation system
+		TypePackageName = GetTypePackageNameHelper(*this);
+		break;
+	}
+}
+
+void FUnrealPropertyDefinitionInfo::CreateUObjectEngineTypesInternal(ECreateEngineTypesPhase Phase)
+{
+	FUnrealTypeDefinitionInfo::CreateUObjectEngineTypesInternal(Phase);
+
+	switch (Phase)
+	{
+	case ECreateEngineTypesPhase::Phase1:
+		break;
+
+	case ECreateEngineTypesPhase::Phase2:
+		break;
+
+	case ECreateEngineTypesPhase::Phase3:
 		if (GetPropertySafe() == nullptr)
 		{
 			FPropertyTraits::CreateEngineType(SharedThis(this));
 		}
-		TypePackageName = GetTypePackageNameHelper(*this);
 		break;
 	}
 }
@@ -1236,10 +1252,10 @@
 
 	switch (Phase)
 	{
-	case EPostParseFinalizePhase::PreCreateEngineTypes:
+	case EPostParseFinalizePhase::Phase1:
 		break;
 
-	case EPostParseFinalizePhase::PostCreateEngineTypes:
+	case EPostParseFinalizePhase::Phase2:
 	{
 		UPackage* Package = GetPackage();
 
@@ -1276,10 +1292,10 @@
 
 	switch (Phase)
 	{
-	case EPostParseFinalizePhase::PreCreateEngineTypes:
+	case EPostParseFinalizePhase::Phase1:
 		break;
 
-	case EPostParseFinalizePhase::PostCreateEngineTypes:
+	case EPostParseFinalizePhase::Phase2:
 	{
 		const FString& ClassName = GetEngineClassName(true);
 		const FString& PackageShortName = GetPackageDef().GetShortUpperName();
@@ -1694,6 +1710,9 @@
 
 	case ECreateEngineTypesPhase::Phase2:
 		break;
+
+	case ECreateEngineTypesPhase::Phase3:
+		break;
 	}
 }
 
@@ -1771,6 +1790,23 @@
 			FunctionDef->CreateUObjectEngineTypes(Phase);
 		}
 		break;
+
+	case ECreateEngineTypesPhase::Phase3:
+		for (TSharedRef<FUnrealPropertyDefinitionInfo> PropertyDef : Properties)
+		{
+			PropertyDef->CreateUObjectEngineTypes(Phase);
+		}
+
+		GetStruct()->Bind();
+
+		// Internals will assert of we are relinking an intrinsic 
+		bool bRelinkExistingProperties = true;
+		if (FUnrealClassDefinitionInfo* ClassDef = UHTCast<FUnrealClassDefinitionInfo>(this))
+		{
+			bRelinkExistingProperties = !ClassDef->HasAnyClassFlags(CLASS_Intrinsic);
+		}
+		GetStruct()->StaticLink(bRelinkExistingProperties);
+		break;
 	}
 }
 
@@ -1803,24 +1839,12 @@
 
 	switch (Phase)
 	{
-	case EPostParseFinalizePhase::PreCreateEngineTypes:
+	case EPostParseFinalizePhase::Phase1:
 		break;
 
-	case EPostParseFinalizePhase::PostCreateEngineTypes:
-	{
-
-		GetStruct()->Bind();
-
-		// Internals will assert of we are relinking an intrinsic 
-		bool bRelinkExistingProperties = true;
-		if (FUnrealClassDefinitionInfo* ClassDef = UHTCast<FUnrealClassDefinitionInfo>(this))
-		{
-			bRelinkExistingProperties = !ClassDef->HasAnyClassFlags(CLASS_Intrinsic);
-		}
-		GetStruct()->StaticLink(bRelinkExistingProperties);
+	case EPostParseFinalizePhase::Phase2:
 		break;
 	}
-	}
 }
 
 void FUnrealStructDefinitionInfo::CreatePropertyEngineTypes()
@@ -1916,6 +1940,9 @@
 
 	case ECreateEngineTypesPhase::Phase2:
 		break;
+
+	case ECreateEngineTypesPhase::Phase3:
+		break;
 	}
 }
 
@@ -1925,7 +1952,7 @@
 
 	switch (Phase)
 	{
-	case EPostParseFinalizePhase::PreCreateEngineTypes:
+	case EPostParseFinalizePhase::Phase1:
 		// Collect the instanced reference state for the structure and propagate it to any
 		// property that references the structure.
 		if (FUnrealScriptStructDefinitionInfo* Base = GetSuperScriptStruct())
@@ -1949,7 +1976,7 @@
 		}
 		break;
 
-	case EPostParseFinalizePhase::PostCreateEngineTypes:
+	case EPostParseFinalizePhase::Phase2:
 		break;
 	}
 }
@@ -2141,8 +2168,67 @@
 	}
 
 	case ECreateEngineTypesPhase::Phase2:
+	{
+
+		// Initialize the class object
+		UClass* Class = GetClass();
+
+		// Clear the property size
+		Class->PropertiesSize = 0;
+
+		// Make the visible to the package
+		Class->ClearFlags(RF_Transient);
+		check(Class->HasAnyFlags(RF_Public));
+		check(Class->HasAnyFlags(RF_Standalone));
+
+		// Finalize all of the children introduced in this class
+		for (TSharedRef<FUnrealFunctionDefinitionInfo> FunctionDef : GetFunctions())
+		{
+			UFunction* Function = FunctionDef->GetFunction();
+			Class->AddFunctionToFunctionMap(Function, Function->GetFName());
+		}
+
+
+		if (!HasAnyClassFlags(CLASS_Native))
+		{
+			//Class->UnMark(EObjectMark(OBJECTMARK_TagImp | OBJECTMARK_TagExp));
+		}
+		else if (!HasAnyClassFlags(CLASS_NoExport | CLASS_Intrinsic))
+		{
+			//Class->UnMark(OBJECTMARK_TagImp);
+			//Class->Mark(OBJECTMARK_TagExp);
+		}
+
+		// This needs to be done outside of parallel blocks because it will modify UClass memory.
+		// Later calls to SetUpUhtReplicationData inside parallel blocks should be fine, because
+		// they will see the memory has already been set up, and just return the parent pointer.
+		Class->SetUpUhtReplicationData();
 		break;
 	}
+
+	case ECreateEngineTypesPhase::Phase3:
+	{
+
+		// Initialize the class object
+		UClass* Class = GetClass();
+
+		if (!HasAnyClassFlags(CLASS_Native))
+		{
+			//Class->UnMark(EObjectMark(OBJECTMARK_TagImp | OBJECTMARK_TagExp));
+		}
+		else if (!HasAnyClassFlags(CLASS_NoExport | CLASS_Intrinsic))
+		{
+			//Class->UnMark(OBJECTMARK_TagImp);
+			//Class->Mark(OBJECTMARK_TagExp);
+		}
+
+		// This needs to be done outside of parallel blocks because it will modify UClass memory.
+		// Later calls to SetUpUhtReplicationData inside parallel blocks should be fine, because
+		// they will see the memory has already been set up, and just return the parent pointer.
+		Class->SetUpUhtReplicationData();
+		break;
+	}
+	}
 }
 
 void FUnrealClassDefinitionInfo::PostParseFinalizeInternal(EPostParseFinalizePhase Phase)
@@ -2151,10 +2237,10 @@
 
 	switch (Phase)
 	{
-	case EPostParseFinalizePhase::PreCreateEngineTypes:
+	case EPostParseFinalizePhase::Phase1:
 		break;
 
-	case EPostParseFinalizePhase::PostCreateEngineTypes:
+	case EPostParseFinalizePhase::Phase2:
 	{
 
 		// Pick up any newly set flags from the super class (i.e. HasInstancedReference)
@@ -2199,39 +2285,13 @@
 			}
 		}
 
-		// Initialize the class object
-		UClass* Class = GetClass();
-
-		// Clear the property size
-		Class->PropertiesSize = 0;
-
-		// Make the visible to the package
-		Class->ClearFlags(RF_Transient);
-		check(Class->HasAnyFlags(RF_Public));
-		check(Class->HasAnyFlags(RF_Standalone));
-
-		// Finalize all of the children introduced in this class
-		for (TSharedRef<FUnrealFunctionDefinitionInfo> FunctionDef : GetFunctions())
-		{
-			UFunction* Function = FunctionDef->GetFunction();
-			Class->AddFunctionToFunctionMap(Function, Function->GetFName());
-		}
-
 		if (!HasAnyClassFlags(CLASS_Native))
 		{
-			//Class->UnMark(EObjectMark(OBJECTMARK_TagImp | OBJECTMARK_TagExp));
 		}
 		else if (!HasAnyClassFlags(CLASS_NoExport | CLASS_Intrinsic))
 		{
 			GetPackageDef().SetWriteClassesH(true);
-			//Class->UnMark(OBJECTMARK_TagImp);
-			//Class->Mark(OBJECTMARK_TagExp);
 		}
-
-		// This needs to be done outside of parallel blocks because it will modify UClass memory.
-		// Later calls to SetUpUhtReplicationData inside parallel blocks should be fine, because
-		// they will see the memory has already been set up, and just return the parent pointer.
-		Class->SetUpUhtReplicationData();
 		break;
 	}
 	}
@@ -2906,6 +2966,18 @@
 		Function->Bind();
 		break;
 	}
+
+	case ECreateEngineTypesPhase::Phase3:
+	{
+		// The following code is only performed on functions in a class.  
+		if (UHTCast<FUnrealClassDefinitionInfo>(GetOuter()) != nullptr)
+		{
+			UFunction* Function = GetFunction();
+			Function->StaticLink(true);
+			Function->ParmsSize = 0;
+		}
+		break;
 	}
+	}
 }
 
@@ -2916,10 +2988,10 @@
 
 	switch (Phase)
 	{
-	case EPostParseFinalizePhase::PreCreateEngineTypes:
+	case EPostParseFinalizePhase::Phase1:
 		break;
 
-	case EPostParseFinalizePhase::PostCreateEngineTypes:
+	case EPostParseFinalizePhase::Phase2:
 	{
 
 		// Due to structures that might not be fully parsed at parse time, do blueprint validation here
@@ -2942,17 +3014,9 @@
 			}
 		}
 
-		UFunction* Function = GetFunction();
-
 		// The following code is only performed on functions in a class.  
 		if (UHTCast<FUnrealClassDefinitionInfo>(GetOuter()) != nullptr)
 		{
-
-			Function->StaticLink(true);
-
-			// Compute the function parameter size, propagate some flags to the outer function, and save the return offset
-			// Must be done in a second phase, as StaticLink resets various fields again!
-			Function->ParmsSize = 0;
 			for (TSharedRef<FUnrealPropertyDefinitionInfo> PropertyDef : GetProperties())
 			{
 				if (PropertyDef->HasSpecificPropertyFlags(CPF_ReturnParm | CPF_OutParm, CPF_OutParm))
diff --git a/Engine/Source/Programs/UnrealHeaderTool/Private/UnrealTypeDefinitionInfo.h b/Engine/Source/Programs/UnrealHeaderTool/Private/UnrealTypeDefinitionInfo.h
--- a/Engine/Source/Programs/UnrealHeaderTool/Private/UnrealTypeDefinitionInfo.h
+++ b/Engine/Source/Programs/UnrealHeaderTool/Private/UnrealTypeDefinitionInfo.h
@@ -81,14 +81,15 @@ enum class EParsedInterface
 enum class ECreateEngineTypesPhase : uint8
 {
 	Phase1,				// Create the Class, ScriptStruct, and Enum types
-	Phase2,				// Create the functions and properties
+	Phase2,				// Create the functions
+	Phase3,				// Create the properties and do final initialization
 	MAX_Phases,
 };
 
 enum class EPostParseFinalizePhase : uint8
 {
-	PreCreateEngineTypes,
-	PostCreateEngineTypes,
+	Phase1,				// Flag propagation for structures and property finalization
+	Phase2,				// All other finalization requirements
 	MAX_Phases,
 };
 
@@ -1238,6 +1239,8 @@ protected:
 	}
 #endif
 
+	virtual void CreateUObjectEngineTypesInternal(ECreateEngineTypesPhase Phase) override;
+
 	/**
 	 * Perform any post parsing finalization and validation
 	 */
