1) Moved some string initialization to the parsing code so it can be run concurrently with all other parsing and not be applied to the post parsing time.
2) Improved the performance of the array dimension enum lookup code.
3) Moved type shutdown to be tracked in the app instead of at c++ exit time.
4) Changed to TSharedRef iteration to use references instead of by value to avoid the performance hit.

#rnx

diff --git a/Engine/Source/Programs/UnrealHeaderTool/Private/ClassMaps.h b/Engine/Source/Programs/UnrealHeaderTool/Private/ClassMaps.h
--- a/Engine/Source/Programs/UnrealHeaderTool/Private/ClassMaps.h
+++ b/Engine/Source/Programs/UnrealHeaderTool/Private/ClassMaps.h
@@ -96,6 +96,11 @@ struct FTypeDefinitionInfoMap
 		return nullptr;
 	}
 
+	void Reset()
+	{
+		DefinitionsByName.Reset();
+	}
+
 private:
 	TMap<FName, TSharedRef<FUnrealTypeDefinitionInfo>> DefinitionsByName;
 };
@@ -122,6 +127,12 @@ struct FUnrealSourceFiles : public FFreezableContainer
 		return AllSourceFiles;
 	}
 
+	void Reset()
+	{
+		AllSourceFiles.Reset();
+		SourceFilesByString.Reset();
+	}
+
 private:
 	// A map of all source files indexed by string.
 	TMap<FString, TSharedRef<FUnrealSourceFile>> SourceFilesByString;
diff --git a/Engine/Source/Programs/UnrealHeaderTool/Private/CodeGenerator.cpp b/Engine/Source/Programs/UnrealHeaderTool/Private/CodeGenerator.cpp
--- a/Engine/Source/Programs/UnrealHeaderTool/Private/CodeGenerator.cpp
+++ b/Engine/Source/Programs/UnrealHeaderTool/Private/CodeGenerator.cpp
@@ -2171,7 +2171,7 @@ void FNativeClassHeaderGenerator::ExportNativeGeneratedInitCode(FOutputDevice& O
 		GeneratedClassRegisterFunctionText.Logf(TEXT("\t\treturn OuterClass;\r\n"));
 		GeneratedClassRegisterFunctionText.Logf(TEXT("\t}\r\n"));
 
-		Out.Logf(TEXT("%s"), *GeneratedClassRegisterFunctionText);
+		Out.Log(*GeneratedClassRegisterFunctionText);
 	}
 
 	if (OutFriendText.Len() && bIsNoExport)
@@ -2579,7 +2579,7 @@ void FNativeClassHeaderGenerator::ExportInterfaceCallFunctions(FOutputDevice& Ou
 				}
 				bFirst = false;
 
-				OutCpp.Logf(TEXT("%s"), *ParamDef->GetName());
+				OutCpp.Log(*ParamDef->GetName());
 			}
 
 			OutCpp.Logf(TEXT(");") LINE_TERMINATOR);
@@ -3176,7 +3176,7 @@ void FNativeClassHeaderGenerator::ExportConstructorsMacros(FOutputDevice& OutGen
 bool FNativeClassHeaderGenerator::WriteHeader(FGeneratedFileInfo& FileInfo, const FString& InBodyText, const TSet<FString>& InAdditionalHeaders, const TSet<FString>& ForwardDeclarations)
 {
 	FUHTStringBuilder GeneratedHeaderTextWithCopyright;
-	GeneratedHeaderTextWithCopyright.Logf(TEXT("%s"), HeaderCopyright);
+	GeneratedHeaderTextWithCopyright.Log(HeaderCopyright);
 	GeneratedHeaderTextWithCopyright.Log(TEXT("#include \"UObject/ObjectMacros.h\"\r\n"));
 	GeneratedHeaderTextWithCopyright.Log(TEXT("#include \"UObject/ScriptMacros.h\"\r\n"));
 
@@ -3983,8 +3983,8 @@ void FNativeClassHeaderGenerator::ExportGeneratedEnumInitCode(FOutputDevice& Out
 
 	FUHTStringBuilder GeneratedEnumRegisterFunctionText;
 
-	GeneratedEnumRegisterFunctionText.Logf(TEXT("%s"), *StaticDeclarations);
-	GeneratedEnumRegisterFunctionText.Logf(TEXT("%s"), *StaticDefinitions);
+	GeneratedEnumRegisterFunctionText.Log(*StaticDeclarations);
+	GeneratedEnumRegisterFunctionText.Log(*StaticDefinitions);
 
 	GeneratedEnumRegisterFunctionText.Logf(TEXT("\tUEnum* %s\r\n"), *EnumSingletonName);
 	GeneratedEnumRegisterFunctionText.Logf(TEXT("\t{\r\n"));
@@ -4610,7 +4610,7 @@ void FNativeClassHeaderGenerator::ExportNativeFunctionHeader(
 	// Emit extra parameter if we have one
 	if( ExtraParam )
 	{
-		Out.Logf(TEXT("%s"), ExtraParam);
+		Out.Log(ExtraParam);
 		++ParmCount;
 	}
 
@@ -4633,7 +4633,7 @@ void FNativeClassHeaderGenerator::ExportNativeFunctionHeader(
 		PropertyDef->ExportCppDeclaration( PropertyText, EExportedDeclaration::Parameter, PropertyDef->GetArrayDimensions() );
 		ApplyAlternatePropertyExportText(*PropertyDef, PropertyText, EExportingState::Normal);
 
-		Out.Logf(TEXT("%s"), *PropertyText);
+		Out.Log(*PropertyText);
 	}
 
 	Out.Log( TEXT(")") );
@@ -5032,7 +5032,7 @@ void FNativeClassHeaderGenerator::ExportNativeFunctions(FOutputDevice& OutGenera
 				FuncStringBuilders.AutogeneratedBlueprintFunctionDeclarations.Log(*ValidDecl);
 				if (!Validation.WasFound())
 				{
-					FuncStringBuilders.AutogeneratedBlueprintFunctionDeclarationsOnlyNotDeclared.Logf(TEXT("%s"), *ValidDecl);
+					FuncStringBuilders.AutogeneratedBlueprintFunctionDeclarationsOnlyNotDeclared.Log(*ValidDecl);
 				}
 			}
 
@@ -6778,6 +6778,11 @@ void PostParseFinalize(TArray<FUnrealPackageDefinitionInfo*>& PackageDefs)
 		}
 	};
 
+	for (FUnrealTypeDefinitionInfo* TypeDef : NoSourceTypeDefs)
+	{
+		TypeDef->ConcurrentPostParseFinalize();
+	}
+
 	PostParseFinalize(EPostParseFinalizePhase::Phase1);
 	PostParseFinalize(EPostParseFinalizePhase::Phase2);
 	FResults::WaitForErrorTasks();
@@ -7010,6 +7015,8 @@ ECompilationResult::Type UnrealHeaderTool_Main(const FString& ModuleInfoFilename
 	// Finish all async file tasks before stopping the clock
 	FTaskGraphInterface::Get().WaitUntilTasksComplete(GAsyncFileTasks);
 
+	double TotalShutdownTime = FResults::TimedTry([]() { GTypeDefinitionInfoMap.Reset(); GUnrealSourceFilesMap.Reset(); });
+
 	MainTimer.Stop();
 
 	// Count the number of sources
@@ -7027,12 +7034,13 @@ ECompilationResult::Type UnrealHeaderTool_Main(const FString& ModuleInfoFilename
 	UE_LOG(LogCompile, Log, TEXT("Parsing took %.3f seconds"), TotalParseTime);
 	UE_LOG(LogCompile, Log, TEXT("Post parse finalization took %.3f seconds"), TotalPostParseFinalizeTime);
 	UE_LOG(LogCompile, Log, TEXT("Code generation took %.3f seconds"), TotalCodeGenTime);
-	UE_LOG(LogCompile, Log, TEXT("Check for script plugins took % .3f seconds"), TotalCheckForScriptPluginsTime);
+	UE_LOG(LogCompile, Log, TEXT("Check for script plugins took %.3f seconds"), TotalCheckForScriptPluginsTime);
 	UE_LOG(LogCompile, Log, TEXT("Create engine types took % .3f seconds"), TotalCreateEngineTypesTime);
 	UE_LOG(LogCompile, Log, TEXT("ScriptPlugin overhead was %.3f seconds"), TotalPluginTime);
 	UE_LOG(LogCompile, Log, TEXT("Write external dependencies overhead was %.3f seconds"), TotalWriteExternalDependenciesTime);
 	UE_LOG(LogCompile, Log, TEXT("Summary generation took %.3f seconds"), TotalSummaryTime);
 	UE_LOG(LogCompile, Log, TEXT("Macroize time was %.3f CPU seconds"), GMacroizeTime);
+	UE_LOG(LogCompile, Log, TEXT("Freeing types was %.3f seconds"), TotalShutdownTime);
 
 	FUnrealHeaderToolStats& Stats = FUnrealHeaderToolStats::Get();
 	for (const TPair<FName, double>& Pair : Stats.Counters)
diff --git a/Engine/Source/Programs/UnrealHeaderTool/Private/HeaderParser.cpp b/Engine/Source/Programs/UnrealHeaderTool/Private/HeaderParser.cpp
--- a/Engine/Source/Programs/UnrealHeaderTool/Private/HeaderParser.cpp
+++ b/Engine/Source/Programs/UnrealHeaderTool/Private/HeaderParser.cpp
@@ -7295,6 +7295,12 @@ void FHeaderParser::ParseHeader()
 		ClassDef->AsClassChecked().SetClassFlags(CLASS_Parsed);
 	}
 
+	// Perform any final concurrent finalization
+	for (const TSharedRef<FUnrealTypeDefinitionInfo>& TypeDef : SourceFile.GetDefinedTypes())
+	{
+		TypeDef->ConcurrentPostParseFinalize();
+	}
+
 	// Remember stats about this file
 	SourceFile.SetLinesParsed(LinesParsed);
 	SourceFile.SetStatementsParsed(StatementsParsed);
diff --git a/Engine/Source/Programs/UnrealHeaderTool/Private/PropertyTypes.cpp b/Engine/Source/Programs/UnrealHeaderTool/Private/PropertyTypes.cpp
--- a/Engine/Source/Programs/UnrealHeaderTool/Private/PropertyTypes.cpp
+++ b/Engine/Source/Programs/UnrealHeaderTool/Private/PropertyTypes.cpp
@@ -141,9 +141,18 @@ namespace
 				{
 					if (FUnrealEnumDefinitionInfo* EnumDef = UHTCast<FUnrealEnumDefinitionInfo>(TypeDef))
 					{
+						bool bAddShortNames = EnumDef->GetCppForm() == UEnum::ECppForm::EnumClass || EnumDef->GetCppForm() == UEnum::ECppForm::Namespaced;
+						FString CheckName = EnumDef->GetNameCPP() + TEXT("::");
 						for (const TPair<FName, int64>& Kvp : EnumDef->GetEnums())
 						{
-							EnumValueMap.Add(Kvp.Key, FEnumAndValue{ EnumDef, Kvp.Value });
+							FullEnumValueMap.Add(Kvp.Key, FEnumAndValue{ EnumDef, Kvp.Value });
+							if (bAddShortNames)
+							{
+								FString EnumName = Kvp.Key.ToString();
+								check(EnumName.StartsWith(CheckName, ESearchCase::CaseSensitive));
+								EnumName.RightChopInline(CheckName.Len());
+								ShortEnumValueMap.Add(FName(EnumName), EnumDef);
+							}
 						}
 					}
 				}
@@ -154,24 +163,26 @@ namespace
 		// rejecting matches.
 		FUnrealEnumDefinitionInfo* Find(const TCHAR* EnumValueName)
 		{
-			const FEnumAndValue* Value = FindEnumAndValue(EnumValueName);
+			FName Key(EnumValueName, FNAME_Find);
+			if (Key == NAME_None)
+			{
+				return nullptr;
+			}
+
+			// If we can find the full name and it has an index, then use this enum
+			const FEnumAndValue* Value = FullEnumValueMap.Find(Key);
 			if (Value != nullptr && Value->Value != INDEX_NONE)
 			{
 				return Value->Enum;
 			}
 
-			FString TestShortName = FString(TEXT("::")) + EnumValueName;
-			return GTypeDefinitionInfoMap.Find<FUnrealEnumDefinitionInfo>([&TestShortName](FUnrealEnumDefinitionInfo& EnumDef)
-				{
-					for (const TPair<FName, int64>& NameAndValue : EnumDef.GetEnums())
-					{
-						if (NameAndValue.Key.ToString().Contains(TestShortName))
-						{
-							return true;
-						}
-					}
-					return false;
-				});
+			// If the enum name already has a ::, then we can't do a short name search
+			if (FCString::Strifind(EnumValueName, TEXT("::"), false) != nullptr)
+			{
+				return nullptr;
+			}
+
+			return ShortEnumValueMap.FindRef(Key);
 		}
 
 	private:
@@ -181,17 +192,8 @@ namespace
 			int64 Value;
 		};
 
-		const FEnumAndValue* FindEnumAndValue(const TCHAR* EnumValueName)
-		{
-			FName Key(EnumValueName, FNAME_Find);
-			if (Key == NAME_None)
-			{
-				return nullptr;
-			}
-			return EnumValueMap.Find(Key);
-		}
-
-		TMap<FName, FEnumAndValue> EnumValueMap;
+		TMap<FName, FEnumAndValue> FullEnumValueMap;
+		TMap<FName, FUnrealEnumDefinitionInfo*> ShortEnumValueMap;
 	};
 
 	FEnumLookup& GetEnumLookup()
diff --git a/Engine/Source/Programs/UnrealHeaderTool/Private/UnrealTypeDefinitionInfo.cpp b/Engine/Source/Programs/UnrealHeaderTool/Private/UnrealTypeDefinitionInfo.cpp
--- a/Engine/Source/Programs/UnrealHeaderTool/Private/UnrealTypeDefinitionInfo.cpp
+++ b/Engine/Source/Programs/UnrealHeaderTool/Private/UnrealTypeDefinitionInfo.cpp
@@ -621,8 +621,6 @@ void FUnrealPropertyDefinitionInfo::PostParseFinalizeInternal(EPostParseFinalize
 				}
 			}
 		}
-
-		TypePackageName = GetTypePackageNameHelper(*this);
 		break;
 	}
 }
@@ -648,6 +646,12 @@ void FUnrealPropertyDefinitionInfo::CreateUObjectEngineTypesInternal(ECreateEngi
 	}
 }
 
+void FUnrealPropertyDefinitionInfo::ConcurrentPostParseFinalize()
+{
+	FUnrealTypeDefinitionInfo::ConcurrentPostParseFinalize();
+	TypePackageName = GetTypePackageNameHelper(*this);
+}
+
 bool FUnrealPropertyDefinitionInfo::IsDynamic() const
 {
 	return HasMetaData(NAME_ReplaceConverted);
@@ -1296,26 +1300,30 @@ void FUnrealFieldDefinitionInfo::PostParseFinalizeInternal(EPostParseFinalizePha
 		break;
 
 	case EPostParseFinalizePhase::Phase2:
-	{
-		const FString& ClassName = GetEngineClassName(true);
-		const FString& PackageShortName = GetPackageDef().GetShortUpperName();
+		break;
+	}
+}
 
-		FUHTStringBuilder Out;
-		GenerateSingletonName(Out, this, false);
-		ExternDecl[0].Appendf(TEXT("\t%s_API U%s* %s;\r\n"), *PackageShortName, *ClassName, *Out);
-		SingletonName[0] = Out;
-		SingletonNameChopped[0] = SingletonName[0].LeftChop(2);
+void FUnrealFieldDefinitionInfo::ConcurrentPostParseFinalize()
+{
+	FUnrealTypeDefinitionInfo::ConcurrentPostParseFinalize();
 
-		Out.Reset();
-		GenerateSingletonName(Out, this, true);
-		ExternDecl[1].Appendf(TEXT("\t%s_API U%s* %s;\r\n"), *PackageShortName, *ClassName, *Out);
-		SingletonName[1] = Out;
-		SingletonNameChopped[1] = SingletonName[1].LeftChop(2);
+	const FString& ClassName = GetEngineClassName(true);
+	const FString& PackageShortName = GetPackageDef().GetShortUpperName();
 
-		TypePackageName = GetTypePackageNameHelper(*this);
-		break;
-	}
-	}
+	FUHTStringBuilder Out;
+	GenerateSingletonName(Out, this, false);
+	ExternDecl[0].Appendf(TEXT("\t%s_API U%s* %s;\r\n"), *PackageShortName, *ClassName, *Out);
+	SingletonName[0] = Out;
+	SingletonNameChopped[0] = SingletonName[0].LeftChop(2);
+
+	Out.Reset();
+	GenerateSingletonName(Out, this, true);
+	ExternDecl[1].Appendf(TEXT("\t%s_API U%s* %s;\r\n"), *PackageShortName, *ClassName, *Out);
+	SingletonName[1] = Out;
+	SingletonNameChopped[1] = SingletonName[1].LeftChop(2);
+
+	TypePackageName = GetTypePackageNameHelper(*this);
 }
 
 FText FUnrealFieldDefinitionInfo::GetToolTipText(bool bShortTooltip) const
@@ -1785,14 +1793,14 @@ void FUnrealStructDefinitionInfo::CreateUObjectEngineTypesInternal(ECreateEngine
 		break;
 
 	case ECreateEngineTypesPhase::Phase2:
-		for (TSharedRef<FUnrealFunctionDefinitionInfo> FunctionDef : Functions)
+		for (TSharedRef<FUnrealFunctionDefinitionInfo>& FunctionDef : Functions)
 		{
 			FunctionDef->CreateUObjectEngineTypes(Phase);
 		}
 		break;
 
 	case ECreateEngineTypesPhase::Phase3:
-		for (TSharedRef<FUnrealPropertyDefinitionInfo> PropertyDef : Properties)
+		for (TSharedRef<FUnrealPropertyDefinitionInfo>& PropertyDef : Properties)
 		{
 			PropertyDef->CreateUObjectEngineTypes(Phase);
 		}
@@ -1827,12 +1835,12 @@ void FUnrealStructDefinitionInfo::PostParseFinalizeInternal(EPostParseFinalizePh
 		}
 	}
 
-	for (TSharedRef<FUnrealFunctionDefinitionInfo> FunctionDef : Functions)
+	for (TSharedRef<FUnrealFunctionDefinitionInfo>& FunctionDef : Functions)
 	{
 		FunctionDef->PostParseFinalize(Phase);
 	}
 
-	for (TSharedRef<FUnrealPropertyDefinitionInfo> PropertyDef : Properties)
+	for (TSharedRef<FUnrealPropertyDefinitionInfo>& PropertyDef : Properties)
 	{
 		PropertyDef->PostParseFinalize(Phase);
 	}
@@ -1847,9 +1855,23 @@ void FUnrealStructDefinitionInfo::PostParseFinalizeInternal(EPostParseFinalizePh
 	}
 }
 
+void FUnrealStructDefinitionInfo::ConcurrentPostParseFinalize()
+{
+	FUnrealFieldDefinitionInfo::ConcurrentPostParseFinalize();
+	for (TSharedRef<FUnrealFunctionDefinitionInfo>& FunctionDef : Functions)
+	{
+		FunctionDef->ConcurrentPostParseFinalize();
+	}
+
+	for (TSharedRef<FUnrealPropertyDefinitionInfo>& PropertyDef : Properties)
+	{
+		PropertyDef->ConcurrentPostParseFinalize();
+	}
+}
+
 void FUnrealStructDefinitionInfo::CreatePropertyEngineTypes()
 {
-	for (TSharedRef<FUnrealPropertyDefinitionInfo> PropertyDef : Properties)
+	for (TSharedRef<FUnrealPropertyDefinitionInfo>& PropertyDef : Properties)
 	{
 		if (PropertyDef->GetPropertySafe() == nullptr)
 		{
diff --git a/Engine/Source/Programs/UnrealHeaderTool/Private/UnrealTypeDefinitionInfo.h b/Engine/Source/Programs/UnrealHeaderTool/Private/UnrealTypeDefinitionInfo.h
--- a/Engine/Source/Programs/UnrealHeaderTool/Private/UnrealTypeDefinitionInfo.h
+++ b/Engine/Source/Programs/UnrealHeaderTool/Private/UnrealTypeDefinitionInfo.h
@@ -426,6 +426,11 @@ public:
 	 */
 	void PostParseFinalize(EPostParseFinalizePhase Phase);
 
+	/**
+	 * Perform any post parsing finalization that can happen specific to the containing source file
+	 */
+	virtual void ConcurrentPostParseFinalize() {}
+
 	/**
 	 * Return the compilation scope associated with this object
 	 */
@@ -1227,6 +1232,11 @@ public:
 	/** Linked list of properties referencing a specific field */
 	FUnrealPropertyDefinitionInfo* NextReferencingProperty = nullptr;
 
+	/**
+	 * Perform any post parsing finalization that can happen specific to the containing source file
+	 */
+	virtual void ConcurrentPostParseFinalize() override;
+
 protected:
 #if UHT_ENABLE_ENGINE_TYPE_CHECKS
 	virtual void CheckFindMetaData(const FName& Key, const FString* ValuePtr) const
@@ -1805,6 +1815,11 @@ public:
 		}
 	}
 
+	/**
+	 * Perform any post parsing finalization that can happen specific to the containing source file
+	 */
+	virtual void ConcurrentPostParseFinalize() override;
+
 protected:
 	/**
 	 * Perform any post parsing finalization and validation
@@ -2296,6 +2311,11 @@ public:
 		return Declarations;
 	}
 
+	/**
+	 * Perform any post parsing finalization that can happen specific to the containing source file
+	 */
+	virtual void ConcurrentPostParseFinalize() override;
+
 protected:
 	explicit FUnrealStructDefinitionInfo(UObject* Object)
 		: FUnrealFieldDefinitionInfo(Object)
