Add a prefix to generated macro names when the project directory starts with something other than an alpha character

#rnx

diff --git a/Engine/Source/Programs/UnrealHeaderTool/Private/UnrealSourceFile.cpp b/Engine/Source/Programs/UnrealHeaderTool/Private/UnrealSourceFile.cpp
--- a/Engine/Source/Programs/UnrealHeaderTool/Private/UnrealSourceFile.cpp
+++ b/Engine/Source/Programs/UnrealHeaderTool/Private/UnrealSourceFile.cpp
@@ -57,7 +57,9 @@ const FString& FUnrealSourceFile::GetFileId() const
 			while (StdFilename.RemoveFromStart(TEXT("../")));
 		}
 
-		FStringOutputDevice Out;
+		// Always prefix the file ID to avoid the issue where directories starting with number would generate an invalid file id.
+		// Source code should use the CURRENT_FILE_ID macro to access this generated file id.
+		FStringOutputDevice Out(TEXT("FID_"));
 
 		for (TCHAR Char : StdFilename)
 		{
