Fix issue where UHT was unable to generate functions with deprecated parameters.

#rnx

diff --git a/Engine/Source/Programs/UnrealHeaderTool/Private/CodeGenerator.cpp b/Engine/Source/Programs/UnrealHeaderTool/Private/CodeGenerator.cpp
--- a/Engine/Source/Programs/UnrealHeaderTool/Private/CodeGenerator.cpp
+++ b/Engine/Source/Programs/UnrealHeaderTool/Private/CodeGenerator.cpp
@@ -2537,7 +2537,7 @@ void FNativeClassHeaderGenerator::ExportInterfaceCallFunctions(FOutputDevice& Ou
 		// code to populate Parms struct
 		for (FUnrealPropertyDefinitionInfo* ParamDef : Parameters.Parms)
 		{
-			const FString ParamName = ParamDef->GetName();
+			const FString ParamName = ParamDef->GetNameWithDeprecated();
 			OutCpp.Logf(TEXT("\t\t\tParms.%s=%s;") LINE_TERMINATOR, *ParamName, *ParamName);
 		}
 
@@ -2548,7 +2548,7 @@ void FNativeClassHeaderGenerator::ExportInterfaceCallFunctions(FOutputDevice& Ou
 		{
 			if(ParamDef->HasAllPropertyFlags(CPF_OutParm) && !ParamDef->HasAnyPropertyFlags(CPF_ConstParm|CPF_ReturnParm))
 			{
-				const FString ParamName = ParamDef->GetName();
+				const FString ParamName = ParamDef->GetNameWithDeprecated();
 				OutCpp.Logf(TEXT("\t\t\t%s=Parms.%s;") LINE_TERMINATOR, *ParamName, *ParamName);
 			}
 		}
@@ -4081,7 +4081,7 @@ void WriteEventFunctionPrologue(FOutputDevice& Output, int32 Indent, const FParm
 	// Declare a parameter struct for this event/delegate and assign the struct members using the values passed into the event/delegate call.
 	for (FUnrealPropertyDefinitionInfo* PropDef : Parameters.Parms)
 	{
-		const FString PropertyName = PropDef->GetName();
+		const FString PropertyName = PropDef->GetNameWithDeprecated();
 		if (PropDef->IsStaticArray())
 		{
 			Output.Logf(TEXT("%sFMemory::Memcpy(Parms.%s,%s,sizeof(Parms.%s));\r\n"), FCString::Tab(Indent + 1), *PropertyName, *PropertyName, *PropertyName);
@@ -4106,7 +4106,7 @@ void WriteEventFunctionEpilogue(FOutputDevice& Output, int32 Indent, const FParm
 	{
 		if (PropDef->HasSpecificPropertyFlags(CPF_OutParm | CPF_ConstParm, CPF_OutParm))
 		{
-			const FString PropertyName = PropDef->GetName();
+			const FString PropertyName = PropDef->GetNameWithDeprecated();
 			if (PropDef->IsStaticArray())
 			{
 				Output.Logf(TEXT("%sFMemory::Memcpy(&%s,&Parms.%s,sizeof(%s));\r\n"), FCString::Tab(Indent + 1), *PropertyName, *PropertyName, *PropertyName);
diff --git a/Engine/Source/Programs/UnrealHeaderTool/Private/HeaderParser.cpp b/Engine/Source/Programs/UnrealHeaderTool/Private/HeaderParser.cpp
--- a/Engine/Source/Programs/UnrealHeaderTool/Private/HeaderParser.cpp
+++ b/Engine/Source/Programs/UnrealHeaderTool/Private/HeaderParser.cpp
@@ -2707,6 +2707,7 @@ void FHeaderParser::CompileDirective()
 
 void FHeaderParser::GetVarType(
 	FScope*                         Scope,
+	bool							bOuterTypeDeprecated,
 	FPropertyBase&                  VarProperty,
 	EPropertyFlags                  Disallow,
 	EUHTPropertyType				OuterPropertyType,
@@ -3453,7 +3454,7 @@ void FHeaderParser::GetVarType(
 	{
 		RequireSymbol( TEXT('<'), TEXT("'tarray'") );
 
-		GetVarType(Scope, VarProperty, Disallow, EUHTPropertyType::DynamicArray, Flags, EPropertyDeclarationStyle::None, VariableCategory);
+		GetVarType(Scope, bOuterTypeDeprecated, VarProperty, Disallow, EUHTPropertyType::DynamicArray, Flags, EPropertyDeclarationStyle::None, VariableCategory);
 		if (VarProperty.IsContainer())
 		{
 			Throwf(TEXT("Nested containers are not supported.") );
@@ -3540,7 +3541,7 @@ void FHeaderParser::GetVarType(
 		RequireSymbol( TEXT('<'), TEXT("'tmap'") );
 
 		FPropertyBase MapKeyType(CPT_None);
-		GetVarType(Scope, MapKeyType, Disallow, EUHTPropertyType::Map, Flags, EPropertyDeclarationStyle::None, VariableCategory);
+		GetVarType(Scope, bOuterTypeDeprecated, MapKeyType, Disallow, EUHTPropertyType::Map, Flags, EPropertyDeclarationStyle::None, VariableCategory);
 		if (MapKeyType.IsContainer())
 		{
 			Throwf(TEXT("Nested containers are not supported.") );
@@ -3568,7 +3569,7 @@ void FHeaderParser::GetVarType(
 			Throwf(TEXT("Missing value type while parsing TMap."));
 		}
 
-		GetVarType(Scope, VarProperty, Disallow, EUHTPropertyType::Map, Flags, EPropertyDeclarationStyle::None, VariableCategory);
+		GetVarType(Scope, bOuterTypeDeprecated, VarProperty, Disallow, EUHTPropertyType::Map, Flags, EPropertyDeclarationStyle::None, VariableCategory);
 		if (VarProperty.IsContainer())
 		{
 			Throwf(TEXT("Nested containers are not supported.") );
@@ -3620,7 +3621,7 @@ void FHeaderParser::GetVarType(
 	{
 		RequireSymbol( TEXT('<'), TEXT("'tset'") );
 
-		GetVarType(Scope, VarProperty, Disallow, EUHTPropertyType::Set, Flags, EPropertyDeclarationStyle::None, VariableCategory);
+		GetVarType(Scope, bOuterTypeDeprecated, VarProperty, Disallow, EUHTPropertyType::Set, Flags, EPropertyDeclarationStyle::None, VariableCategory);
 		if (VarProperty.IsContainer())
 		{
 			Throwf(TEXT("Nested containers are not supported.") );
@@ -4237,7 +4238,7 @@ void FHeaderParser::GetVarType(
 	if (VariableCategory != EVariableCategory::Member)
 	{
 		// These conditions are checked externally for struct/member variables where the flag can be inferred later on from the variable name itself
-		ValidatePropertyIsDeprecatedIfNecessary(VarProperty, OuterPropertyType, OuterPropertyFlags);
+		ValidatePropertyIsDeprecatedIfNecessary(bOuterTypeDeprecated, VarProperty, OuterPropertyType, OuterPropertyFlags);
 	}
 
 	// Check for invalid transients
@@ -5751,13 +5752,15 @@ void FHeaderParser::ParseParameterList(FUnrealFunctionDefinitionInfo& FunctionDe
 		return;
 	}
 
+	const bool bDeprecated = MetaData ? MetaData->Contains(NAME_DeprecatedFunction) : false;
+
 	FAdvancedDisplayParameterHandler AdvancedDisplay(MetaData);
 	do
 	{
 		// Get parameter type.
 		FPropertyBase Property(CPT_None);
 		EVariableCategory VariableCategory = FunctionDef.HasAnyFunctionFlags(FUNC_Net) ? EVariableCategory::ReplicatedParameter : EVariableCategory::RegularParameter;
-		GetVarType(GetCurrentScope(), Property, ~(CPF_ParmFlags | CPF_AutoWeak | CPF_RepSkip | CPF_UObjectWrapper | CPF_NativeAccessSpecifiers), EUHTPropertyType::None, CPF_None, EPropertyDeclarationStyle::None, VariableCategory);
+		GetVarType(GetCurrentScope(), bDeprecated, Property, ~(CPF_ParmFlags | CPF_AutoWeak | CPF_RepSkip | CPF_UObjectWrapper | CPF_NativeAccessSpecifiers), EUHTPropertyType::None, CPF_None, EPropertyDeclarationStyle::None, VariableCategory);
 		Property.PropertyFlags |= CPF_Parm;
 
 		if (bExpectCommaBeforeName)
@@ -5966,6 +5969,8 @@ FUnrealFunctionDefinitionInfo& FHeaderParser::CompileDelegateDeclaration(const F
 		CheckAllow(CurrentScopeName, ENestAllowFlags::ImplicitDelegateDecl);
 	}
 
+	const bool bDeprecated = MetaData.Contains(NAME_DeprecatedFunction); // FBlueprintMetadata::MD_DeprecatedFunction
+
 	// Break the delegate declaration macro down into parts
 	const bool bHasReturnValue = DelegateMacro.Contains(TEXT("_RetVal"), ESearchCase::CaseSensitive);
 	const bool bDeclaredConst  = DelegateMacro.Contains(TEXT("_Const"), ESearchCase::CaseSensitive);
@@ -6012,7 +6017,7 @@ FUnrealFunctionDefinitionInfo& FHeaderParser::CompileDelegateDeclaration(const F
 
 	if (bHasReturnValue)
 	{
-		GetVarType(GetCurrentScope(), ReturnType, CPF_None, EUHTPropertyType::None, CPF_None, EPropertyDeclarationStyle::None, EVariableCategory::Return);
+		GetVarType(GetCurrentScope(), bDeprecated, ReturnType, CPF_None, EUHTPropertyType::None, CPF_None, EPropertyDeclarationStyle::None, EVariableCategory::Return);
 		RequireSymbol(TEXT(','), CurrentScopeName);
 	}
 
@@ -6253,10 +6258,11 @@ FUnrealFunctionDefinitionInfo& FHeaderParser::CompileFunctionDeclaration()
 	FString*   InternalPtr = MetaData.Find(NAME_BlueprintInternalUseOnly); // FBlueprintMetadata::MD_BlueprintInternalUseOnly
 	const bool bInternalOnly = InternalPtr && *InternalPtr == TEXT("true");
 
+	const bool bDeprecated = MetaData.Contains(NAME_DeprecatedFunction); // FBlueprintMetadata::MD_DeprecatedFunction
+
 	// If this function is blueprint callable or blueprint pure, require a category 
 	if ((FuncInfo.FunctionFlags & (FUNC_BlueprintCallable | FUNC_BlueprintPure)) != 0) 
 	{ 
-		const bool bDeprecated = MetaData.Contains(NAME_DeprecatedFunction);       // FBlueprintMetadata::MD_DeprecatedFunction
 		const bool bBlueprintAccessor = MetaData.Contains(NAME_BlueprintSetter) || MetaData.Contains(NAME_BlueprintGetter); // FBlueprintMetadata::MD_BlueprintSetter, // FBlueprintMetadata::MD_BlueprintGetter
 		const bool bHasMenuCategory = MetaData.Contains(NAME_Category);                 // FBlueprintMetadata::MD_FunctionCategory
 
@@ -6393,7 +6399,7 @@ FUnrealFunctionDefinitionInfo& FHeaderParser::CompileFunctionDeclaration()
 	FPropertyBase ReturnType( CPT_None );
 
 	// C++ style functions always have a return value type, even if it's void
-	GetVarType(GetCurrentScope(), ReturnType, CPF_None, EUHTPropertyType::None, CPF_None, EPropertyDeclarationStyle::None, EVariableCategory::Return);
+	GetVarType(GetCurrentScope(), bDeprecated, ReturnType, CPF_None, EUHTPropertyType::None, CPF_None, EPropertyDeclarationStyle::None, EVariableCategory::Return);
 	bool bHasReturnValue = ReturnType.Type != CPT_None;
 
 	// Skip whitespaces to get InputPos exactly on beginning of function name.
@@ -6738,8 +6744,14 @@ bool FHeaderParser::IsBitfieldProperty(ELayoutMacroType LayoutMacroType)
 	return bIsBitfield;
 }
 
-void FHeaderParser::ValidatePropertyIsDeprecatedIfNecessary(const FPropertyBase& VarProperty, EUHTPropertyType OuterPropertyType, EPropertyFlags OuterPropertyFlags)
+void FHeaderParser::ValidatePropertyIsDeprecatedIfNecessary(bool bOuterTypeDeprecated, const FPropertyBase& VarProperty, EUHTPropertyType OuterPropertyType, EPropertyFlags OuterPropertyFlags)
 {
+	// If the outer object being parsed is deprecated,  then we don't need to do any further tests 
+	if (bOuterTypeDeprecated)
+	{
+		return;
+	}
+
 	// If the property is in a container that has been deprecated, then we don't need to do any further tests 
 	if (OuterPropertyType != EUHTPropertyType::None && (OuterPropertyFlags & CPF_Deprecated) != 0)
 	{
@@ -6865,7 +6877,7 @@ void FHeaderParser::CompileVariableDeclaration(FUnrealStructDefinitionInfo& Stru
 	FPropertyBase OriginalPropertyBase(CPT_None);
 	FIndexRange TypeRange;
 	ELayoutMacroType LayoutMacroType = ELayoutMacroType::None;
-	GetVarType(&*StructDef.GetScope(), OriginalPropertyBase, DisallowFlags, EUHTPropertyType::None, CPF_None, EPropertyDeclarationStyle::UPROPERTY, EVariableCategory::Member, &TypeRange, &LayoutMacroType);
+	GetVarType(&*StructDef.GetScope(), false, OriginalPropertyBase, DisallowFlags, EUHTPropertyType::None, CPF_None, EPropertyDeclarationStyle::UPROPERTY, EVariableCategory::Member, &TypeRange, &LayoutMacroType);
 	OriginalPropertyBase.PropertyFlags |= EdFlags;
 
 	FString* Category = OriginalPropertyBase.MetaData.Find(NAME_Category);
@@ -6957,7 +6969,7 @@ void FHeaderParser::CompileVariableDeclaration(FUnrealStructDefinitionInfo& Stru
 		}
 
 		// Deprecation validation
-		ValidatePropertyIsDeprecatedIfNecessary(PropertyBase, EUHTPropertyType::None, CPF_None);
+		ValidatePropertyIsDeprecatedIfNecessary(false, PropertyBase, EUHTPropertyType::None, CPF_None);
 
 		if (TopNest->NestType != ENestType::FunctionDeclaration)
 		{
diff --git a/Engine/Source/Programs/UnrealHeaderTool/Private/HeaderParser.h b/Engine/Source/Programs/UnrealHeaderTool/Private/HeaderParser.h
--- a/Engine/Source/Programs/UnrealHeaderTool/Private/HeaderParser.h
+++ b/Engine/Source/Programs/UnrealHeaderTool/Private/HeaderParser.h
@@ -419,6 +419,7 @@ protected:
 	 * Parses a variable or return value declaration and determines the variable type and property flags.
 	 *
 	 * @param   Scope                     struct to create the property in
+	 * @param   bOuterTypeDeprecated	  true if the type containing the variable/property has been deprecated
 	 * @param   VarProperty               will be filled in with type and property flag data for the property declaration that was parsed
 	 * @param   Disallow                  contains a mask of variable modifiers that are disallowed in this context
 	 * @param   OuterPropertyType         only specified when compiling the inner properties for arrays or maps.
@@ -429,6 +430,7 @@ protected:
 	 */
 	void GetVarType(
 		FScope*							Scope,
+		bool							bOuterTypeDeprecated, 
 		FPropertyBase&                  VarProperty,
 		EPropertyFlags                  Disallow,
 		EUHTPropertyType				OuterPropertyType,
@@ -532,7 +534,7 @@ protected:
 	void InitScriptLocation( FScriptLocation& Retry );
 	void ReturnToLocation( const FScriptLocation& Retry, bool Binary=1, bool Text=1 );
 
-	void ValidatePropertyIsDeprecatedIfNecessary(const FPropertyBase& VarProperty, EUHTPropertyType OuterPropertyType, EPropertyFlags OuterPropertyFlags);
+	void ValidatePropertyIsDeprecatedIfNecessary(bool bOuterTypeDeprecated, const FPropertyBase& VarProperty, EUHTPropertyType OuterPropertyType, EPropertyFlags OuterPropertyFlags);
 
 	// Cache of ScriptStructs that have been validated for Net Replication and RPC
 	TSet<FUnrealScriptStructDefinitionInfo*> ScriptStructsValidForNet;
