Fixed issue where UHT would crash if an exception is thrown in one of the code generation routines.

#rnx

diff --git a/Engine/Source/Programs/UnrealHeaderTool/Private/CodeGenerator.cpp b/Engine/Source/Programs/UnrealHeaderTool/Private/CodeGenerator.cpp
--- a/Engine/Source/Programs/UnrealHeaderTool/Private/CodeGenerator.cpp
+++ b/Engine/Source/Programs/UnrealHeaderTool/Private/CodeGenerator.cpp
@@ -5674,50 +5674,53 @@ void FNativeClassHeaderGenerator::Generate(
 
 		auto ClasssesH = [&PackageDef, &GeneratedFileInfo, &ExportedSorted]()
 		{
-			const FManifestModule& Module = PackageDef.GetModule();
-
-			// Write the classes and enums header prefixes.
-			FUHTStringBuilder ClassesHText;
-			ClassesHText.Log(HeaderCopyright);
-			ClassesHText.Log(TEXT("#pragma once\r\n"));
-			ClassesHText.Log(TEXT("\r\n"));
-			ClassesHText.Log(TEXT("\r\n"));
-
-			// Fill with the rest source files from this package.
-			TSet<FUnrealSourceFile*> PublicHeaderGroupIncludes;
-			for (FGeneratedCPP* GeneratedCPP : ExportedSorted)
+			FResults::Try([&PackageDef, &GeneratedFileInfo, &ExportedSorted]()
 			{
-				if (GeneratedCPP->SourceFile.IsPublic())
+				const FManifestModule& Module = PackageDef.GetModule();
+
+				// Write the classes and enums header prefixes.
+				FUHTStringBuilder ClassesHText;
+				ClassesHText.Log(HeaderCopyright);
+				ClassesHText.Log(TEXT("#pragma once\r\n"));
+				ClassesHText.Log(TEXT("\r\n"));
+				ClassesHText.Log(TEXT("\r\n"));
+
+				// Fill with the rest source files from this package.
+				TSet<FUnrealSourceFile*> PublicHeaderGroupIncludes;
+				for (FGeneratedCPP* GeneratedCPP : ExportedSorted)
 				{
-					PublicHeaderGroupIncludes.Add(&GeneratedCPP->SourceFile);
+					if (GeneratedCPP->SourceFile.IsPublic())
+					{
+						PublicHeaderGroupIncludes.Add(&GeneratedCPP->SourceFile);
+					}
 				}
-			}
-			for (TSharedRef<FUnrealSourceFile>& SourceFile : PackageDef.GetAllSourceFiles())
-			{
-				if (SourceFile->IsPublic())
+				for (TSharedRef<FUnrealSourceFile>& SourceFile : PackageDef.GetAllSourceFiles())
 				{
-					PublicHeaderGroupIncludes.Add(&*SourceFile);
+					if (SourceFile->IsPublic())
+					{
+						PublicHeaderGroupIncludes.Add(&*SourceFile);
+					}
 				}
-			}
 
-			// Make the public header list stable regardless of the order the files are parsed.
-			TArray<FString> BuildPaths;
-			BuildPaths.Reserve(PublicHeaderGroupIncludes.Num());
-			for (FUnrealSourceFile* SourceFile : PublicHeaderGroupIncludes)
-			{
-				BuildPaths.Emplace(GetBuildPath(*SourceFile));
-			}
-			BuildPaths.Sort();
+				// Make the public header list stable regardless of the order the files are parsed.
+				TArray<FString> BuildPaths;
+				BuildPaths.Reserve(PublicHeaderGroupIncludes.Num());
+				for (FUnrealSourceFile* SourceFile : PublicHeaderGroupIncludes)
+				{
+					BuildPaths.Emplace(GetBuildPath(*SourceFile));
+				}
+				BuildPaths.Sort();
 
-			for (const FString& BuildPath : BuildPaths)
-			{
-				ClassesHText.Logf(TEXT("#include \"%s\"") LINE_TERMINATOR, *BuildPath);
-			}
+				for (const FString& BuildPath : BuildPaths)
+				{
+					ClassesHText.Logf(TEXT("#include \"%s\"") LINE_TERMINATOR, *BuildPath);
+				}
 
-			ClassesHText.Log(LINE_TERMINATOR);
+				ClassesHText.Log(LINE_TERMINATOR);
 
-			// Save the classes header if it has changed.
-			SaveHeaderIfChanged(GeneratedFileInfo, MoveTemp(ClassesHText));
+				// Save the classes header if it has changed.
+				SaveHeaderIfChanged(GeneratedFileInfo, MoveTemp(ClassesHText));
+			});
 		};
 
 		TempTasks.Reset();
@@ -5738,37 +5741,40 @@ void FNativeClassHeaderGenerator::Generate(
 
 		auto Functions = [&PackageDef, &GeneratedFileInfo, &ExportedSorted]()
 		{
-			const FManifestModule& Module = PackageDef.GetModule();
-
-			// Export an include line for each header
-			FUHTStringBuilder GeneratedFunctionDeclarations;
-			for (FGeneratedCPP* GeneratedCPP : ExportedSorted)
+			FResults::Try([&PackageDef, &GeneratedFileInfo, &ExportedSorted]()
 			{
-				GeneratedFunctionDeclarations.Log(GeneratedCPP->GeneratedFunctionDeclarations);
-			}
+				const FManifestModule& Module = PackageDef.GetModule();
 
-			if (GeneratedFunctionDeclarations.Len())
-			{
-				FNativeClassHeaderGenerator Generator(PackageDef);
-
-				uint32 CombinedHash = 0;
+				// Export an include line for each header
+				FUHTStringBuilder GeneratedFunctionDeclarations;
 				for (FGeneratedCPP* GeneratedCPP : ExportedSorted)
 				{
-					uint32 SourceHash = GeneratedCPP->Source.GetGeneratedBodyHash();
-					if (CombinedHash == 0)
-					{
-						// Don't combine in the first case because it keeps GUID backwards compatibility
-						CombinedHash = SourceHash;
-					}
-					else
+					GeneratedFunctionDeclarations.Log(GeneratedCPP->GeneratedFunctionDeclarations);
+				}
+
+				if (GeneratedFunctionDeclarations.Len())
+				{
+					FNativeClassHeaderGenerator Generator(PackageDef);
+
+					uint32 CombinedHash = 0;
+					for (FGeneratedCPP* GeneratedCPP : ExportedSorted)
 					{
-						CombinedHash = HashCombine(SourceHash, CombinedHash);
+						uint32 SourceHash = GeneratedCPP->Source.GetGeneratedBodyHash();
+						if (CombinedHash == 0)
+						{
+							// Don't combine in the first case because it keeps GUID backwards compatibility
+							CombinedHash = SourceHash;
+						}
+						else
+						{
+							CombinedHash = HashCombine(SourceHash, CombinedHash);
+						}
 					}
-				}
 
-				Generator.ExportGeneratedPackageInitCode(GeneratedFileInfo.GetGeneratedBody(), *GeneratedFunctionDeclarations, CombinedHash);
-				WriteSource(Module, GeneratedFileInfo, GeneratedFileInfo.GetGeneratedBody(), nullptr, TSet<FString>());
-			}
+					Generator.ExportGeneratedPackageInitCode(GeneratedFileInfo.GetGeneratedBody(), *GeneratedFunctionDeclarations, CombinedHash);
+					WriteSource(Module, GeneratedFileInfo, GeneratedFileInfo.GetGeneratedBody(), nullptr, TSet<FString>());
+				}
+			});
 		};
 
 
