Removed the UScriptStruct _Hash() method from being marked as an exported API method.  Reduces the number of exports.

#rnx

diff --git a/Engine/Source/Programs/UnrealHeaderTool/Private/CodeGenerator.cpp b/Engine/Source/Programs/UnrealHeaderTool/Private/CodeGenerator.cpp
--- a/Engine/Source/Programs/UnrealHeaderTool/Private/CodeGenerator.cpp
+++ b/Engine/Source/Programs/UnrealHeaderTool/Private/CodeGenerator.cpp
@@ -3535,7 +3535,7 @@ void FNativeClassHeaderGenerator::ExportGeneratedStructBodyMacros(FOutputDevice&
 			Out.Logf(TEXT("\t}\r\n"));
 		}
 
-		Out.Logf(TEXT("\textern %suint32 %s();\r\n"), *FriendApiString, *GetHashName);
+		Out.Logf(TEXT("\textern uint32 %s();\r\n"), *GetHashName);
 
 		Out.Logf(TEXT("class UScriptStruct* %s::StaticStruct()\r\n"), *StructNameCPP);
 		Out.Logf(TEXT("{\r\n"));
@@ -3746,7 +3746,7 @@ void FNativeClassHeaderGenerator::ExportGeneratedStructBodyMacros(FOutputDevice&
 	ScriptStructDef.SetHash(StructHash);
 
 	Out.Log(GeneratedStructRegisterFunctionText);
-	Out.Logf(TEXT("\tuint32 %s() { return %uU; }\r\n"), *HashFuncName, StructHash);
+	Out.Logf(TEXT("\tFORCENOINLINE uint32 %s() { return %uU; }\r\n"), *HashFuncName, StructHash);
 
 	// if this struct has RigVM methods we need to implement both the 
 	// virtual function as well as the stub method here.
@@ -4010,7 +4010,7 @@ void FNativeClassHeaderGenerator::ExportGeneratedEnumInitCode(FOutputDevice& Out
 	uint32 EnumHash = GenerateTextHash(*GeneratedEnumRegisterFunctionText);
 	EnumDef.SetHash(EnumHash);
 
-	Out.Logf(TEXT("\tuint32 %s() { return %uU; }\r\n"), *HashFuncName, EnumHash);
+	Out.Logf(TEXT("\tFORCENOINLINE uint32 %s() { return %uU; }\r\n"), *HashFuncName, EnumHash);
 
 	if (bIsDynamic)
 	{
