Fix issue with Placeable/NotPlaceable,  EditInlineNew/NotEditInlineNew, and CollapseCategories/DontCollapseCategories would not properly clear inherited flag from super class.

#rnx

diff --git a/Engine/Source/Programs/UnrealHeaderTool/Private/HeaderParser.cpp b/Engine/Source/Programs/UnrealHeaderTool/Private/HeaderParser.cpp
--- a/Engine/Source/Programs/UnrealHeaderTool/Private/HeaderParser.cpp
+++ b/Engine/Source/Programs/UnrealHeaderTool/Private/HeaderParser.cpp
@@ -5201,7 +5201,7 @@ FUnrealClassDefinitionInfo& FHeaderParser::ParseClassNameDeclaration(FString& De
 			}
 
 			// Propagate the inheritable ClassFlags
-			ClassDef->SetClassFlags(BaseClassDef->GetClassFlags() & CLASS_ScriptInherit);
+			ClassDef->SetClassFlags(BaseClassDef->GetClassFlags() & ClassDef->GetInheritClassFlags());
 		}
 	}
 	return *ClassDef;
@@ -8405,7 +8405,7 @@ void FHeaderParser::ResetClassData()
 
 	if (FUnrealClassDefinitionInfo* SuperClassDef = CurrentClassDef.GetSuperClass())
 	{
-		CurrentClassDef.SetClassFlags(SuperClassDef->GetClassFlags() & CLASS_ScriptInherit);
+		CurrentClassDef.SetClassFlags(SuperClassDef->GetClassFlags() & CurrentClassDef.GetInheritClassFlags());
 		CurrentClassDef.SetClassConfigName(SuperClassDef->GetClassConfigName());
 
 		check(SuperClassDef->GetClassWithin());
diff --git a/Engine/Source/Programs/UnrealHeaderTool/Private/UnrealTypeDefinitionInfo.cpp b/Engine/Source/Programs/UnrealHeaderTool/Private/UnrealTypeDefinitionInfo.cpp
--- a/Engine/Source/Programs/UnrealHeaderTool/Private/UnrealTypeDefinitionInfo.cpp
+++ b/Engine/Source/Programs/UnrealHeaderTool/Private/UnrealTypeDefinitionInfo.cpp
@@ -2268,7 +2268,15 @@ void FUnrealClassDefinitionInfo::PostParseFinalizeInternal(EPostParseFinalizePha
 		// Pick up any newly set flags from the super class (i.e. HasInstancedReference)
 		if (FUnrealClassDefinitionInfo* SuperClassDef = GetSuperClass())
 		{
-			SetClassFlags(SuperClassDef->GetClassFlags() & CLASS_ScriptInherit);
+			if ((InheritClassFlags & CLASS_NotPlaceable) == 0)
+			{
+				if (!SuperClassDef->HasAnyClassFlags(CLASS_NotPlaceable))
+				{
+					Throwf(TEXT("The 'placeable' specifier is only allowed on classes which have a base class that's marked as not placeable. Classes are assumed to be placeable by default."));
+				}
+			}
+
+			SetClassFlags(SuperClassDef->GetClassFlags() & InheritClassFlags);
 		}
 
 		// Check to see if we have any instances referenced
@@ -2401,12 +2409,12 @@ void FUnrealClassDefinitionInfo::ParseClassProperties(TArray<FPropertySpecifier>
 
 			// Class cannot be constructed from the New button in editinline
 			ParsedClassFlags &= ~CLASS_EditInlineNew;
+			InheritClassFlags &= ~CLASS_EditInlineNew;
 			break;
 
 		case EClassMetadataSpecifier::Placeable:
-
-			bWantsToBePlaceable = true;
 			ParsedClassFlags &= ~CLASS_NotPlaceable;
+			InheritClassFlags &= ~CLASS_NotPlaceable;
 			break;
 
 		case EClassMetadataSpecifier::DefaultToInstanced:
@@ -2477,6 +2485,7 @@ void FUnrealClassDefinitionInfo::ParseClassProperties(TArray<FPropertySpecifier>
 
 			// this child of a transient class is not transient - remove the transient flag
 			ParsedClassFlags &= ~CLASS_Transient;
+			InheritClassFlags &= ~CLASS_Transient;
 			break;
 
 		case EClassMetadataSpecifier::CustomConstructor:
@@ -2612,6 +2621,7 @@ void FUnrealClassDefinitionInfo::ParseClassProperties(TArray<FPropertySpecifier>
 
 			// Class' properties should be shown categorized in the editor.
 			ParsedClassFlags &= ~CLASS_CollapseCategories;
+			InheritClassFlags &= ~CLASS_CollapseCategories;
 			break;
 
 		case EClassMetadataSpecifier::AdvancedClassDisplay:
@@ -2720,16 +2730,6 @@ void FUnrealClassDefinitionInfo::MergeClassCategories()
 
 void FUnrealClassDefinitionInfo::MergeAndValidateClassFlags(const FString& DeclaredClassName)
 {
-	if (bWantsToBePlaceable)
-	{
-		if (!HasAnyClassFlags(CLASS_NotPlaceable))
-		{
-			Throwf(TEXT("The 'placeable' specifier is only allowed on classes which have a base class that's marked as not placeable. Classes are assumed to be placeable by default."));
-		}
-		ClearClassFlags(CLASS_NotPlaceable);
-		bWantsToBePlaceable = false; // Reset this flag after it's been merged
-	}
-
 	// Now merge all remaining flags/properties
 	SetClassFlags(ParsedClassFlags);
 	SetClassConfigName(FName(*ConfigName));
diff --git a/Engine/Source/Programs/UnrealHeaderTool/Private/UnrealTypeDefinitionInfo.h b/Engine/Source/Programs/UnrealHeaderTool/Private/UnrealTypeDefinitionInfo.h
--- a/Engine/Source/Programs/UnrealHeaderTool/Private/UnrealTypeDefinitionInfo.h
+++ b/Engine/Source/Programs/UnrealHeaderTool/Private/UnrealTypeDefinitionInfo.h
@@ -3054,6 +3054,14 @@ public:
 		return ParsedClassFlags;
 	}
 
+	/**
+	 * Get the flags we inherit from base classes
+	 */
+	EClassFlags GetInheritClassFlags() const
+	{
+		return InheritClassFlags;
+	}
+
 	/**
 	 * Return the meta data from the preparse phase
 	 */
@@ -3396,6 +3404,7 @@ private:
 	FString ConfigName;
 	EClassFlags ClassFlags = CLASS_None; // Current class flags
 	EClassFlags ParsedClassFlags = CLASS_None; // Class flags from the preparse phase
+	EClassFlags InheritClassFlags = CLASS_ScriptInherit; // Flags to inherit from the super class
 	EClassFlags InitialEngineClassFlags = CLASS_None; // Class flags from a preexisting UObject
 	EClassCastFlags ClassCastFlags = CASTCLASS_None;
 	FUnrealClassDefinitionInfo* ClassWithin = nullptr;
@@ -3432,7 +3441,6 @@ private:
 	bool bCustomVTableHelperConstructorDeclared = false;
 
 	bool bIsInterface = false;
-	bool bWantsToBePlaceable = false;
 };
 
 template <typename To>
