Generate an error when a BlueprintNativeEvent references a const class but the parameter is not marked const.

#rnx

diff --git a/Engine/Source/Programs/UnrealHeaderTool/Private/HeaderParser.cpp b/Engine/Source/Programs/UnrealHeaderTool/Private/HeaderParser.cpp
--- a/Engine/Source/Programs/UnrealHeaderTool/Private/HeaderParser.cpp
+++ b/Engine/Source/Programs/UnrealHeaderTool/Private/HeaderParser.cpp
@@ -2707,7 +2707,7 @@ void FHeaderParser::CompileDirective()
 
 void FHeaderParser::GetVarType(
 	FScope*                         Scope,
-	bool							bOuterTypeDeprecated,
+	EGetVarTypeOptions				Options,
 	FPropertyBase&                  VarProperty,
 	EPropertyFlags                  Disallow,
 	EUHTPropertyType				OuterPropertyType,
@@ -3454,7 +3454,7 @@ void FHeaderParser::GetVarType(
 	{
 		RequireSymbol( TEXT('<'), TEXT("'tarray'") );
 
-		GetVarType(Scope, bOuterTypeDeprecated, VarProperty, Disallow, EUHTPropertyType::DynamicArray, Flags, EPropertyDeclarationStyle::None, VariableCategory);
+		GetVarType(Scope, Options, VarProperty, Disallow, EUHTPropertyType::DynamicArray, Flags, EPropertyDeclarationStyle::None, VariableCategory);
 		if (VarProperty.IsContainer())
 		{
 			Throwf(TEXT("Nested containers are not supported.") );
@@ -3541,7 +3541,7 @@ void FHeaderParser::GetVarType(
 		RequireSymbol( TEXT('<'), TEXT("'tmap'") );
 
 		FPropertyBase MapKeyType(CPT_None);
-		GetVarType(Scope, bOuterTypeDeprecated, MapKeyType, Disallow, EUHTPropertyType::Map, Flags, EPropertyDeclarationStyle::None, VariableCategory);
+		GetVarType(Scope, Options, MapKeyType, Disallow, EUHTPropertyType::Map, Flags, EPropertyDeclarationStyle::None, VariableCategory);
 		if (MapKeyType.IsContainer())
 		{
 			Throwf(TEXT("Nested containers are not supported.") );
@@ -3569,7 +3569,7 @@ void FHeaderParser::GetVarType(
 			Throwf(TEXT("Missing value type while parsing TMap."));
 		}
 
-		GetVarType(Scope, bOuterTypeDeprecated, VarProperty, Disallow, EUHTPropertyType::Map, Flags, EPropertyDeclarationStyle::None, VariableCategory);
+		GetVarType(Scope, Options, VarProperty, Disallow, EUHTPropertyType::Map, Flags, EPropertyDeclarationStyle::None, VariableCategory);
 		if (VarProperty.IsContainer())
 		{
 			Throwf(TEXT("Nested containers are not supported.") );
@@ -3621,7 +3621,7 @@ void FHeaderParser::GetVarType(
 	{
 		RequireSymbol( TEXT('<'), TEXT("'tset'") );
 
-		GetVarType(Scope, bOuterTypeDeprecated, VarProperty, Disallow, EUHTPropertyType::Set, Flags, EPropertyDeclarationStyle::None, VariableCategory);
+		GetVarType(Scope, Options, VarProperty, Disallow, EUHTPropertyType::Set, Flags, EPropertyDeclarationStyle::None, VariableCategory);
 		if (VarProperty.IsContainer())
 		{
 			Throwf(TEXT("Nested containers are not supported.") );
@@ -4059,9 +4059,15 @@ void FHeaderParser::GetVarType(
 					}
 
 					// Imply const if it's a parameter that is a pointer to a const class
+					// NOTE: We shouldn't be automatically adding const param because in some cases with functions and blueprint native event, the 
+					// generated code won't match.  For now, just disabled the auto add in that case and check for the error in the validation code.
+					// Otherwise, the user is not warned and they will get compile errors.
 					if (VariableCategory != EVariableCategory::Member && (TempClassDef != NULL) && (TempClassDef->HasAnyClassFlags(CLASS_Const)))
 					{
-						Flags |= CPF_ConstParm;
+						if (!EnumHasAnyFlags(Options, EGetVarTypeOptions::NoAutoConst))
+						{
+							Flags |= CPF_ConstParm;
+						}
 					}
 
 					// Class keyword in front of a class is legal, we 'consume' it
@@ -4238,7 +4244,7 @@ void FHeaderParser::GetVarType(
 	if (VariableCategory != EVariableCategory::Member)
 	{
 		// These conditions are checked externally for struct/member variables where the flag can be inferred later on from the variable name itself
-		ValidatePropertyIsDeprecatedIfNecessary(bOuterTypeDeprecated, VarProperty, OuterPropertyType, OuterPropertyFlags);
+		ValidatePropertyIsDeprecatedIfNecessary(EnumHasAnyFlags(Options, EGetVarTypeOptions::OuterTypeDeprecated), VarProperty, OuterPropertyType, OuterPropertyFlags);
 	}
 
 	// Check for invalid transients
@@ -5779,7 +5785,7 @@ bool FHeaderParser::IsValidDelegateDeclaration(const FToken& Token) const
 }
 
 // Parse the parameter list of a function or delegate declaration
-void FHeaderParser::ParseParameterList(FUnrealFunctionDefinitionInfo& FunctionDef, bool bExpectCommaBeforeName, TMap<FName, FString>* MetaData)
+void FHeaderParser::ParseParameterList(FUnrealFunctionDefinitionInfo& FunctionDef, bool bExpectCommaBeforeName, TMap<FName, FString>* MetaData, EGetVarTypeOptions Options)
 {
 	// Get parameter list.
 	if (MatchSymbol(TEXT(')')))
@@ -5787,15 +5793,13 @@ void FHeaderParser::ParseParameterList(FUnrealFunctionDefinitionInfo& FunctionDe
 		return;
 	}
 
-	const bool bDeprecated = MetaData ? MetaData->Contains(NAME_DeprecatedFunction) : false;
-
 	FAdvancedDisplayParameterHandler AdvancedDisplay(MetaData);
 	do
 	{
 		// Get parameter type.
 		FPropertyBase Property(CPT_None);
 		EVariableCategory VariableCategory = FunctionDef.HasAnyFunctionFlags(FUNC_Net) ? EVariableCategory::ReplicatedParameter : EVariableCategory::RegularParameter;
-		GetVarType(GetCurrentScope(), bDeprecated, Property, ~(CPF_ParmFlags | CPF_AutoWeak | CPF_RepSkip | CPF_UObjectWrapper | CPF_NativeAccessSpecifiers), EUHTPropertyType::None, CPF_None, EPropertyDeclarationStyle::None, VariableCategory);
+		GetVarType(GetCurrentScope(), Options, Property, ~(CPF_ParmFlags | CPF_AutoWeak | CPF_RepSkip | CPF_UObjectWrapper | CPF_NativeAccessSpecifiers), EUHTPropertyType::None, CPF_None, EPropertyDeclarationStyle::None, VariableCategory);
 		Property.PropertyFlags |= CPF_Parm;
 
 		if (bExpectCommaBeforeName)
@@ -6004,7 +6008,11 @@ FUnrealFunctionDefinitionInfo& FHeaderParser::CompileDelegateDeclaration(const F
 		CheckAllow(CurrentScopeName, ENestAllowFlags::ImplicitDelegateDecl);
 	}
 
-	const bool bDeprecated = MetaData.Contains(NAME_DeprecatedFunction); // FBlueprintMetadata::MD_DeprecatedFunction
+	EGetVarTypeOptions Options = EGetVarTypeOptions::None;
+	if (MetaData.Contains(NAME_DeprecatedFunction))
+	{
+		Options |= EGetVarTypeOptions::OuterTypeDeprecated;
+	}
 
 	// Break the delegate declaration macro down into parts
 	const bool bHasReturnValue = DelegateMacro.Contains(TEXT("_RetVal"), ESearchCase::CaseSensitive);
@@ -6052,7 +6060,7 @@ FUnrealFunctionDefinitionInfo& FHeaderParser::CompileDelegateDeclaration(const F
 
 	if (bHasReturnValue)
 	{
-		GetVarType(GetCurrentScope(), bDeprecated, ReturnType, CPF_None, EUHTPropertyType::None, CPF_None, EPropertyDeclarationStyle::None, EVariableCategory::Return);
+		GetVarType(GetCurrentScope(), Options, ReturnType, CPF_None, EUHTPropertyType::None, CPF_None, EPropertyDeclarationStyle::None, EVariableCategory::Return);
 		RequireSymbol(TEXT(','), CurrentScopeName);
 	}
 
@@ -6430,11 +6438,21 @@ FUnrealFunctionDefinitionInfo& FHeaderParser::CompileFunctionDeclaration()
 		}
 	}
 
+	EGetVarTypeOptions Options = EGetVarTypeOptions::None;
+	if (bDeprecated)
+	{
+		Options |= EGetVarTypeOptions::OuterTypeDeprecated;
+	}
+	if (EnumHasAllFlags(FuncInfo.FunctionFlags, FUNC_BlueprintEvent | FUNC_Native))
+	{
+		Options |= EGetVarTypeOptions::NoAutoConst;
+	}
+
 	// Get return type.
 	FPropertyBase ReturnType( CPT_None );
 
 	// C++ style functions always have a return value type, even if it's void
-	GetVarType(GetCurrentScope(), bDeprecated, ReturnType, CPF_None, EUHTPropertyType::None, CPF_None, EPropertyDeclarationStyle::None, EVariableCategory::Return);
+	GetVarType(GetCurrentScope(), Options, ReturnType, CPF_None, EUHTPropertyType::None, CPF_None, EPropertyDeclarationStyle::None, EVariableCategory::Return);
 	bool bHasReturnValue = ReturnType.Type != CPT_None;
 
 	// Skip whitespaces to get InputPos exactly on beginning of function name.
@@ -6502,7 +6520,7 @@ FUnrealFunctionDefinitionInfo& FHeaderParser::CompileFunctionDeclaration()
 	FUnrealFunctionDefinitionInfo* SuperFuncDef = FuncDef.GetSuperFunction();
 
 	// Get parameter list.
-	ParseParameterList(FuncDef, false, &MetaData);
+	ParseParameterList(FuncDef, false, &MetaData, Options);
 
 	// Get return type, if any.
 	if (bHasReturnValue)
@@ -6912,7 +6930,7 @@ void FHeaderParser::CompileVariableDeclaration(FUnrealStructDefinitionInfo& Stru
 	FPropertyBase OriginalPropertyBase(CPT_None);
 	FIndexRange TypeRange;
 	ELayoutMacroType LayoutMacroType = ELayoutMacroType::None;
-	GetVarType(&*StructDef.GetScope(), false, OriginalPropertyBase, DisallowFlags, EUHTPropertyType::None, CPF_None, EPropertyDeclarationStyle::UPROPERTY, EVariableCategory::Member, &TypeRange, &LayoutMacroType);
+	GetVarType(&*StructDef.GetScope(), EGetVarTypeOptions::None, OriginalPropertyBase, DisallowFlags, EUHTPropertyType::None, CPF_None, EPropertyDeclarationStyle::UPROPERTY, EVariableCategory::Member, &TypeRange, &LayoutMacroType);
 	OriginalPropertyBase.PropertyFlags |= EdFlags;
 
 	FString* Category = OriginalPropertyBase.MetaData.Find(NAME_Category);
@@ -8976,4 +8994,3 @@ bool FRecordTokens::Stop()
 	}
 	return false;
 }
-
diff --git a/Engine/Source/Programs/UnrealHeaderTool/Private/HeaderParser.h b/Engine/Source/Programs/UnrealHeaderTool/Private/HeaderParser.h
--- a/Engine/Source/Programs/UnrealHeaderTool/Private/HeaderParser.h
+++ b/Engine/Source/Programs/UnrealHeaderTool/Private/HeaderParser.h
@@ -53,6 +53,16 @@ enum class ENestAllowFlags
 
 ENUM_CLASS_FLAGS(ENestAllowFlags)
 
+/** Options for GetVarType method */
+enum class EGetVarTypeOptions
+{
+	None = 0,
+	OuterTypeDeprecated = 1,	// Containing type has been deprecated
+	NoAutoConst = 2,			// Don't automatically mark properties as CPF_Const
+};
+
+ENUM_CLASS_FLAGS(EGetVarTypeOptions)
+
 namespace EDelegateSpecifierAction
 {
 	enum Type
@@ -309,7 +319,7 @@ protected:
 	bool IsBitfieldProperty(ELayoutMacroType LayoutMacroType);
 
 	// Parse the parameter list of a function or delegate declaration
-	void ParseParameterList(FUnrealFunctionDefinitionInfo& FunctionDef, bool bExpectCommaBeforeName = false, TMap<FName, FString>* MetaData = NULL);
+	void ParseParameterList(FUnrealFunctionDefinitionInfo& FunctionDef, bool bExpectCommaBeforeName = false, TMap<FName, FString>* MetaData = NULL, EGetVarTypeOptions Options = EGetVarTypeOptions::None);
 
 public:
 	// Throws if a specifier value wasn't provided
@@ -420,7 +430,7 @@ protected:
 	 * Parses a variable or return value declaration and determines the variable type and property flags.
 	 *
 	 * @param   Scope                     struct to create the property in
-	 * @param   bOuterTypeDeprecated	  true if the type containing the variable/property has been deprecated
+	 * @param   Options					  options to control how the variable type is parsed
 	 * @param   VarProperty               will be filled in with type and property flag data for the property declaration that was parsed
 	 * @param   Disallow                  contains a mask of variable modifiers that are disallowed in this context
 	 * @param   OuterPropertyType         only specified when compiling the inner properties for arrays or maps.
@@ -431,7 +441,7 @@ protected:
 	 */
 	void GetVarType(
 		FScope*							Scope,
-		bool							bOuterTypeDeprecated, 
+		EGetVarTypeOptions				Options,
 		FPropertyBase&                  VarProperty,
 		EPropertyFlags                  Disallow,
 		EUHTPropertyType				OuterPropertyType,
diff --git a/Engine/Source/Programs/UnrealHeaderTool/Private/UnrealTypeDefinitionInfo.cpp b/Engine/Source/Programs/UnrealHeaderTool/Private/UnrealTypeDefinitionInfo.cpp
--- a/Engine/Source/Programs/UnrealHeaderTool/Private/UnrealTypeDefinitionInfo.cpp
+++ b/Engine/Source/Programs/UnrealHeaderTool/Private/UnrealTypeDefinitionInfo.cpp
@@ -3019,7 +3019,7 @@ void FUnrealFunctionDefinitionInfo::PostParseFinalizeInternal(EPostParseFinalize
 		// Due to structures that might not be fully parsed at parse time, do blueprint validation here
 		if (GetFunctionType() == EFunctionType::Function && HasAnyFunctionFlags(FUNC_BlueprintCallable | FUNC_BlueprintEvent))
 		{
-			for (TSharedRef<FUnrealPropertyDefinitionInfo> PropertyDef : GetProperties())
+			for (TSharedRef<FUnrealPropertyDefinitionInfo>& PropertyDef : GetProperties())
 			{
 				if (PropertyDef->IsStaticArray())
 				{
@@ -3036,10 +3036,34 @@ void FUnrealFunctionDefinitionInfo::PostParseFinalizeInternal(EPostParseFinalize
 			}
 		}
 
+		// Do a more generic validation of the arguments
+		// 
+		// This code is technically incorrect but here for compatibility reasons.  Container types (TObjectPtr, TArray, etc...) should have the 
+		// const on the inside of the template arguments but it places the const on the outside which isn't correct.  This needs to be addressed
+		// in any updates to the type parser.  See EGetVarTypeOptions::NoAutoConst for the parser side of the problem.
+		for (TSharedRef<FUnrealPropertyDefinitionInfo>& PropertyDef : GetProperties())
+		{
+			FPropertyBase& PropertyBase = PropertyDef->GetPropertyBase();
+			if (FUnrealClassDefinitionInfo* ClassDef = UHTCast<FUnrealClassDefinitionInfo>(PropertyBase.TypeDef))
+			{
+				if (ClassDef->HasAnyClassFlags(CLASS_Const) && !PropertyDef->HasAnyPropertyFlags(CPF_ConstParm))
+				{
+					if (PropertyDef->HasAnyPropertyFlags(CPF_ReturnParm))
+					{
+						LogError(TEXT("Return value must be 'const' since '%s' is marked 'const'"), *ClassDef->GetNameCPP());
+					}
+					else
+					{
+						LogError(TEXT("Argument '%s' must be 'const' since '%s' is marked 'const'"), *PropertyDef->GetName(), *ClassDef->GetNameCPP());
+					}
+				}
+			}
+		}
+
 		// The following code is only performed on functions in a class.  
 		if (UHTCast<FUnrealClassDefinitionInfo>(GetOuter()) != nullptr)
 		{
-			for (TSharedRef<FUnrealPropertyDefinitionInfo> PropertyDef : GetProperties())
+			for (TSharedRef<FUnrealPropertyDefinitionInfo>& PropertyDef : GetProperties())
 			{
 				if (PropertyDef->HasSpecificPropertyFlags(CPF_ReturnParm | CPF_OutParm, CPF_OutParm))
 				{
@@ -3048,8 +3072,7 @@ void FUnrealFunctionDefinitionInfo::PostParseFinalizeInternal(EPostParseFinalize
 
 				if (PropertyDef->IsStructOrStructStaticArray())
 				{
-					FUnrealScriptStructDefinitionInfo& StructDef = UHTCastChecked<FUnrealScriptStructDefinitionInfo>(PropertyDef->GetPropertyBase().ClassDef);
-					if (StructDef.HasDefaults())
+					if (PropertyDef->GetPropertyBase().ScriptStructDef->HasDefaults())
 					{
 						SetFunctionFlags(FUNC_HasDefaults);
 					}
