Remove the usage of CLASS_Parsed and CLASS_CustomConstructor.

#rnx

diff --git a/Engine/Source/Programs/UnrealHeaderTool/Private/CodeGenerator.cpp b/Engine/Source/Programs/UnrealHeaderTool/Private/CodeGenerator.cpp
--- a/Engine/Source/Programs/UnrealHeaderTool/Private/CodeGenerator.cpp
+++ b/Engine/Source/Programs/UnrealHeaderTool/Private/CodeGenerator.cpp
@@ -3018,7 +3018,7 @@ void ExportVTableHelperCtorAndCaller(FOutputDevice& Out, FUnrealClassDefinitionI
  */
 void ExportStandardConstructorsMacro(FOutputDevice& Out, FUnrealClassDefinitionInfo& ClassDef, const TCHAR* API, const TCHAR* ClassCPPName)
 {
-	if (!ClassDef.HasAnyClassFlags(CLASS_CustomConstructor))
+	if (!ClassDef.HasCustomConstructor())
 	{
 		Out.Logf(TEXT("\t/** Standard constructor, called after all reflected properties have been initialized */\r\n"));
 		Out.Logf(TEXT("\t%s_API %s(const FObjectInitializer& ObjectInitializer%s);\r\n"), API, ClassCPPName,
diff --git a/Engine/Source/Programs/UnrealHeaderTool/Private/HeaderParser.cpp b/Engine/Source/Programs/UnrealHeaderTool/Private/HeaderParser.cpp
--- a/Engine/Source/Programs/UnrealHeaderTool/Private/HeaderParser.cpp
+++ b/Engine/Source/Programs/UnrealHeaderTool/Private/HeaderParser.cpp
@@ -5287,7 +5287,7 @@ FUnrealClassDefinitionInfo& FHeaderParser::CompileClassDeclaration()
 	// If we have existing class flags (preexisting engine objects), make sure we have some shared flags.
 	check(ClassDef.GetClassFlags() == 0 || (ClassDef.GetClassFlags() & ClassDef.GetParsedClassFlags()) != 0);
 
-	ClassDef.SetClassFlags(CLASS_Parsed);
+	ClassDef.MarkParsed();
 
 	PushNest(ENestType::Class, &ClassDef);
 	
@@ -5296,7 +5296,7 @@ FUnrealClassDefinitionInfo& FHeaderParser::CompileClassDeclaration()
 	// Make sure our parent classes is parsed.
 	for (FUnrealClassDefinitionInfo* TempDef = ClassDef.GetSuperClass(); TempDef; TempDef = TempDef->GetSuperClass())
 	{
-		bool bIsParsed = TempDef->HasAnyClassFlags(CLASS_Parsed);
+		bool bIsParsed = TempDef->IsParsed();
 		bool bIsIntrinsic = TempDef->HasAnyClassFlags(CLASS_Intrinsic);
 		if (!(bIsParsed || bIsIntrinsic))
 		{
@@ -7244,7 +7244,7 @@ void FHeaderParser::ParseHeader()
 			FUnrealClassDefinitionInfo& ClassDef = ClassDataPair->AsClassChecked();
 			for (FUnrealClassDefinitionInfo* ParentClassDef = ClassDef.GetSuperClass(); ParentClassDef; ParentClassDef = ParentClassDef->GetSuperClass())
 			{
-				if (ParentClassDef->HasAnyClassFlags(CLASS_Parsed | CLASS_Intrinsic))
+				if (ParentClassDef->IsParsed() || ParentClassDef->HasAnyClassFlags(CLASS_Intrinsic))
 				{
 					break;
 				}
@@ -7357,7 +7357,7 @@ void FHeaderParser::ParseHeader()
 
 	for (const TSharedRef<FUnrealTypeDefinitionInfo>& ClassDef : SourceFile.GetDefinedClasses())
 	{
-		ClassDef->AsClassChecked().SetClassFlags(CLASS_Parsed);
+		ClassDef->AsClassChecked().MarkParsed();
 	}
 
 	// Perform any final concurrent finalization
diff --git a/Engine/Source/Programs/UnrealHeaderTool/Private/UnrealTypeDefinitionInfo.cpp b/Engine/Source/Programs/UnrealHeaderTool/Private/UnrealTypeDefinitionInfo.cpp
--- a/Engine/Source/Programs/UnrealHeaderTool/Private/UnrealTypeDefinitionInfo.cpp
+++ b/Engine/Source/Programs/UnrealHeaderTool/Private/UnrealTypeDefinitionInfo.cpp
@@ -2209,22 +2209,6 @@ void FUnrealClassDefinitionInfo::CreateUObjectEngineTypesInternal(ECreateEngineT
 			UFunction* Function = FunctionDef->GetFunction();
 			Class->AddFunctionToFunctionMap(Function, Function->GetFName());
 		}
-
-
-		if (!HasAnyClassFlags(CLASS_Native))
-		{
-			//Class->UnMark(EObjectMark(OBJECTMARK_TagImp | OBJECTMARK_TagExp));
-		}
-		else if (!HasAnyClassFlags(CLASS_NoExport | CLASS_Intrinsic))
-		{
-			//Class->UnMark(OBJECTMARK_TagImp);
-			//Class->Mark(OBJECTMARK_TagExp);
-		}
-
-		// This needs to be done outside of parallel blocks because it will modify UClass memory.
-		// Later calls to SetUpUhtReplicationData inside parallel blocks should be fine, because
-		// they will see the memory has already been set up, and just return the parent pointer.
-		Class->SetUpUhtReplicationData();
 		break;
 	}
 
@@ -2491,7 +2475,7 @@ void FUnrealClassDefinitionInfo::ParseClassProperties(TArray<FPropertySpecifier>
 		case EClassMetadataSpecifier::CustomConstructor:
 
 			// we will not export a constructor for this class, assuming it is in the CPP block
-			ParsedClassFlags |= CLASS_CustomConstructor;
+			MarkCustomConstructor();
 			break;
 
 		case EClassMetadataSpecifier::Config:
diff --git a/Engine/Source/Programs/UnrealHeaderTool/Private/UnrealTypeDefinitionInfo.h b/Engine/Source/Programs/UnrealHeaderTool/Private/UnrealTypeDefinitionInfo.h
--- a/Engine/Source/Programs/UnrealHeaderTool/Private/UnrealTypeDefinitionInfo.h
+++ b/Engine/Source/Programs/UnrealHeaderTool/Private/UnrealTypeDefinitionInfo.h
@@ -3321,6 +3321,38 @@ public:
 		bCustomVTableHelperConstructorDeclared = true;
 	}
 
+	/**
+	 * Return true if the class has been parsed
+	 */
+	bool IsParsed() const
+	{
+		return bParsed;
+	}
+
+	/**
+	 * Mark the class as being parsed
+	 */
+	void MarkParsed()
+	{
+		bParsed = true;
+	}
+
+	/**
+	 * Return true if the class has a custom constructor
+	 */
+	bool HasCustomConstructor() const
+	{
+		return bCustomConstructor;
+	}
+
+	/**
+	 * Mark the class as having a custom constructor
+	 */
+	void MarkCustomConstructor()
+	{
+		bCustomConstructor = true;
+	}
+
 	/**
 	 * Get the generated body access specifier
 	 */
@@ -3434,12 +3466,19 @@ private:
 	/** Is default constructor declared? */
 	bool bDefaultConstructorDeclared = false;
 
-	/* Is ObjectInitializer constructor (i.e. a constructor with only one parameter of type FObjectInitializer) declared? */
+	/** Is ObjectInitializer constructor (i.e. a constructor with only one parameter of type FObjectInitializer) declared? */
 	bool bObjectInitializerConstructorDeclared = false;
 
-	/* Is custom VTable helper constructor declared? */
+	/** Is custom VTable helper constructor declared? */
 	bool bCustomVTableHelperConstructorDeclared = false;
 
+	/** True if the class definition has been parsed */
+	bool bParsed = false;
+
+	/** True if the class has a custom constructor */
+	bool bCustomConstructor = false;
+
+	/** True if the class is an interface */
 	bool bIsInterface = false;
 };
 
