Control Rig: RigVM Property Based storage support for sliced memory

diff --git a/Engine/Source/Developer/RigVMDeveloper/Private/RigVMCompiler/RigVMCompiler.cpp b/Engine/Source/Developer/RigVMDeveloper/Private/RigVMCompiler/RigVMCompiler.cpp
--- a/Engine/Source/Developer/RigVMDeveloper/Private/RigVMCompiler/RigVMCompiler.cpp
+++ b/Engine/Source/Developer/RigVMDeveloper/Private/RigVMCompiler/RigVMCompiler.cpp
@@ -2354,6 +2354,17 @@ FRigVMOperand URigVMCompiler::FindOrAddRegister(const FRigVMVarExprAST* InVarExp
 			// debug values are always stored as arrays
 			CPPType = FString::Printf(TEXT("TArray<%s>"), *CPPType);
 		}
+		else if(Pin->GetDirection() == ERigVMPinDirection::Hidden && Pin->GetNode()->IsA<URigVMUnitNode>())
+		{
+			UScriptStruct* UnitStruct = Cast<URigVMUnitNode>(Pin->GetNode())->GetScriptStruct();
+			const FProperty* Property = UnitStruct->FindPropertyByName(Pin->GetFName());
+			check(Property);
+
+			if (!Property->HasMetaData(FRigVMStruct::SingletonMetaName))
+			{
+				CPPType = FString::Printf(TEXT("TArray<%s>"), *CPPType);
+			}
+		}
 		Operand = WorkData.AddProperty(MemoryType, RegisterName, CPPType, Pin->GetCPPTypeObject(), JoinedDefaultValue);
 		
 #endif
diff --git a/Engine/Source/Developer/RigVMDeveloper/Private/RigVMModel/RigVMPin.cpp b/Engine/Source/Developer/RigVMDeveloper/Private/RigVMModel/RigVMPin.cpp
--- a/Engine/Source/Developer/RigVMDeveloper/Private/RigVMModel/RigVMPin.cpp
+++ b/Engine/Source/Developer/RigVMDeveloper/Private/RigVMModel/RigVMPin.cpp
@@ -1257,9 +1257,20 @@ bool URigVMPin::CanLink(URigVMPin* InSourcePin, URigVMPin* InTargetPin, FString*
 		bool bCPPTypesDiffer = true;
 		static const FString Float = TEXT("float");
 		static const FString Double = TEXT("double");
+#if UE_RIGVM_UCLASS_BASED_STORAGE_DISABLED
 
 		if((InSourcePin->CPPType == Float && InTargetPin->CPPType == Double) ||
 			(InSourcePin->CPPType == Double && InTargetPin->CPPType == Float))
+#else
+
+		static const FString FloatArray = TEXT("TArray<float>");
+		static const FString DoubleArray = TEXT("TArray<double>");
+		
+		if((InSourcePin->CPPType == Float && InTargetPin->CPPType == Double) ||
+			(InSourcePin->CPPType == Double && InTargetPin->CPPType == Float) ||
+			(InSourcePin->CPPType == FloatArray && InTargetPin->CPPType == DoubleArray) ||
+			(InSourcePin->CPPType == DoubleArray && InTargetPin->CPPType == FloatArray))
+#endif
 		{
 			bCPPTypesDiffer = false;
 		}
diff --git a/Engine/Source/Programs/UnrealHeaderTool/Private/CodeGenerator.cpp b/Engine/Source/Programs/UnrealHeaderTool/Private/CodeGenerator.cpp
--- a/Engine/Source/Programs/UnrealHeaderTool/Private/CodeGenerator.cpp
+++ b/Engine/Source/Programs/UnrealHeaderTool/Private/CodeGenerator.cpp
@@ -3420,45 +3420,67 @@ void FNativeClassHeaderGenerator::ExportGeneratedStructBodyMacros(FOutputDevice&
 
 				OperandIndex++;
 			}
+			else
+			{
+				FString VariableType = Parameter.TypeVariableRef(true);
+				FString ExtractedType = Parameter.TypeOriginal();
+				FString ParameterCast = FString::Printf(TEXT("*(%s*)"), *ExtractedType);
+
+				// if the parameter is a const enum we need to cast it slightly differently,
+				// we'll get the reference of the stored uint8 and cast it by value.
+				if (Parameter.bIsEnum && !Parameter.bOutput)
+				{
+					VariableType = Parameter.TypeOriginal();
+					ParameterCast = FString::Printf(TEXT("(%s)*(uint8*)"), *ExtractedType);
+				}
+
+				RigVMStubProlog.Add(FString::Printf(TEXT("%s %s = %sRigVMMemoryHandles[%d].GetData();"),
+				*VariableType,
+				*ParamNameOriginal,
+					*ParameterCast,
+					OperandIndex));
+
+				OperandIndex++;
+			}
 
 #else
 
+			FString AdditionalParameters;
+			if(!Parameter.bInput && !Parameter.bOutput && !Parameter.bSingleton)
+			{
+				static const FString SliceContextParameter = TEXT(", RigVMExecuteContext.GetSlice().GetIndex()");
+				AdditionalParameters = SliceContextParameter;
+			}
+
 			if (Parameter.IsArray())
 			{
 				FString ExtendedType = Parameter.ExtendedType();
-				RigVMStubProlog.Add(FString::Printf(TEXT("TArray%s& %s = *(TArray%s*)RigVMMemoryHandles[%d].GetData();"),
+				RigVMStubProlog.Add(FString::Printf(TEXT("TArray%s& %s = *(TArray%s*)RigVMMemoryHandles[%d].GetData(false%s);"),
 					*ExtendedType,
 					*ParamNameOriginal,
 					*ExtendedType,
-					OperandIndex));
+					OperandIndex,
+					*AdditionalParameters));
 
 				OperandIndex++;
 			}
-			
-#endif
-			
 			else
 			{
 				FString VariableType = Parameter.TypeVariableRef(true);
 				FString ExtractedType = Parameter.TypeOriginal();
 				FString ParameterCast = FString::Printf(TEXT("*(%s*)"), *ExtractedType);
 
-				// if the parameter is a const enum we need to cast it slightly differently,
-				// we'll get the reference of the stored uint8 and cast it by value.
-				if (Parameter.bIsEnum && !Parameter.bOutput)
-				{
-					VariableType = Parameter.TypeOriginal();
-					ParameterCast = FString::Printf(TEXT("(%s)*(uint8*)"), *ExtractedType);
-				}
-
-				RigVMStubProlog.Add(FString::Printf(TEXT("%s %s = %sRigVMMemoryHandles[%d].GetData();"),
+				RigVMStubProlog.Add(FString::Printf(TEXT("%s %s = %sRigVMMemoryHandles[%d].GetData(false%s);"),
 				*VariableType,
 				*ParamNameOriginal,
-					*ParameterCast,
-					OperandIndex));
+				*ParameterCast,
+				OperandIndex,
+				*AdditionalParameters));
 
 				OperandIndex++;
 			}
+			
+#endif
 		}
 
 		FString StructMembers = StructRigVMInfo.Members.Declarations(false, TEXT(", \\\r\n\t\t"), true, false);
diff --git a/Engine/Source/Runtime/RigVM/Private/RigVMCore/RigVMMemoryStorage.cpp b/Engine/Source/Runtime/RigVM/Private/RigVMCore/RigVMMemoryStorage.cpp
--- a/Engine/Source/Runtime/RigVM/Private/RigVMCore/RigVMMemoryStorage.cpp
+++ b/Engine/Source/Runtime/RigVM/Private/RigVMCore/RigVMMemoryStorage.cpp
@@ -577,6 +577,43 @@ bool URigVMMemoryStorage::CopyProperty(
 				return true;
 			}
 		}
+		else if(const FArrayProperty* TargetArrayProperty = CastField<FArrayProperty>(InTargetProperty))
+		{
+			if(const FArrayProperty* SourceArrayProperty = CastField<FArrayProperty>(InSourceProperty))
+			{
+				FScriptArrayHelper TargetArray(TargetArrayProperty, InTargetPtr);
+				FScriptArrayHelper SourceArray(SourceArrayProperty, InSourcePtr);
+				
+				if(TargetArrayProperty->IsA<FFloatProperty>())
+				{
+					if(SourceArrayProperty->IsA<FDoubleProperty>())
+					{
+						TargetArray.Resize(SourceArray.Num());
+						for(int32 Index=0;Index<TargetArray.Num();Index++)
+						{
+							float* TargetFloat = (float*)TargetArray.GetRawPtr(Index);
+							const double* SourceDouble = (const double*)SourceArray.GetRawPtr(Index);
+							*TargetFloat = (float)*SourceDouble;
+						}
+						return true;
+					}
+				}
+				else if(TargetArrayProperty->IsA<FDoubleProperty>())
+				{
+					if(SourceArrayProperty->IsA<FFloatProperty>())
+					{
+						TargetArray.Resize(SourceArray.Num());
+						for(int32 Index=0;Index<TargetArray.Num();Index++)
+						{
+							double* TargetDouble = (double*)TargetArray.GetRawPtr(Index);
+							const float* SourceFloat = (const float*)SourceArray.GetRawPtr(Index);
+							*TargetDouble = (double)*SourceFloat;
+						}
+						return true;
+					}
+				}
+			}
+		}
 		
 		ensure(InTargetProperty->SameType(InSourceProperty));
 		return false;
diff --git a/Engine/Source/Runtime/RigVM/Public/RigVMCore/RigVMMemoryStorage.h b/Engine/Source/Runtime/RigVM/Public/RigVMCore/RigVMMemoryStorage.h
--- a/Engine/Source/Runtime/RigVM/Public/RigVMCore/RigVMMemoryStorage.h
+++ b/Engine/Source/Runtime/RigVM/Public/RigVMCore/RigVMMemoryStorage.h
@@ -34,14 +34,14 @@ public:
 	{
 	}
 
-	FORCEINLINE uint8* GetData(bool bFollowPropertyPath = false)
+	FORCEINLINE uint8* GetData(bool bFollowPropertyPath = false, int32 InSliceIndex = INDEX_NONE)
 	{
-		return GetData_Internal(bFollowPropertyPath);
+		return GetData_Internal(bFollowPropertyPath, InSliceIndex);
 	}
 
-	FORCEINLINE const uint8* GetData(bool bFollowPropertyPath = false) const
+	FORCEINLINE const uint8* GetData(bool bFollowPropertyPath = false, int32 InSliceIndex = INDEX_NONE) const
 	{
-		return GetData_Internal(bFollowPropertyPath);
+		return GetData_Internal(bFollowPropertyPath, InSliceIndex);
 	}
 
 	FORCEINLINE const FProperty* GetProperty() const
@@ -65,8 +65,27 @@ public:
 
 private:
 
-	FORCEINLINE_DEBUGGABLE uint8* GetData_Internal(bool bFollowPropertyPath = false) const
+	FORCEINLINE_DEBUGGABLE uint8* GetData_Internal(bool bFollowPropertyPath, int32 InSliceIndex) const
 	{
+		if(InSliceIndex != INDEX_NONE)
+		{
+			// sliced memory cannot be accessed
+			// using a property path.
+			// it refers to opaque memory only
+			check(PropertyPath == nullptr);
+			check(!bFollowPropertyPath);
+
+			const FArrayProperty* ArrayProperty = CastFieldChecked<FArrayProperty>(Property);
+			FScriptArrayHelper ArrayHelper(ArrayProperty, Ptr);
+			if(ArrayHelper.Num() <= InSliceIndex)
+			{
+				const int32 NumValuesToAdd = 1 + InSliceIndex - ArrayHelper.Num();
+				ArrayHelper.AddValues(NumValuesToAdd);
+			}
+
+			return ArrayHelper.GetRawPtr(InSliceIndex);
+		}
+		
 		if(bFollowPropertyPath && PropertyPath != nullptr)
 		{
 			return PropertyPath->GetData<uint8>(Ptr, Property);
