Improved the error message when a deprecated type is found in a non-deprecated function

#rnx

diff --git a/Engine/Source/Programs/UnrealHeaderTool/Private/HeaderParser.cpp b/Engine/Source/Programs/UnrealHeaderTool/Private/HeaderParser.cpp
--- a/Engine/Source/Programs/UnrealHeaderTool/Private/HeaderParser.cpp
+++ b/Engine/Source/Programs/UnrealHeaderTool/Private/HeaderParser.cpp
@@ -4244,7 +4244,7 @@ void FHeaderParser::GetVarType(
 	if (VariableCategory != EVariableCategory::Member)
 	{
 		// These conditions are checked externally for struct/member variables where the flag can be inferred later on from the variable name itself
-		ValidatePropertyIsDeprecatedIfNecessary(EnumHasAnyFlags(Options, EGetVarTypeOptions::OuterTypeDeprecated), VarProperty, OuterPropertyType, OuterPropertyFlags);
+		ValidatePropertyIsDeprecatedIfNecessary(EnumHasAnyFlags(Options, EGetVarTypeOptions::OuterTypeDeprecated), VariableCategory, VarProperty, OuterPropertyType, OuterPropertyFlags);
 	}
 
 	// Check for invalid transients
@@ -6797,7 +6797,23 @@ bool FHeaderParser::IsBitfieldProperty(ELayoutMacroType LayoutMacroType)
 	return bIsBitfield;
 }
 
-void FHeaderParser::ValidatePropertyIsDeprecatedIfNecessary(bool bOuterTypeDeprecated, const FPropertyBase& VarProperty, EUHTPropertyType OuterPropertyType, EPropertyFlags OuterPropertyFlags)
+void FHeaderParser::ValidateTypeIsDeprecated(EVariableCategory VariableCategory, FUnrealTypeDefinitionInfo* TypeDef)
+{
+	if (FUnrealClassDefinitionInfo* ClassDef = UHTCast<FUnrealClassDefinitionInfo>(TypeDef); ClassDef != nullptr && ClassDef->HasAnyClassFlags(CLASS_Deprecated))
+	{
+		if (VariableCategory == EVariableCategory::Member)
+		{
+			LogError(TEXT("Property is using a deprecated class: %s.  Property should be marked deprecated as well."), *ClassDef->GetPathName());
+		}
+		else
+		{
+			LogError(TEXT("Function is using a deprecated class: %s.  Function should be marked deprecated as well."), *ClassDef->GetPathName());
+		}
+	}
+}
+
+
+void FHeaderParser::ValidatePropertyIsDeprecatedIfNecessary(bool bOuterTypeDeprecated, EVariableCategory VariableCategory, const FPropertyBase& VarProperty, EUHTPropertyType OuterPropertyType, EPropertyFlags OuterPropertyFlags)
 {
 	// If the outer object being parsed is deprecated,  then we don't need to do any further tests 
 	if (bOuterTypeDeprecated)
@@ -6817,22 +6833,8 @@ void FHeaderParser::ValidatePropertyIsDeprecatedIfNecessary(bool bOuterTypeDepre
 		return;
 	}
 
-	// check to see if we have a FClassProperty using a deprecated class
-	if (VarProperty.MetaClassDef != nullptr && VarProperty.MetaClassDef->HasAnyClassFlags(CLASS_Deprecated))
-	{
-		LogError(TEXT("Property is using a deprecated class: %s.  Property should be marked deprecated as well."), *VarProperty.MetaClassDef->GetPathName());
-	}
-
-	// check to see if we have a FObjectProperty using a deprecated class.
-	// PropertyClass is part of a union, so only check PropertyClass if this token represents an object property
-	if ((VarProperty.Type == CPT_ObjectReference || VarProperty.Type == CPT_WeakObjectReference || VarProperty.Type == CPT_LazyObjectReference || VarProperty.Type == CPT_ObjectPtrReference || VarProperty.Type == CPT_SoftObjectReference)
-		&& VarProperty.ClassDef != nullptr)
-	{
-		if (VarProperty.ClassDef->HasAnyClassFlags(CLASS_Deprecated))	// and the object class being used has been deprecated
-		{
-			LogError(TEXT("Property is using a deprecated class: %s.  Property should be marked deprecated as well."), *VarProperty.ClassDef->GetPathName());
-		}
-	}
+	ValidateTypeIsDeprecated(VariableCategory, VarProperty.MetaClassDef);
+	ValidateTypeIsDeprecated(VariableCategory, VarProperty.TypeDef);
 }
 
 bool FHeaderParser::ValidateScriptStructOkForNet(const FString& OriginStructName, FUnrealScriptStructDefinitionInfo& InStructDef)
@@ -7022,7 +7024,7 @@ void FHeaderParser::CompileVariableDeclaration(FUnrealStructDefinitionInfo& Stru
 		}
 
 		// Deprecation validation
-		ValidatePropertyIsDeprecatedIfNecessary(false, PropertyBase, EUHTPropertyType::None, CPF_None);
+		ValidatePropertyIsDeprecatedIfNecessary(false, EVariableCategory::Member, PropertyBase, EUHTPropertyType::None, CPF_None);
 
 		if (TopNest->NestType != ENestType::FunctionDeclaration)
 		{
diff --git a/Engine/Source/Programs/UnrealHeaderTool/Private/HeaderParser.h b/Engine/Source/Programs/UnrealHeaderTool/Private/HeaderParser.h
--- a/Engine/Source/Programs/UnrealHeaderTool/Private/HeaderParser.h
+++ b/Engine/Source/Programs/UnrealHeaderTool/Private/HeaderParser.h
@@ -545,7 +545,8 @@ protected:
 	void InitScriptLocation( FScriptLocation& Retry );
 	void ReturnToLocation( const FScriptLocation& Retry, bool Binary=1, bool Text=1 );
 
-	void ValidatePropertyIsDeprecatedIfNecessary(bool bOuterTypeDeprecated, const FPropertyBase& VarProperty, EUHTPropertyType OuterPropertyType, EPropertyFlags OuterPropertyFlags);
+	void ValidateTypeIsDeprecated(EVariableCategory VariableCategory, FUnrealTypeDefinitionInfo* TypeDef);
+	void ValidatePropertyIsDeprecatedIfNecessary(bool bOuterTypeDeprecated, EVariableCategory VariableCategory, const FPropertyBase& VarProperty, EUHTPropertyType OuterPropertyType, EPropertyFlags OuterPropertyFlags);
 
 	// Cache of ScriptStructs that have been validated for Net Replication and RPC
 	TSet<FUnrealScriptStructDefinitionInfo*> ScriptStructsValidForNet;
