Added Internal directory support to UHT

#rnx

diff --git a/Engine/Source/Programs/UnrealHeaderTool/Private/CodeGenerator.cpp b/Engine/Source/Programs/UnrealHeaderTool/Private/CodeGenerator.cpp
--- a/Engine/Source/Programs/UnrealHeaderTool/Private/CodeGenerator.cpp
+++ b/Engine/Source/Programs/UnrealHeaderTool/Private/CodeGenerator.cpp
@@ -6372,7 +6372,8 @@ void PrepareModules(TArray<FUnrealPackageDefinitionInfo*>& PackageDefs, const FS
 	enum EHeaderFolderTypes
 	{
 		PublicClassesHeaders = 0,
-		PublicHeaders = 1,
+		PublicHeaders,
+		InternalHeaders,
 		PrivateHeaders,
 
 		FolderType_Count
@@ -6393,7 +6394,7 @@ void PrepareModules(TArray<FUnrealPackageDefinitionInfo*>& PackageDefs, const FS
 		PackageDefs.Add(&PackageDef);
 
 		TArray<TSharedRef<FUnrealSourceFile>>& AllSourceFiles = PackageDef.GetAllSourceFiles();
-		AllSourceFiles.Reserve(Module.PublicUObjectClassesHeaders.Num() + Module.PublicUObjectHeaders.Num() + Module.PrivateUObjectHeaders.Num());
+		AllSourceFiles.Reserve(Module.PublicUObjectClassesHeaders.Num() + Module.PublicUObjectHeaders.Num() + Module.InternalUObjectHeaders.Num() + Module.PrivateUObjectHeaders.Num());
 
 		// Initialize the other header data structures and create the unreal source files
 		for (int32 PassIndex = 0; PassIndex < FolderType_Count && FResults::IsSucceeding(); ++PassIndex)
@@ -6403,6 +6404,7 @@ void PrepareModules(TArray<FUnrealPackageDefinitionInfo*>& PackageDefs, const FS
 			const TArray<FString>& UObjectHeaders =
 				(CurrentlyProcessing == PublicClassesHeaders) ? Module.PublicUObjectClassesHeaders :
 				(CurrentlyProcessing == PublicHeaders) ? Module.PublicUObjectHeaders :
+				(CurrentlyProcessing == InternalHeaders) ? Module.InternalUObjectHeaders :
 				Module.PrivateUObjectHeaders;
 
 			if (UObjectHeaders.Num() == 0)
@@ -7090,7 +7092,7 @@ ECompilationResult::Type UnrealHeaderTool_Main(const FString& ModuleInfoFilename
 	int NumSources = 0;
 	for (const FManifestModule& Module : GManifest.Modules)
 	{
-		NumSources += Module.PublicUObjectClassesHeaders.Num() + Module.PublicUObjectHeaders.Num() + Module.PrivateUObjectHeaders.Num();
+		NumSources += Module.PublicUObjectClassesHeaders.Num() + Module.PublicUObjectHeaders.Num() + Module.InternalUObjectHeaders.Num() + Module.PrivateUObjectHeaders.Num();
 	}
 	UE_LOG(LogCompile, Log, TEXT("Preparing %d modules took %.3f seconds"), GManifest.Modules.Num(), TotalPrepareModuleTime);
 	UE_LOG(LogCompile, Log, TEXT("Preparsing %d sources took %.3f seconds"), NumSources, TotalPreparseTime);
diff --git a/Engine/Source/Programs/UnrealHeaderTool/Private/Manifest.cpp b/Engine/Source/Programs/UnrealHeaderTool/Private/Manifest.cpp
--- a/Engine/Source/Programs/UnrealHeaderTool/Private/Manifest.cpp
+++ b/Engine/Source/Programs/UnrealHeaderTool/Private/Manifest.cpp
@@ -122,6 +122,7 @@ FManifest FManifest::LoadFromFile(const FString& Filename)
 
 		TArray<TSharedPtr<FJsonValue>> ClassesHeaders;
 		TArray<TSharedPtr<FJsonValue>> PublicHeaders;
+		TArray<TSharedPtr<FJsonValue>> InternalHeaders;
 		TArray<TSharedPtr<FJsonValue>> PrivateHeaders;
 
 		Result.Modules.AddZeroed();
@@ -136,6 +137,7 @@ FManifest FManifest::LoadFromFile(const FString& Filename)
 		GetJsonFieldValue(KnownModule.SaveExportedHeaders,       ModuleObj, TEXT("SaveExportedHeaders"),      *Outer);
 		GetJsonFieldValue(ClassesHeaders,                        ModuleObj, TEXT("ClassesHeaders"),           *Outer);
 		GetJsonFieldValue(PublicHeaders,                         ModuleObj, TEXT("PublicHeaders"),            *Outer);
+		GetJsonFieldValue(InternalHeaders,                       ModuleObj, TEXT("InternalHeaders"),          *Outer);
 		GetJsonFieldValue(PrivateHeaders,                        ModuleObj, TEXT("PrivateHeaders"),           *Outer);
 		GetJsonFieldValue(KnownModule.GeneratedCPPFilenameBase,  ModuleObj, TEXT("GeneratedCPPFilenameBase"), *Outer);
 		GetJsonFieldValue(GeneratedCodeVersionString, ModuleObj, TEXT("UHTGeneratedCodeVersion"), *Outer);
@@ -162,19 +164,22 @@ FManifest FManifest::LoadFromFile(const FString& Filename)
 		if (!KnownModule.IncludeBase              .EndsWith(TEXT("/"))) { KnownModule.IncludeBase              .AppendChar(TEXT('/')); }
 		if (!KnownModule.GeneratedIncludeDirectory.EndsWith(TEXT("/"))) { KnownModule.GeneratedIncludeDirectory.AppendChar(TEXT('/')); }
 
-		KnownModule.PublicUObjectClassesHeaders.AddZeroed(ClassesHeaders.Num());
-		KnownModule.PublicUObjectHeaders       .AddZeroed(PublicHeaders .Num());
-		KnownModule.PrivateUObjectHeaders      .AddZeroed(PrivateHeaders.Num());
+		KnownModule.PublicUObjectClassesHeaders.AddZeroed(ClassesHeaders .Num());
+		KnownModule.PublicUObjectHeaders       .AddZeroed(PublicHeaders  .Num());
+		KnownModule.InternalUObjectHeaders     .AddZeroed(InternalHeaders.Num());
+		KnownModule.PrivateUObjectHeaders      .AddZeroed(PrivateHeaders .Num());
 
-		ProcessHeaderArray(KnownModule.PublicUObjectClassesHeaders.GetData(), ClassesHeaders, *(Outer + TEXT(".ClassHeaders")));
-		ProcessHeaderArray(KnownModule.PublicUObjectHeaders       .GetData(), PublicHeaders , *(Outer + TEXT(".PublicHeaders" )));
-		ProcessHeaderArray(KnownModule.PrivateUObjectHeaders      .GetData(), PrivateHeaders, *(Outer + TEXT(".PrivateHeaders")));
+		ProcessHeaderArray(KnownModule.PublicUObjectClassesHeaders.GetData(), ClassesHeaders , *(Outer + TEXT(".ClassHeaders")));
+		ProcessHeaderArray(KnownModule.PublicUObjectHeaders       .GetData(), PublicHeaders  , *(Outer + TEXT(".PublicHeaders" )));
+		ProcessHeaderArray(KnownModule.InternalUObjectHeaders     .GetData(), InternalHeaders, *(Outer + TEXT(".InternalHeaders")));
+		ProcessHeaderArray(KnownModule.PrivateUObjectHeaders      .GetData(), PrivateHeaders , *(Outer + TEXT(".PrivateHeaders")));
 
 		// Sort the headers alphabetically.  This is just to add determinism to the compilation dependency order, since we currently
 		// don't rely on explicit includes (but we do support 'dependson')
 		// @todo uht: Ideally, we should sort these by sensical order before passing them in -- or better yet, follow include statements ourselves in here.
 		KnownModule.PublicUObjectClassesHeaders.Sort();
 		KnownModule.PublicUObjectHeaders       .Sort();
+		KnownModule.InternalUObjectHeaders     .Sort();
 		KnownModule.PrivateUObjectHeaders      .Sort();
 
 		UE_LOG(LogCompile, Log, TEXT("  %s"), *KnownModule.Name);
@@ -229,6 +234,15 @@ bool FManifestModule::NeedsRegeneration() const
 		}
 	}
 
+	for (const FString& Header : InternalUObjectHeaders)
+	{
+		if (IFileManager::Get().GetTimeStamp(*Header) > TimestampFileLastModify)
+		{
+			UE_LOG(LogCompile, Log, TEXT("File %s is newer than last timestamp. Regenerating reflection data for module %s."), *Header, *Name);
+			return true;
+		}
+	}
+
 	for (const FString& Header : PrivateUObjectHeaders)
 	{
 		if (IFileManager::Get().GetTimeStamp(*Header) > TimestampFileLastModify)
@@ -254,6 +268,7 @@ bool FManifestModule::IsCompatibleWith(const FManifestModule& ManifestModule)
 		&& GeneratedIncludeDirectory == ManifestModule.GeneratedIncludeDirectory
 		&& PublicUObjectClassesHeaders == ManifestModule.PublicUObjectClassesHeaders
 		&& PublicUObjectHeaders == ManifestModule.PublicUObjectHeaders
+		&& InternalUObjectHeaders == ManifestModule.InternalUObjectHeaders
 		&& PrivateUObjectHeaders == ManifestModule.PrivateUObjectHeaders
 		&& GeneratedCPPFilenameBase == ManifestModule.GeneratedCPPFilenameBase;
 }
diff --git a/Engine/Source/Programs/UnrealHeaderTool/Private/Manifest.h b/Engine/Source/Programs/UnrealHeaderTool/Private/Manifest.h
--- a/Engine/Source/Programs/UnrealHeaderTool/Private/Manifest.h
+++ b/Engine/Source/Programs/UnrealHeaderTool/Private/Manifest.h
@@ -37,6 +37,9 @@ struct FManifestModule
 	/** List of C++ public header files with UObjects in them */
 	TArray<FString> PublicUObjectHeaders;
 
+	/** List of C++ internal header files with UObjects in them */
+	TArray<FString> InternalUObjectHeaders;
+
 	/** List of C++ private header files with UObjects in them */
 	TArray<FString> PrivateUObjectHeaders;
 
