Added missing check for internal include directory.

#rnx

diff --git a/Engine/Source/Programs/UnrealHeaderTool/Private/CodeGenerator.cpp b/Engine/Source/Programs/UnrealHeaderTool/Private/CodeGenerator.cpp
--- a/Engine/Source/Programs/UnrealHeaderTool/Private/CodeGenerator.cpp
+++ b/Engine/Source/Programs/UnrealHeaderTool/Private/CodeGenerator.cpp
@@ -6462,6 +6462,7 @@ void PrepareModules(TArray<FUnrealPackageDefinitionInfo*>& PackageDefs, const FS
 					static const TCHAR PublicFolderName[] = TEXT("Public/");
 					static const TCHAR PrivateFolderName[] = TEXT("Private/");
 					static const TCHAR ClassesFolderName[] = TEXT("Classes/");
+					static const TCHAR InternalFolderName[] = TEXT("Internal/");
 					if (FCString::Strnicmp(IncludePath, PublicFolderName, UE_ARRAY_COUNT(PublicFolderName) - 1) == 0)
 					{
 						IncludePath += (UE_ARRAY_COUNT(PublicFolderName) - 1);
@@ -6474,6 +6475,10 @@ void PrepareModules(TArray<FUnrealPackageDefinitionInfo*>& PackageDefs, const FS
 					{
 						IncludePath += (UE_ARRAY_COUNT(ClassesFolderName) - 1);
 					}
+					else if (FCString::Strnicmp(IncludePath, InternalFolderName, UE_ARRAY_COUNT(InternalFolderName) - 1) == 0)
+					{
+						IncludePath += (UE_ARRAY_COUNT(InternalFolderName) - 1);
+					}
 
 					// Add the include path
 					if (*IncludePath != 0)
