Enforce use of TObjectPtr instead of raw pointers for member properties on all engine modules (but not all engine plugins) regardless of project and all modules used from ShooterGame.

diff --git a/Engine/Programs/UnrealHeaderTool/Config/DefaultEngine.ini b/Engine/Programs/UnrealHeaderTool/Config/DefaultEngine.ini
--- a/Engine/Programs/UnrealHeaderTool/Config/DefaultEngine.ini
+++ b/Engine/Programs/UnrealHeaderTool/Config/DefaultEngine.ini
@@ -32,7 +32,7 @@ gc.MaxObjectsInProgram=262144
 +TypeRedirects=(OldType="TAssetSubclassOf", NewType="TSoftClassPtr")
 
 [UnrealHeaderTool]
-EngineNativePointerMemberBehavior=AllowAndLog
+EngineNativePointerMemberBehavior=Disallow
 EngineObjectPtrMemberBehavior=AllowSilently
 NonEngineNativePointerMemberBehavior=AllowAndLog
 NonEngineObjectPtrMemberBehavior=AllowSilently
diff --git a/Engine/Source/Programs/UnrealHeaderTool/Private/HeaderParser.cpp b/Engine/Source/Programs/UnrealHeaderTool/Private/HeaderParser.cpp
--- a/Engine/Source/Programs/UnrealHeaderTool/Private/HeaderParser.cpp
+++ b/Engine/Source/Programs/UnrealHeaderTool/Private/HeaderParser.cpp
@@ -4045,8 +4045,9 @@ void FHeaderParser::GetVarType(
 						// Optionally emit messages about native pointer members and swallow trailing 'const' after pointer properties
 						if (VariableCategory == EVariableCategory::Member)
 						{
-							ConditionalLogPointerUsage(bIsCurrentModulePartOfEngine ? UHTConfig.EngineNativePointerMemberBehavior : UHTConfig.NonEngineNativePointerMemberBehavior,
-								TEXT("Native pointer"), FString(InputPos - VarStartPos, Input + VarStartPos).TrimStartAndEnd().ReplaceCharWithEscapedChar());
+							// TODO: Remove exclusion for plugins under engine when all plugins have had their raw pointers converted.
+							ConditionalLogPointerUsage(bIsCurrentModulePartOfEngine && !PackageDef.GetModule().BaseDirectory.Contains(TEXT("/Plugins/")) ? UHTConfig.EngineNativePointerMemberBehavior : UHTConfig.NonEngineNativePointerMemberBehavior,
+								TEXT("Native pointer"), FString(InputPos - VarStartPos, Input + VarStartPos).TrimStartAndEnd().ReplaceCharWithEscapedChar(), TEXT("TObjectPtr"));
 
 							MatchIdentifier(TEXT("const"), ESearchCase::CaseSensitive);
 						}
@@ -4055,8 +4056,9 @@ void FHeaderParser::GetVarType(
 					}
 					else if ((PropertyType == CPT_ObjectPtrReference) && (VariableCategory == EVariableCategory::Member))
 					{
-						ConditionalLogPointerUsage(bIsCurrentModulePartOfEngine ? UHTConfig.EngineObjectPtrMemberBehavior : UHTConfig.NonEngineObjectPtrMemberBehavior,
-							TEXT("ObjectPtr"), FString(InputPos - VarStartPos, Input + VarStartPos).TrimStartAndEnd().ReplaceCharWithEscapedChar());
+						// TODO: Remove exclusion for plugins under engine when all plugins have had their raw pointers converted.
+						ConditionalLogPointerUsage(bIsCurrentModulePartOfEngine && !PackageDef.GetModule().BaseDirectory.Contains(TEXT("/Plugins/")) ? UHTConfig.EngineObjectPtrMemberBehavior : UHTConfig.NonEngineObjectPtrMemberBehavior,
+							TEXT("ObjectPtr"), FString(InputPos - VarStartPos, Input + VarStartPos).TrimStartAndEnd().ReplaceCharWithEscapedChar(), nullptr);
 					}
 
 					// Imply const if it's a parameter that is a pointer to a const class
@@ -8904,18 +8906,32 @@ bool FHeaderParser::CheckUIMinMaxRangeFromMetaData(const FString& UIMin, const F
 	return true;
 }
 
-void FHeaderParser::ConditionalLogPointerUsage(EPointerMemberBehavior PointerMemberBehavior, const TCHAR* PointerTypeDesc, FString&& PointerTypeDecl)
+void FHeaderParser::ConditionalLogPointerUsage(EPointerMemberBehavior PointerMemberBehavior, const TCHAR* PointerTypeDesc, FString&& PointerTypeDecl, const TCHAR* AlternativeTypeDesc)
 {
 	switch (PointerMemberBehavior)
 	{
 		case EPointerMemberBehavior::Disallow:
 		{
-			LogError(TEXT("%s usage in member declaration detected in '%s', line %d [[%s]]"), PointerTypeDesc, *Filename, InputLine, *PointerTypeDecl);
+			if (AlternativeTypeDesc)
+			{
+				LogError(TEXT("%s usage in member declaration detected [[[%s]]].  This is disallowed for the target/module, consider %s as an alternative."), PointerTypeDesc, *PointerTypeDecl, AlternativeTypeDesc);
+			}
+			else
+			{
+				LogError(TEXT("%s usage in member declaration detected [[[%s]]]."), PointerTypeDesc, *PointerTypeDecl);
+			}
 		}
 		break;
 		case EPointerMemberBehavior::AllowAndLog:
 		{
-			UE_LOG(LogCompile, Log, TEXT("%s usage in member declaration detected in '%s', line %d [[%s]]"), PointerTypeDesc, *Filename, InputLine, *PointerTypeDecl);
+			if (AlternativeTypeDesc)
+			{
+				UE_LOG(LogCompile, Log, TEXT("%s(%d): %s usage in member declaration detected [[[%s]]].  Consider %s as an alternative."), *Filename, InputLine, PointerTypeDesc, *PointerTypeDecl, AlternativeTypeDesc);
+			}
+			else
+			{
+				UE_LOG(LogCompile, Log, TEXT("%s(%d): usage in member declaration detected [[[%s]]]."), *Filename, InputLine, PointerTypeDesc, *PointerTypeDecl);
+			}
 		}
 		break;
 		case EPointerMemberBehavior::AllowSilently:
diff --git a/Engine/Source/Programs/UnrealHeaderTool/Private/HeaderParser.h b/Engine/Source/Programs/UnrealHeaderTool/Private/HeaderParser.h
--- a/Engine/Source/Programs/UnrealHeaderTool/Private/HeaderParser.h
+++ b/Engine/Source/Programs/UnrealHeaderTool/Private/HeaderParser.h
@@ -617,8 +617,8 @@ private:
 	// Checks if a valid range has been found on the provided metadata
 	bool CheckUIMinMaxRangeFromMetaData(const FString& UIMin, const FString& UIMax);
 
-	// Checks if a valid range has been found on the provided metadata
-	void ConditionalLogPointerUsage(EPointerMemberBehavior PointerMemberBehavior, const TCHAR* PointerTypeDesc, FString&& PointerTypeDecl);
+	// Emits either a log message, warning, or error about a member pointer according to the behavior specified by the PointerMemberBehavior argument
+	void ConditionalLogPointerUsage(EPointerMemberBehavior PointerMemberBehavior, const TCHAR* PointerTypeDesc, FString&& PointerTypeDecl, const TCHAR* AlternativeTypeDesc);
 
 	// Check to see if the declaration is a constructor
 	static bool CheckForConstructor(FUnrealStructDefinitionInfo& StructDef, const FDeclaration& Declaration);
diff --git a/Engine/Source/Programs/UnrealObjectPtrTool/Private/UnrealObjectPtrTool.cpp b/Engine/Source/Programs/UnrealObjectPtrTool/Private/UnrealObjectPtrTool.cpp
--- a/Engine/Source/Programs/UnrealObjectPtrTool/Private/UnrealObjectPtrTool.cpp
+++ b/Engine/Source/Programs/UnrealObjectPtrTool/Private/UnrealObjectPtrTool.cpp
@@ -108,9 +108,9 @@ FStringView ParseBracketedStringSegment(const FString& InString, int32& InOutCur
 TValueOrError<FPointerUpgrader::FPointerUpgradeList, FString> FPointerUpgrader::FillUpgradeList(const FString& LogFilename)
 {
 	FPointerUpgradeList UpgradeList;
-	const TCHAR* Preamble = EnumHasAllFlags(BehaviorFlags, EPointerUpgradeBehaviorFlags::Reverse) ? TEXT("ObjectPtr usage in member declaration detected in '") : TEXT("Native pointer usage in member declaration detected in '");
+	const TCHAR* Preamble = EnumHasAllFlags(BehaviorFlags, EPointerUpgradeBehaviorFlags::Reverse) ? TEXT("ObjectPtr usage in member declaration detected [[") : TEXT("Native pointer usage in member declaration detected [[");
 	TArray<FString> PointerUpgradeEntries;
-	if (!FFileHelper::LoadFileToStringArrayWithPredicate(PointerUpgradeEntries, *LogFilename, [&Preamble](const FString& Line) { return Line.Contains(Preamble);  }))
+	if (!FFileHelper::LoadFileToStringArrayWithPredicate(PointerUpgradeEntries, *LogFilename, [&Preamble](const FString& Line) { return Line.Contains("LogCompile: ") && Line.Contains(Preamble);  }))
 	{
 		return MakeError(FString::Printf(TEXT("Unable to load UHT log: %s"), *LogFilename));
 	}
@@ -119,19 +119,20 @@ TValueOrError<FPointerUpgrader::FPointerUpgradeList, FString> FPointerUpgrader::
 	for (const FString& Entry : PointerUpgradeEntries)
 	{
 		int32 CurrentIndex = 0;
-		FStringView FilenameView = ParseBracketedStringSegment(Entry, CurrentIndex, Preamble, TEXT("'"));
+		FStringView FilenameView = ParseBracketedStringSegment(Entry, CurrentIndex, TEXT("LogCompile: "), TEXT("("));
 		if (FilenameView.IsEmpty())
 		{
 			return MakeError(FString::Printf(TEXT("Failed to parse filename from pointer member upgrade statement: %s"), *Entry));
 		}
 
-		FStringView LineNumberView = ParseBracketedStringSegment(Entry, CurrentIndex, TEXT(", line "), TEXT(" "));
+		CurrentIndex--; // Backtrack one character to allow us to re-find the opening bracket for the line number segment
+		FStringView LineNumberView = ParseBracketedStringSegment(Entry, CurrentIndex, TEXT("("), TEXT(")"));
 		if (LineNumberView.IsEmpty())
 		{
 			return MakeError(FString::Printf(TEXT("Failed to parse line number from pointer member upgrade statement: %s"), *Entry));
 		}
 
-		FStringView TypeNameView = ParseBracketedStringSegment(Entry, CurrentIndex, TEXT("[["), TEXT("]]"));
+		FStringView TypeNameView = ParseBracketedStringSegment(Entry, CurrentIndex, TEXT("[[["), TEXT("]]]"));
 		if (TypeNameView.IsEmpty())
 		{
 			return MakeError(FString::Printf(TEXT("Failed to parse type name from pointer member upgrade statement: %s"), *Entry));
