Fixing some UHT code generation issues to reduce obj size.

1) Functions added for LiveCoding have been removed and the data is now just a static.
2) Consolidated defered registration into a single routine.  This includes CppStructOps registration where every routine was over 100 bytes in length.
3) Removed the extra argument in function registration to avoid 7 bytes on windows
4) Fixed a typo in a macro name.

#b none
#rnx

diff --git a/Engine/Source/Programs/UnrealHeaderTool/Private/CodeGenerator.cpp b/Engine/Source/Programs/UnrealHeaderTool/Private/CodeGenerator.cpp
--- a/Engine/Source/Programs/UnrealHeaderTool/Private/CodeGenerator.cpp
+++ b/Engine/Source/Programs/UnrealHeaderTool/Private/CodeGenerator.cpp
@@ -1820,26 +1820,18 @@ void FNativeClassHeaderGenerator::ExportGeneratedPackageInitCode(FOutputDevice&
 	FOutputDeviceNull OutputDeviceNull;
 	FString MetaDataParams = OutputMetaDataCodeForObject(OutputDeviceNull, Out, PackageDef, TEXT("Package_MetaDataParams"), TEXT(""), TEXT("\t\t\t"));
 
-	Out.Logf(TEXT("\tuint32 %s()\t\r\n\t{\r\n\t\treturn 0x%08X;\r\n\t}\r\n"), *BodyHashFn, Hash);
-	Out.Logf(TEXT("\tuint32 %s()\t\r\n\t{\r\n\t\treturn 0x%08X;\r\n\t}\r\n"), *DeclarationsHashFn, DeclarationsHash);
+	Out.Logf(TEXT("\tstatic FPackageRegistrationInfo Z_Registration_Info_UPackage_%s;\r\n"), *PackageName);
 
-	Out.Logf(TEXT("\tstatic FPackageRegistrationInfo& Z_Registration_Info_UPackage_%s()\r\n"), *PackageName);
+	Out.Logf(TEXT("\tFORCENOINLINE UPackage* %s()\r\n"), *SingletonName);
 	Out.Logf(TEXT("\t{\r\n"));
-	Out.Logf(TEXT("\t\tstatic FPackageRegistrationInfo info;\r\n"));
-	Out.Logf(TEXT("\t\treturn info;\r\n"));
-	Out.Logf(TEXT("\t}\r\n"));
-
-	Out.Logf(TEXT("\tUPackage* %s()\r\n"), *SingletonName);
-	Out.Logf(TEXT("\t{\r\n"));
-	Out.Logf(TEXT("\t\tstatic UPackage*& ReturnPackage = Z_Registration_Info_UPackage_%s().OuterSingleton;\r\n"), *PackageName);
-	Out.Logf(TEXT("\t\tif (!ReturnPackage)\r\n"));
+	Out.Logf(TEXT("\t\tif (!Z_Registration_Info_UPackage_%s.OuterSingleton)\r\n"), *PackageName);
 	Out.Logf(TEXT("\t\t{\r\n"));
 
 	const TCHAR* SingletonArray;
 	const TCHAR* SingletonCount;
 	if (!Singletons.IsEmpty())
 	{
-		Out.Logf(TEXT("\t\t\tUObject* (*const SingletonFuncArray[])() = {\r\n"));
+		Out.Logf(TEXT("\t\t\tstatic UObject* (*const SingletonFuncArray[])() = {\r\n"));
 		for (FUnrealFieldDefinitionInfo* FieldDef : Singletons)
 		{
 			const FString& Name = FieldDef->GetSingletonNameChopped(true);
@@ -1856,7 +1848,7 @@ void FNativeClassHeaderGenerator::ExportGeneratedPackageInitCode(FOutputDevice&
 		SingletonCount = TEXT("0");
 	}
 
-	Out.Logf(TEXT("\t\t\tconst UECodeGen_Private::FPackageParams PackageParams = {\r\n"));
+	Out.Logf(TEXT("\t\t\tstatic const UECodeGen_Private::FPackageParams PackageParams = {\r\n"));
 	Out.Logf(TEXT("\t\t\t\t%s,\r\n"), *CreateUTF8LiteralString(Package->GetName()));
 	Out.Logf(TEXT("\t\t\t\t%s,\r\n"), SingletonArray);
 	Out.Logf(TEXT("\t\t\t\t%s,\r\n"), SingletonCount);
@@ -1865,19 +1857,19 @@ void FNativeClassHeaderGenerator::ExportGeneratedPackageInitCode(FOutputDevice&
 	Out.Logf(TEXT("\t\t\t\t0x%08X,\r\n"), DeclarationsHash);
 	Out.Logf(TEXT("\t\t\t\t%s\r\n"), *MetaDataParams);
 	Out.Logf(TEXT("\t\t\t};\r\n"));
-	Out.Logf(TEXT("\t\t\tUECodeGen_Private::ConstructUPackage(ReturnPackage, PackageParams);\r\n"));
+	Out.Logf(TEXT("\t\t\tUECodeGen_Private::ConstructUPackage(Z_Registration_Info_UPackage_%s.OuterSingleton, PackageParams);\r\n"), *PackageName);
 	Out.Logf(TEXT("\t\t}\r\n"));
-	Out.Logf(TEXT("\t\treturn ReturnPackage;\r\n"));
+	Out.Logf(TEXT("\t\treturn Z_Registration_Info_UPackage_%s.OuterSingleton;\r\n"), *PackageName);
 	Out.Logf(TEXT("\t}\r\n"));
 
 	// Do not change the Z_CompiledInDeferPackage_UPackage_ without changing LC_SymbolPatterns
-	Out.Logf(TEXT("\tstatic FRegisterCompiledInInfo Z_CompiledInDeferPackage_UPackage_%s(%s, TEXT(\"%s\"), Z_Registration_Info_UPackage_%s(), CONSTUCT_RELOAD_VERSION_INFO(FPackageReloadVersionInfo, %s(), %s()));\r\n"),
+	Out.Logf(TEXT("\tstatic FRegisterCompiledInInfo Z_CompiledInDeferPackage_UPackage_%s(%s, TEXT(\"%s\"), Z_Registration_Info_UPackage_%s, CONSTRUCT_RELOAD_VERSION_INFO(FPackageReloadVersionInfo, 0x%08X, 0x%08X));\r\n"),
 		*PackageName,
 		*SingletonName,
 		*Package->GetName(),
 		*PackageName,
-		*BodyHashFn,
-		*DeclarationsHashFn
+		Hash,
+		DeclarationsHash
 	);
 }
 
@@ -2134,18 +2126,20 @@ void FNativeClassHeaderGenerator::ExportNativeGeneratedInitCode(FOutputDevice& O
 
 		if (!bIsDynamic)
 		{
-			GeneratedClassRegisterFunctionText.Logf(TEXT("\textern FClassRegistrationInfo& Z_Registration_Info_UClass_%s();\r\n"), *ClassNameCPP);
+			Out.Logf(TEXT("\tIMPLEMENT_CLASS_NO_AUTO_REGISTRATION(%s);\r\n"), *ClassNameCPP);
 		}
 
 		GeneratedClassRegisterFunctionText.Logf(TEXT("\tUClass* %s\r\n"), *SingletonName);
 		GeneratedClassRegisterFunctionText.Logf(TEXT("\t{\r\n"));
+		FString OuterSingletonName;
 		if (!bIsDynamic)
 		{
-			GeneratedClassRegisterFunctionText.Logf(TEXT("\t\tstatic UClass*& OuterClass = Z_Registration_Info_UClass_%s().OuterSingleton;\r\n"), *ClassNameCPP);
-			GeneratedClassRegisterFunctionText.Logf(TEXT("\t\tif (!OuterClass)\r\n"));
+			OuterSingletonName = FString::Printf(TEXT("Z_Registration_Info_UClass_%s.OuterSingleton"), *ClassNameCPP);
+			GeneratedClassRegisterFunctionText.Logf(TEXT("\t\tif (!%s)\r\n"), *OuterSingletonName);
 		}
 		else
 		{
+			OuterSingletonName = TEXT("OuterClass");
 			const FString& DynamicClassPackageName = ClassDef.GetTypePackageName();
 			GeneratedClassRegisterFunctionText.Logf(TEXT("\t\tUPackage* OuterPackage = FindOrConstructDynamicTypePackage(TEXT(\"%s\"));\r\n"), *DynamicClassPackageName);
 			GeneratedClassRegisterFunctionText.Logf(TEXT("\t\tUClass* OuterClass = Cast<UClass>(StaticFindObjectFast(UClass::StaticClass(), OuterPackage, TEXT(\"%s\")));\r\n"), *FNativeClassHeaderGenerator::GetOverriddenName(ClassDef));
@@ -2153,7 +2147,7 @@ void FNativeClassHeaderGenerator::ExportNativeGeneratedInitCode(FOutputDevice& O
 		}
 
 		GeneratedClassRegisterFunctionText.Logf(TEXT("\t\t{\r\n"));
-		GeneratedClassRegisterFunctionText.Logf(TEXT("\t\t\tUECodeGen_Private::ConstructUClass(OuterClass, %s::ClassParams);\r\n"), *StaticsStructName);
+		GeneratedClassRegisterFunctionText.Logf(TEXT("\t\t\tUECodeGen_Private::ConstructUClass(%s, %s::ClassParams);\r\n"), *OuterSingletonName, *StaticsStructName);
 
 		TArray<FString> SparseClassDataTypes;
 		ClassDef.GetSparseClassDataTypes(SparseClassDataTypes);
@@ -2168,12 +2162,12 @@ void FNativeClassHeaderGenerator::ExportNativeGeneratedInitCode(FOutputDevice& O
 			const FString* CustomDynamicClassInitializationMD = ClassDef.FindMetaData(TEXT("CustomDynamicClassInitialization"));
 			if (CustomDynamicClassInitializationMD)
 			{
-				GeneratedClassRegisterFunctionText.Logf(TEXT("\t\t\t\t%s(CastChecked<UDynamicClass>(OuterClass));\n"), *(*CustomDynamicClassInitializationMD));
+				GeneratedClassRegisterFunctionText.Logf(TEXT("\t\t\t\t%s(CastChecked<UDynamicClass>(%s));\n"), *(*CustomDynamicClassInitializationMD), *OuterSingletonName);
 			}
 		}
 
 		GeneratedClassRegisterFunctionText.Logf(TEXT("\t\t}\r\n"));
-		GeneratedClassRegisterFunctionText.Logf(TEXT("\t\treturn OuterClass;\r\n"));
+		GeneratedClassRegisterFunctionText.Logf(TEXT("\t\treturn %s;\r\n"), *OuterSingletonName);
 		GeneratedClassRegisterFunctionText.Logf(TEXT("\t}\r\n"));
 
 		Out.Log(*GeneratedClassRegisterFunctionText);
@@ -2228,11 +2222,7 @@ void FNativeClassHeaderGenerator::ExportNativeGeneratedInitCode(FOutputDevice& O
 	ClassDef.SetHash(ClassHash);
 
 	// Emit the IMPLEMENT_CLASS macro to go in the generated cpp file.
-	if (!bIsDynamic)
-	{
-		Out.Logf(TEXT("\tIMPLEMENT_CLASS(%s, %u);\r\n"), *ClassNameCPP, ClassHash);
-	}
-	else
+	if (bIsDynamic)
 	{
 		Out.Logf(TEXT("\tIMPLEMENT_DYNAMIC_CLASS(%s, TEXT(\"%s\"), %u);\r\n"), *ClassNameCPP, *OverriddenClassName, ClassHash);
 	}
@@ -2425,7 +2415,7 @@ void FNativeClassHeaderGenerator::ExportFunction(FOutputDevice& Out, FReferenceG
 
 	CurrentFunctionText.Logf(TEXT("\t\tif (!ReturnFunction)\r\n"));
 	CurrentFunctionText.Logf(TEXT("\t\t{\r\n"));
-	CurrentFunctionText.Logf(TEXT("\t\t\tUECodeGen_Private::ConstructUFunction(ReturnFunction, %s::FuncParams, &ReturnFunction);\r\n"), *StaticsStructName);
+	CurrentFunctionText.Logf(TEXT("\t\t\tUECodeGen_Private::ConstructUFunction(&ReturnFunction, %s::FuncParams);\r\n"), *StaticsStructName);
 	CurrentFunctionText.Log (TEXT("\t\t}\r\n"));
 	CurrentFunctionText.Log (TEXT("\t\treturn ReturnFunction;\r\n"));
 	CurrentFunctionText.Log (TEXT("\t}\r\n"));
@@ -3582,35 +3572,36 @@ void FNativeClassHeaderGenerator::ExportGeneratedStructBodyMacros(FOutputDevice&
 
 		if (!bIsDynamic)
 		{
-			Out.Logf(TEXT("\tstatic FStructRegistrationInfo& Z_Registration_Info_UScriptStruct_%s()\r\n"), *ScriptStructDef.GetName());
-			Out.Logf(TEXT("\t{\r\n"));
-			Out.Logf(TEXT("\t\tstatic FStructRegistrationInfo info;\r\n"));
-			Out.Logf(TEXT("\t\treturn info;\r\n"));
-			Out.Logf(TEXT("\t}\r\n"));
+			Out.Logf(TEXT("\tstatic FStructRegistrationInfo Z_Registration_Info_UScriptStruct_%s;\r\n"), *ScriptStructDef.GetName());
 		}
 
-		Out.Logf(TEXT("\textern uint32 %s();\r\n"), *GetHashName);
+		if (bIsDynamic)
+		{
+			Out.Logf(TEXT("\textern uint32 %s();\r\n"), *GetHashName);
+		}
 
 		Out.Logf(TEXT("class UScriptStruct* %s::StaticStruct()\r\n"), *StructNameCPP);
 		Out.Logf(TEXT("{\r\n"));
 
 		// UStructs can have UClass or UPackage outer (if declared in non-UClass headers).
 		const FString& OuterName (bIsDynamic ? STRING_StructPackage : GetPackageSingletonName(UHTCastChecked<FUnrealPackageDefinitionInfo>(ScriptStructDef.GetOuter()), OutReferenceGatherers.UniqueCrossModuleReferences));
+		FString OuterSingletonName;
 		if (!bIsDynamic)
 		{
-			Out.Logf(TEXT("\tstatic class UScriptStruct*& Singleton = Z_Registration_Info_UScriptStruct_%s().OuterSingleton;\r\n"), *ScriptStructDef.GetName());
+			OuterSingletonName = FString::Printf(TEXT("Z_Registration_Info_UScriptStruct_%s.OuterSingleton"), *ScriptStructDef.GetName());
 		}
 		else
 		{
 			Out.Logf(TEXT("\tclass UPackage* %s = FindOrConstructDynamicTypePackage(TEXT(\"%s\"));\r\n"), *OuterName, *ScriptStructDef.GetTypePackageName());
 			Out.Logf(TEXT("\tclass UScriptStruct* Singleton = Cast<UScriptStruct>(StaticFindObjectFast(UScriptStruct::StaticClass(), %s, TEXT(\"%s\")));\r\n"), *OuterName, *ActualStructName);
+			OuterSingletonName = TEXT("Singleton");
 		}
 
-		Out.Logf(TEXT("\tif (!Singleton)\r\n"));
+		Out.Logf(TEXT("\tif (!%s)\r\n"), *OuterSingletonName);
 		Out.Logf(TEXT("\t{\r\n"));
 
-		Out.Logf(TEXT("\t\tSingleton = GetStaticStruct(%s, %s, TEXT(\"%s\"));\r\n"),
-			*ChoppedSingletonName, *OuterName, *ActualStructName);
+		Out.Logf(TEXT("\t\t%s = GetStaticStruct(%s, %s, TEXT(\"%s\"));\r\n"),
+			*OuterSingletonName, *ChoppedSingletonName, *OuterName, *ActualStructName);
 
 		// if this struct has RigVM methods - we need to register the method to our central
 		// registry on construction of the static struct
@@ -3618,19 +3609,19 @@ void FNativeClassHeaderGenerator::ExportGeneratedStructBodyMacros(FOutputDevice&
 		{
 			for (const FRigVMMethodInfo& MethodInfo : StructRigVMInfo.Methods)
 			{
-				Out.Logf(TEXT("\t\tFRigVMRegistry::Get().Register(TEXT(\"%s::%s\"), &%s::RigVM%s, Singleton);\r\n"),
-					*StructNameCPP, *MethodInfo.Name, *StructNameCPP, *MethodInfo.Name);
+				Out.Logf(TEXT("\t\tFRigVMRegistry::Get().Register(TEXT(\"%s::%s\"), &%s::RigVM%s, %s);\r\n"),
+					*StructNameCPP, *MethodInfo.Name, *StructNameCPP, *MethodInfo.Name, *OuterSingletonName);
 			}
 		}
 
 		Out.Logf(TEXT("\t}\r\n"));
-		Out.Logf(TEXT("\treturn Singleton;\r\n"));
+		Out.Logf(TEXT("\treturn %s;\r\n"), *OuterSingletonName);
 		Out.Logf(TEXT("}\r\n"));
 
-		// Forward declare the StaticStruct specialisation in the header
+		// Forward declare the StaticStruct specialization in the header
 		OutGeneratedHeaderText.Logf(TEXT("template<> %sUScriptStruct* StaticStruct<struct %s>();\r\n\r\n"), *GetAPIString(), *StructNameCPP);
 
-		// Generate the StaticStruct specialisation
+		// Generate the StaticStruct specialization
 		Out.Logf(TEXT("template<> %sUScriptStruct* StaticStruct<%s>()\r\n"), *GetAPIString(), *StructNameCPP);
 		Out.Logf(TEXT("{\r\n"));
 		Out.Logf(TEXT("\treturn %s::StaticStruct();\r\n"), *StructNameCPP);
@@ -3647,33 +3638,6 @@ void FNativeClassHeaderGenerator::ExportGeneratedStructBodyMacros(FOutputDevice&
 				*AsTEXT(StructPackageName),
 				*AsTEXT(FNativeClassHeaderGenerator::GetOverriddenPathName(ScriptStructDef)));
 		}
-		else
-		{
-			// Do not change the Z_CompiledInDeferStruct_UScriptStruct_ without changing LC_SymbolPatterns
-			Out.Logf(TEXT("static FRegisterCompiledInInfo Z_CompiledInDeferStruct_UScriptStruct_%s(%s::StaticStruct, TEXT(\"%s\"), TEXT(\"%s\"), Z_Registration_Info_UScriptStruct_%s(), CONSTUCT_RELOAD_VERSION_INFO(FStructReloadVersionInfo, sizeof(%s), %s()));\r\n"),
-				*StructNameCPP,
-				*StructNameCPP,
-				*ScriptStructDef.GetPackageDef().GetName(),
-				*ActualStructName,
-				*ScriptStructDef.GetName(),
-				*StructNameCPP, 
-				*GetHashName
-			);
-		}
-
-		// Generate StaticRegisterNatives equivalent for structs without classes.
-		if (!ScriptStructDef.GetOuter()->GetObject()->IsA(UStruct::StaticClass()))
-		{
-			// Do not change the Z_CompiledInDeferCppStructOps_UScriptStruct_ without changing LC_SymbolPatterns
-			const FString ShortPackageName = FPackageName::GetShortName(ScriptStructDef.GetOuter()->GetName());
-			Out.Logf(TEXT("static struct FScriptStruct_%s_StaticRegisterNatives%s\r\n"), *ShortPackageName, *StructNameCPP);
-			Out.Logf(TEXT("{\r\n"));
-			Out.Logf(TEXT("\tFScriptStruct_%s_StaticRegisterNatives%s()\r\n"), *ShortPackageName, *StructNameCPP);
-			Out.Logf(TEXT("\t{\r\n"));
-			Out.Logf(TEXT("\t\tUScriptStruct::DeferCppStructOps<%s>(FName(TEXT(\"%s\")));\r\n"), *StructNameCPP, *ActualStructName);
-			Out.Logf(TEXT("\t}\r\n"));
-			Out.Logf(TEXT("} Z_CompiledInDeferCppStructOps_UScriptStruct_%s;\r\n"), *StructNameCPP);
-		}
 	}
 
 	FString StaticsStructName = ChoppedSingletonName + TEXT("_Statics");
@@ -3772,35 +3736,41 @@ void FNativeClassHeaderGenerator::ExportGeneratedStructBodyMacros(FOutputDevice&
 
 	FString HashFuncName = FString::Printf(TEXT("Get_%s_Hash"), *SingletonName.Replace(TEXT("()"), TEXT(""), ESearchCase::CaseSensitive));
 	// Structs can either have a UClass or UPackage as outer (if declared in non-UClass header).
+	FString InnerSingletonName;
 	if (!bIsDynamic)
 	{
 		if (ScriptStructDef.HasAnyStructFlags(STRUCT_Native))
 		{
-			GeneratedStructRegisterFunctionText.Logf(TEXT("\t\tstatic UScriptStruct*& ReturnStruct = Z_Registration_Info_UScriptStruct_%s().InnerSingleton;\r\n"), *ScriptStructDef.GetName());
+			InnerSingletonName = FString::Printf(TEXT("Z_Registration_Info_UScriptStruct_%s.InnerSingleton"), *ScriptStructDef.GetName());
 		}
 		else
 		{
 			GeneratedStructRegisterFunctionText.Logf(TEXT("\t\tstatic UScriptStruct* ReturnStruct = nullptr;\r\n"));
+			InnerSingletonName = TEXT("ReturnStruct");
 		}
 	}
 	else
 	{
 		GeneratedStructRegisterFunctionText.Logf(TEXT("\t\tUPackage* Outer = FindOrConstructDynamicTypePackage(TEXT(\"%s\"));\r\n"), *ScriptStructDef.GetTypePackageName());
 		GeneratedStructRegisterFunctionText.Logf(TEXT("\t\tUScriptStruct* ReturnStruct = FindExistingStructIfHotReloadOrDynamic(Outer, TEXT(\"%s\"), sizeof(%s), %s(), true);\r\n"), *ActualStructName, *NoExportStructNameCPP, *HashFuncName);
+		InnerSingletonName = TEXT("ReturnStruct");
 	}
-	GeneratedStructRegisterFunctionText.Logf(TEXT("\t\tif (!ReturnStruct)\r\n"));
+	GeneratedStructRegisterFunctionText.Logf(TEXT("\t\tif (!%s)\r\n"), *InnerSingletonName);
 	GeneratedStructRegisterFunctionText.Log (TEXT("\t\t{\r\n"));
 
-	GeneratedStructRegisterFunctionText.Logf(TEXT("\t\t\tUECodeGen_Private::ConstructUScriptStruct(ReturnStruct, %s::ReturnStructParams);\r\n"), *StaticsStructName);
+	GeneratedStructRegisterFunctionText.Logf(TEXT("\t\t\tUECodeGen_Private::ConstructUScriptStruct(%s, %s::ReturnStructParams);\r\n"), *InnerSingletonName, *StaticsStructName);
 	GeneratedStructRegisterFunctionText.Log (TEXT("\t\t}\r\n"));
-	GeneratedStructRegisterFunctionText.Logf(TEXT("\t\treturn ReturnStruct;\r\n"));
+	GeneratedStructRegisterFunctionText.Logf(TEXT("\t\treturn %s;\r\n"), *InnerSingletonName);
 	GeneratedStructRegisterFunctionText.Log (TEXT("\t}\r\n"));
 
 	uint32 StructHash = GenerateTextHash(*GeneratedStructRegisterFunctionText);
 	ScriptStructDef.SetHash(StructHash);
 
 	Out.Log(GeneratedStructRegisterFunctionText);
-	Out.Logf(TEXT("\tFORCENOINLINE uint32 %s() { return %uU; }\r\n"), *HashFuncName, StructHash);
+	if (bIsDynamic)
+	{
+		Out.Logf(TEXT("\tuint32 %s() { return %uU; }\r\n"), *HashFuncName, StructHash);
+	}
 
 	// if this struct has RigVM methods we need to implement both the 
 	// virtual function as well as the stub method here.
@@ -3934,11 +3904,7 @@ void FNativeClassHeaderGenerator::ExportGeneratedEnumInitCode(FOutputDevice& Out
 
 	if (!bIsDynamic)
 	{
-		Out.Logf(TEXT("\tstatic FEnumRegistrationInfo& Z_Registration_Info_UEnum_%s()\r\n"), *EnumNameCpp);
-		Out.Logf(TEXT("\t{\r\n"));
-		Out.Logf(TEXT("\t\tstatic FEnumRegistrationInfo info;\r\n"));
-		Out.Logf(TEXT("\t\treturn info;\r\n"));
-		Out.Logf(TEXT("\t}\r\n"));
+		Out.Logf(TEXT("\tstatic FEnumRegistrationInfo Z_Registration_Info_UEnum_%s;\r\n"), *EnumNameCpp);
 	}
 
 	Out.Logf(TEXT("\tstatic UEnum* %s_StaticEnum()\r\n"), *EnumNameCpp);
@@ -3946,12 +3912,11 @@ void FNativeClassHeaderGenerator::ExportGeneratedEnumInitCode(FOutputDevice& Out
 
 	if (!bIsDynamic)
 	{
-		Out.Logf(TEXT("\t\tstatic UEnum*& OuterSingleton = Z_Registration_Info_UEnum_%s().OuterSingleton;\r\n"), *EnumNameCpp);
-		Out.Logf(TEXT("\t\tif (!OuterSingleton)\r\n"));
+		Out.Logf(TEXT("\t\tif (!Z_Registration_Info_UEnum_%s.OuterSingleton)\r\n"), *EnumNameCpp);
 		Out.Logf(TEXT("\t\t{\r\n"));
-		Out.Logf(TEXT("\t\t\tOuterSingleton = GetStaticEnum(%s, %s, TEXT(\"%s\"));\r\n"), *SingletonName, *PackageSingletonName, *EnumNameCpp);
+		Out.Logf(TEXT("\t\t\tZ_Registration_Info_UEnum_%s.OuterSingleton = GetStaticEnum(%s, %s, TEXT(\"%s\"));\r\n"), *EnumNameCpp, *SingletonName, *PackageSingletonName, *EnumNameCpp);
 		Out.Logf(TEXT("\t\t}\r\n"));
-		Out.Logf(TEXT("\t\treturn OuterSingleton;\r\n"));
+		Out.Logf(TEXT("\t\treturn Z_Registration_Info_UEnum_%s.OuterSingleton;\r\n"), *EnumNameCpp);
 	}
 	else
 	{
@@ -3969,6 +3934,11 @@ void FNativeClassHeaderGenerator::ExportGeneratedEnumInitCode(FOutputDevice& Out
 	const FString& EnumSingletonName = EnumDef.GetSingletonName(true);
 	const FString HashFuncName = FString::Printf(TEXT("Get_%s_Hash"), *SingletonName);
 
+	if (bIsDynamic)
+	{
+		Out.Logf(TEXT("\textern uint32 %s();\r\n"), *HashFuncName);
+	}
+
 	Out.Logf(TEXT("\ttemplate<> %sUEnum* StaticEnum<%s>()\r\n"), *GetAPIString(), *EnumDef.GetCppType());
 	Out.Logf(TEXT("\t{\r\n"));
 	Out.Logf(TEXT("\t\treturn %s_StaticEnum();\r\n"), *EnumNameCpp);
@@ -4049,27 +4019,32 @@ void FNativeClassHeaderGenerator::ExportGeneratedEnumInitCode(FOutputDevice& Out
 	GeneratedEnumRegisterFunctionText.Logf(TEXT("\t{\r\n"));
 
 	// Enums can either have a UClass or UPackage as outer (if declared in non-UClass header).
+	FString InnerSingletonName;
 	if (!bIsDynamic)
 	{
-		GeneratedEnumRegisterFunctionText.Logf(TEXT("\t\tstatic UEnum*& ReturnEnum = Z_Registration_Info_UEnum_%s().InnerSingleton;\r\n"), *EnumNameCpp);
+		InnerSingletonName = FString::Printf(TEXT("Z_Registration_Info_UEnum_%s.InnerSingleton"), *EnumNameCpp);
 	}
 	else
 	{
 		GeneratedEnumRegisterFunctionText.Logf(TEXT("\t\tUPackage* Outer = FindOrConstructDynamicTypePackage(TEXT(\"%s\"));"), *PackageSingletonName);
-		GeneratedEnumRegisterFunctionText.Logf(TEXT("\t\tUEnum* ReturnEnum = FindExistingEnumIfHotReloadOrDynamic(Outer, TEXT(\"%s\"), 0, %s(), true);\r\n"), *OverriddenEnumNameCpp, *HashFuncName);
+		GeneratedEnumRegisterFunctionText.Logf(TEXT("\t\tUEnum* ReturnEnum = FindExistingEnumIfHotReloadOrDynamic(Outer, TEXT(\"%s\"), 0, %uU, true);\r\n"), *OverriddenEnumNameCpp, *HashFuncName);
+		InnerSingletonName = TEXT("ReturnEnum");
 	}
-	GeneratedEnumRegisterFunctionText.Logf(TEXT("\t\tif (!ReturnEnum)\r\n"));
+	GeneratedEnumRegisterFunctionText.Logf(TEXT("\t\tif (!%s)\r\n"), *InnerSingletonName);
 	GeneratedEnumRegisterFunctionText.Logf(TEXT("\t\t{\r\n"));
 
-	GeneratedEnumRegisterFunctionText.Logf(TEXT("\t\t\tUECodeGen_Private::ConstructUEnum(ReturnEnum, %s::EnumParams);\r\n"), *StaticsStructName);
+	GeneratedEnumRegisterFunctionText.Logf(TEXT("\t\t\tUECodeGen_Private::ConstructUEnum(%s, %s::EnumParams);\r\n"), *InnerSingletonName, *StaticsStructName);
 	GeneratedEnumRegisterFunctionText.Logf(TEXT("\t\t}\r\n"));
-	GeneratedEnumRegisterFunctionText.Logf(TEXT("\t\treturn ReturnEnum;\r\n"));
+	GeneratedEnumRegisterFunctionText.Logf(TEXT("\t\treturn %s;\r\n"), *InnerSingletonName);
 	GeneratedEnumRegisterFunctionText.Logf(TEXT("\t}\r\n"));
 
 	uint32 EnumHash = GenerateTextHash(*GeneratedEnumRegisterFunctionText);
 	EnumDef.SetHash(EnumHash);
 
-	Out.Logf(TEXT("\tFORCENOINLINE uint32 %s() { return %uU; }\r\n"), *HashFuncName, EnumHash);
+	if (bIsDynamic)
+	{
+		Out.Logf(TEXT("\tuint32 %s() { return %uU; }\r\n"), *HashFuncName, EnumHash);
+	}
 
 	if (bIsDynamic)
 	{
@@ -4084,19 +4059,6 @@ void FNativeClassHeaderGenerator::ExportGeneratedEnumInitCode(FOutputDevice& Out
 			*AsTEXT(FNativeClassHeaderGenerator::GetOverriddenPathName(EnumDef))
 		);
 	}
-	else
-	{
-		// Do not change the Z_CompiledInDeferEnum_UEnum_ without changing LC_SymbolPatterns
-		Out.Logf(
-			TEXT("\tstatic FRegisterCompiledInInfo Z_CompiledInDeferEnum_UEnum_%s(%s_StaticEnum, TEXT(\"%s\"), TEXT(\"%s\"), Z_Registration_Info_UEnum_%s(), CONSTUCT_RELOAD_VERSION_INFO(FEnumReloadVersionInfo, %s()));\r\n"),
-			*EnumNameCpp,
-			*EnumNameCpp,
-			*EnumDef.GetPackageDef().GetName(),
-			*OverriddenEnumNameCpp,
-			*EnumNameCpp,
-			*HashFuncName
-		);
-	}
 
 	Out.Log(GeneratedEnumRegisterFunctionText);
 }
@@ -5505,6 +5467,12 @@ void FNativeClassHeaderGenerator::GenerateSourceFile(FGeneratedCPP& GeneratedCPP
 
 			TArray<FUnrealFieldDefinitionInfo*> Singletons;
 			Singletons.Reserve(Types.Num());
+			TArray<FUnrealEnumDefinitionInfo*> EnumRegs;
+			EnumRegs.Reserve(Types.Num());
+			TArray<FUnrealScriptStructDefinitionInfo*> ScriptStructRegs;
+			ScriptStructRegs.Reserve(Types.Num());
+			TArray<FUnrealClassDefinitionInfo*> ClassRegs;
+			ClassRegs.Reserve(Types.Num());
 
 			const FString FileDefineName = SourceFile.GetFileDefineName();
 			const FString& StrippedFilename = SourceFile.GetStrippedFilename();
@@ -5526,6 +5494,10 @@ void FNativeClassHeaderGenerator::GenerateSourceFile(FGeneratedCPP& GeneratedCPP
 					{
 						GeneratedFunctionDeclarations.Log(EnumDef->GetExternDecl(true));
 						Generator.ExportGeneratedEnumInitCode(GeneratedCPPText, ReferenceGatherers, SourceFile, *EnumDef);
+						if (!EnumDef->IsDynamic())
+						{
+							EnumRegs.Add(EnumDef);
+						}
 					}
 				}
 				else if (FUnrealScriptStructDefinitionInfo* ScriptStructDef = UHTCast<FUnrealScriptStructDefinitionInfo>(FieldDef))
@@ -5536,6 +5508,10 @@ void FNativeClassHeaderGenerator::GenerateSourceFile(FGeneratedCPP& GeneratedCPP
 					}
 					GeneratedFunctionDeclarations.Log(ScriptStructDef->GetExternDecl(true));
 					Generator.ExportGeneratedStructBodyMacros(GeneratedHeaderText, GeneratedCPPText, ReferenceGatherers, SourceFile, *ScriptStructDef);
+					if (!ScriptStructDef->IsDynamic() && ScriptStructDef->HasAnyStructFlags(STRUCT_Native))
+					{
+						ScriptStructRegs.Add(ScriptStructDef);
+					}
 				}
 				else if (FUnrealFunctionDefinitionInfo* FunctionDef = UHTCast<FUnrealFunctionDefinitionInfo>(FieldDef))
 				{
@@ -5552,6 +5528,10 @@ void FNativeClassHeaderGenerator::GenerateSourceFile(FGeneratedCPP& GeneratedCPP
 					if (!ClassDef->HasAnyClassFlags(CLASS_Intrinsic))
 					{
 						Generator.ExportClassFromSourceFileInner(GeneratedHeaderText, GeneratedCPPText, GeneratedFunctionDeclarations, ReferenceGatherers, *ClassDef, SourceFile, GeneratedCPP.ExportFlags);
+						if (!ClassDef->IsDynamic())
+						{
+							ClassRegs.Add(ClassDef);
+						}
 					}
 				}
 			}
@@ -5567,6 +5547,148 @@ void FNativeClassHeaderGenerator::GenerateSourceFile(FGeneratedCPP& GeneratedCPP
 				}
 			}
 
+			// Generate the single registration method
+			if (!EnumRegs.IsEmpty() || !ScriptStructRegs.IsEmpty() || !ClassRegs.IsEmpty())
+			{
+				static const TCHAR* Prefix = TEXT("Z_CompiledInDeferFile_");
+				FString StaticsName = FString::Printf(TEXT("%s%s_Statics"), Prefix, *SourceFile.GetFileId());
+
+				GeneratedCPPText.Logf(TEXT("\tstruct %s\r\n"), *StaticsName);
+				GeneratedCPPText.Log(TEXT("\t{\r\n"));
+
+				size_t EditorOnlyEnumCount = 0;
+				for (FUnrealEnumDefinitionInfo* EnumDef : EnumRegs)
+				{
+					if (EnumDef->IsEditorOnly())
+					{
+						++EditorOnlyEnumCount;
+					}
+				}
+				bool bAllEnumsEditorOnly = !EnumRegs.IsEmpty() && EditorOnlyEnumCount == EnumRegs.Num();
+
+				if (!EnumRegs.IsEmpty())
+				{
+					if (bAllEnumsEditorOnly)
+					{
+						GeneratedCPPText.Log(TEXT("#if WITH_EDITORONLY_DATA\r\n"));
+					}
+					GeneratedCPPText.Log(TEXT("\t\tstatic const FEnumRegisterCompiledInInfo EnumInfo[];\r\n"));
+					if (bAllEnumsEditorOnly)
+					{
+						GeneratedCPPText.Log(TEXT("#endif\r\n"));
+					}
+				}
+				if (!ScriptStructRegs.IsEmpty())
+				{
+					GeneratedCPPText.Log(TEXT("\t\tstatic const FStructRegisterCompiledInInfo ScriptStructInfo[];\r\n"));
+				}
+				if (!ClassRegs.IsEmpty())
+				{
+					GeneratedCPPText.Log(TEXT("\t\tstatic const FClassRegisterCompiledInInfo ClassInfo[];\r\n"));
+				}
+
+				GeneratedCPPText.Log(TEXT("\t};\r\n"));
+
+				int32 CombinedHash = -1;
+				if (!EnumRegs.IsEmpty())
+				{
+					if (bAllEnumsEditorOnly)
+					{
+						GeneratedCPPText.Log(TEXT("#if WITH_EDITORONLY_DATA\r\n"));
+					}
+					GeneratedCPPText.Logf(TEXT("\tconst FEnumRegisterCompiledInInfo %s::EnumInfo[] = {\r\n"), *StaticsName);
+					for (FUnrealEnumDefinitionInfo* EnumDef : EnumRegs)
+					{
+						if (!bAllEnumsEditorOnly && EnumDef->IsEditorOnly())
+						{
+							GeneratedCPPText.Log(TEXT("#if WITH_EDITORONLY_DATA\r\n"));
+						}
+						const FString EnumNameCpp = EnumDef->GetName();
+						const FString OverriddenEnumNameCpp = FNativeClassHeaderGenerator::GetOverriddenName(*EnumDef);
+						GeneratedCPPText.Logf(TEXT("\t\t{ %s_StaticEnum, TEXT(\"%s\"), &Z_Registration_Info_UEnum_%s, CONSTRUCT_RELOAD_VERSION_INFO(FEnumReloadVersionInfo, %uU) },\r\n"),
+							*EnumNameCpp,
+							*OverriddenEnumNameCpp,
+							*EnumNameCpp,
+							EnumDef->GetHash(*EnumDef));
+						if (!bAllEnumsEditorOnly && EnumDef->IsEditorOnly())
+						{
+							GeneratedCPPText.Log(TEXT("#endif\r\n"));
+						}
+						CombinedHash = HashCombine(CombinedHash, EnumDef->GetHash(*EnumDef));
+					}
+					GeneratedCPPText.Logf(TEXT("\t};\r\n"));
+					if (bAllEnumsEditorOnly)
+					{
+						GeneratedCPPText.Log(TEXT("#endif\r\n"));
+					}
+				}
+
+				if (!ScriptStructRegs.IsEmpty())
+				{
+					GeneratedCPPText.Logf(TEXT("\tconst FStructRegisterCompiledInInfo %s::ScriptStructInfo[] = {\r\n"), *StaticsName);
+					for (FUnrealScriptStructDefinitionInfo* ScriptStructDef : ScriptStructRegs)
+					{
+						const FString StructNameCPP = ScriptStructDef->GetAlternateNameCPP();
+						const FString ActualStructName = FNativeClassHeaderGenerator::GetOverriddenName(*ScriptStructDef);
+						GeneratedCPPText.Logf(TEXT("\t\t{ %s::StaticStruct, Z_Construct_UScriptStruct_%s_Statics::NewStructOps, TEXT(\"%s\"), &Z_Registration_Info_UScriptStruct_%s, CONSTRUCT_RELOAD_VERSION_INFO(FStructReloadVersionInfo, sizeof(%s), %uU) },\r\n"),
+							*StructNameCPP,
+							*StructNameCPP,
+							*ActualStructName,
+							*ScriptStructDef->GetName(),
+							*StructNameCPP,
+							ScriptStructDef->GetHash(*ScriptStructDef));
+						CombinedHash = HashCombine(CombinedHash, ScriptStructDef->GetHash(*ScriptStructDef));
+					}
+					GeneratedCPPText.Logf(TEXT("\t};\r\n"));
+				}
+
+				if (!ClassRegs.IsEmpty())
+				{
+					GeneratedCPPText.Logf(TEXT("\tconst FClassRegisterCompiledInInfo %s::ClassInfo[] = {\r\n"), *StaticsName);
+					for (FUnrealClassDefinitionInfo* ClassDef : ClassRegs)
+					{
+						const FString ClassNameCPP = ClassDef->GetAlternateNameCPP();
+						GeneratedCPPText.Logf(TEXT("\t\t{ Z_Construct_UClass_%s, %s::StaticClass, TEXT(\"%s\"), &Z_Registration_Info_UClass_%s, CONSTRUCT_RELOAD_VERSION_INFO(FClassReloadVersionInfo, sizeof(%s), %uU) },\r\n"),
+							*ClassNameCPP,
+							*ClassNameCPP,
+							*ClassNameCPP,
+							*ClassNameCPP,
+							*ClassNameCPP,
+							ClassDef->GetHash(*ClassDef));
+						CombinedHash = HashCombine(CombinedHash, ClassDef->GetHash(*ClassDef));
+					}
+					GeneratedCPPText.Logf(TEXT("\t};\r\n"));
+				}
+
+				auto FormatArray = [&StaticsName](bool bIsEmpty, bool bAllEditorOnlyData, const FString& Text)
+				{
+					if (bAllEditorOnlyData)
+					{
+						return FString::Printf(TEXT("IF_WITH_EDITORONLY_DATA(%s::%s, nullptr), IF_WITH_EDITORONLY_DATA(UE_ARRAY_COUNT(%s::%s), 0)"), *StaticsName, *Text, *StaticsName, *Text);
+					}
+					else if (!bIsEmpty)
+					{
+						return FString::Printf(TEXT("%s::%s, UE_ARRAY_COUNT(%s::%s)"), *StaticsName, *Text, *StaticsName, *Text);
+					}
+					else
+					{
+						return FString(TEXT("nullptr, 0"));
+					}
+				};
+
+				// We add the hash to the name to generate a unique version of the name so that LiveCoding's default "only invoke static constructor once" 
+				// is avoided when elements change.
+				GeneratedCPPText.Logf(TEXT("\tstatic FRegisterCompiledInInfo %s%s_%u(TEXT(\"%s\"),\r\n\t\t%s,\r\n\t\t%s,\r\n\t\t%s);\r\n"),
+					Prefix,
+					*SourceFile.GetFileId(),
+					CombinedHash,
+					*PackageDefLcl.GetName(),
+					*FormatArray(ClassRegs.IsEmpty(), false, TEXT("ClassInfo")),
+					*FormatArray(ScriptStructRegs.IsEmpty(), false, TEXT("ScriptStructInfo")),
+					*FormatArray(EnumRegs.IsEmpty(), bAllEnumsEditorOnly, TEXT("EnumInfo"))
+				);
+			}
+
 			if (Singletons.Num())
 			{
 				SourceFile.GetSingletons().Append(MoveTemp(Singletons));
diff --git a/Engine/Source/Runtime/CoreUObject/Private/UObject/CoreNative.cpp b/Engine/Source/Runtime/CoreUObject/Private/UObject/CoreNative.cpp
--- a/Engine/Source/Runtime/CoreUObject/Private/UObject/CoreNative.cpp
+++ b/Engine/Source/Runtime/CoreUObject/Private/UObject/CoreNative.cpp
@@ -64,19 +64,18 @@ IMPLEMENT_MODULE( FCoreUObjectModule, CoreUObject );
 // if we are not using compiled in natives, we still need this as a base class for intrinsics
 #if !USE_COMPILED_IN_NATIVES
 COREUOBJECT_API class UClass* Z_Construct_UClass_UObject();
-extern FClassRegistrationInfo& Z_Registration_Info_UClass_UObject();
+extern FClassRegistrationInfo Z_Registration_Info_UClass_UObject;
 UClass* Z_Construct_UClass_UObject()
 {
-	static UClass* OuterClass = Z_Registration_Info_UClass_UObject().OuterSingleton;
-	if (!OuterClass)
+	if (!Z_Registration_Info_UClass_UObject.OuterSingleton)
 	{
-		OuterClass = UObject::StaticClass();
-		UObjectForceRegistration(OuterClass);
-		UObjectBase::EmitBaseReferences(OuterClass);
-		OuterClass->StaticLink();
+		Z_Registration_Info_UClass_UObject.OuterSingleton = UObject::StaticClass();
+		UObjectForceRegistration(Z_Registration_Info_UClass_UObject.OuterSingleton);
+		UObjectBase::EmitBaseReferences(Z_Registration_Info_UClass_UObject.OuterSingleton);
+		Z_Registration_Info_UClass_UObject.OuterSingleton->StaticLink();
 	}
-	check(OuterClass->GetClass());
-	return OuterClass;
+	check(Z_Registration_Info_UClass_UObject.OuterSingleton->GetClass());
+	return Z_Registration_Info_UClass_UObject.OuterSingleton;
 }
 IMPLEMENT_CLASS(UObject, 0);
 #endif
diff --git a/Engine/Source/Runtime/CoreUObject/Private/UObject/UObjectBase.cpp b/Engine/Source/Runtime/CoreUObject/Private/UObject/UObjectBase.cpp
--- a/Engine/Source/Runtime/CoreUObject/Private/UObject/UObjectBase.cpp
+++ b/Engine/Source/Runtime/CoreUObject/Private/UObject/UObjectBase.cpp
@@ -716,7 +716,7 @@ void UClassCompiledInDefer(FFieldCompiledInInfo* ClassInfo, const TCHAR* Name, S
 {
 	FClassDeferredRegistry& Registry = FClassDeferredRegistry::Get();
 	FClassRegistrationInfo& Info = Registry.MakeDeprecatedInfo(ClassInfo->ClassPackage(), Name);
-	Registry.AddRegistration(nullptr, nullptr, ClassInfo->ClassPackage(), Name, Info, CONSTUCT_RELOAD_VERSION_INFO(FClassReloadVersionInfo, ClassSize, Crc), ClassInfo);
+	Registry.AddRegistration(nullptr, nullptr, ClassInfo->ClassPackage(), Name, Info, CONSTRUCT_RELOAD_VERSION_INFO(FClassReloadVersionInfo, ClassSize, Crc), ClassInfo);
 }
 PRAGMA_ENABLE_DEPRECATION_WARNINGS
 
@@ -763,6 +763,32 @@ void RegisterCompiledInInfo(UPackage* (*InOuterRegister)(), const TCHAR* InPacka
 #endif
 }
 
+// Multiple registrations
+void RegisterCompiledInInfo(const TCHAR* PackageName, const FClassRegisterCompiledInInfo* ClassInfo, size_t NumClassInfo, const FStructRegisterCompiledInInfo* StructInfo, size_t NumStructInfo, const FEnumRegisterCompiledInInfo* EnumInfo, size_t NumEnumInfo)
+{
+	for (size_t Index = 0; Index < NumClassInfo; ++Index)
+	{
+		const FClassRegisterCompiledInInfo& Info = ClassInfo[Index];
+		RegisterCompiledInInfo(Info.OuterRegister, Info.InnerRegister, PackageName, Info.Name, *Info.Info, Info.VersionInfo);
+	}
+
+	for (size_t Index = 0; Index < NumStructInfo; ++Index)
+	{
+		const FStructRegisterCompiledInInfo& Info = StructInfo[Index];
+		RegisterCompiledInInfo(Info.OuterRegister, PackageName, Info.Name, *Info.Info, Info.VersionInfo);
+		if (Info.CreateCppStructOps != nullptr)
+		{
+			UScriptStruct::DeferCppStructOps(FName(Info.Name), (UScriptStruct::ICppStructOps*)Info.CreateCppStructOps());
+		}
+	}
+
+	for (size_t Index = 0; Index < NumEnumInfo; ++Index)
+	{
+		const FEnumRegisterCompiledInInfo& Info = EnumInfo[Index];
+		RegisterCompiledInInfo(Info.OuterRegister, PackageName, Info.Name, *Info.Info, Info.VersionInfo);
+	}
+}
+
 /** Register all loaded classes */
 void UClassRegisterAllCompiledInClasses()
 {
@@ -1216,7 +1242,7 @@ const TCHAR* DebugFullName(UObject* Object)
 UScriptStruct* FindExistingStructIfHotReloadOrDynamic(UObject* Outer, const TCHAR* StructName, SIZE_T Size, uint32 Crc, bool bIsDynamic)
 {
 #if WITH_RELOAD
-	UScriptStruct* Result = FStructDeferredRegistry::Get().FindMatchingObject(*Outer->GetName(), StructName, CONSTUCT_RELOAD_VERSION_INFO(FStructReloadVersionInfo, Size, Crc));
+	UScriptStruct* Result = FStructDeferredRegistry::Get().FindMatchingObject(*Outer->GetName(), StructName, CONSTRUCT_RELOAD_VERSION_INFO(FStructReloadVersionInfo, Size, Crc));
 #else
 	UScriptStruct* Result = nullptr;
 #endif
@@ -1230,7 +1256,7 @@ UScriptStruct* FindExistingStructIfHotReloadOrDynamic(UObject* Outer, const TCHA
 UEnum* FindExistingEnumIfHotReloadOrDynamic(UObject* Outer, const TCHAR* EnumName, SIZE_T Size, uint32 Crc, bool bIsDynamic)
 {
 #if WITH_RELOAD
-	UEnum* Result = FEnumDeferredRegistry::Get().FindMatchingObject(*Outer->GetName(), EnumName, CONSTUCT_RELOAD_VERSION_INFO(FEnumReloadVersionInfo, Crc));
+	UEnum* Result = FEnumDeferredRegistry::Get().FindMatchingObject(*Outer->GetName(), EnumName, CONSTRUCT_RELOAD_VERSION_INFO(FEnumReloadVersionInfo, Crc));
 #else
 	UEnum* Result = nullptr;
 #endif
diff --git a/Engine/Source/Runtime/CoreUObject/Private/UObject/UObjectGlobals.cpp b/Engine/Source/Runtime/CoreUObject/Private/UObject/UObjectGlobals.cpp
--- a/Engine/Source/Runtime/CoreUObject/Private/UObject/UObjectGlobals.cpp
+++ b/Engine/Source/Runtime/CoreUObject/Private/UObject/UObjectGlobals.cpp
@@ -4803,7 +4803,7 @@ namespace UECodeGen_Private
 	}
 #endif
 
-	void ConstructUFunction(UFunction*& OutFunction, const FFunctionParams& Params, UFunction** SingletonPtr)
+	FORCEINLINE void ConstructUFunctionInternal(UFunction*& OutFunction, const FFunctionParams& Params, UFunction** SingletonPtr)
 	{
 		UObject*   (*OuterFunc)() = Params.OuterFunc;
 		UFunction* (*SuperFunc)() = Params.SuperFunc;
@@ -4868,6 +4868,16 @@ namespace UECodeGen_Private
 		NewFunction->StaticLink();
 	}
 
+	void ConstructUFunction(UFunction*& OutFunction, const FFunctionParams& Params)
+	{
+		ConstructUFunctionInternal(OutFunction, Params, nullptr);
+	}
+
+	void ConstructUFunction(UFunction** SingletonPtr, const FFunctionParams& Params)
+	{
+		ConstructUFunctionInternal(*SingletonPtr, Params, SingletonPtr);
+	}
+
 	void ConstructUEnum(UEnum*& OutEnum, const FEnumParams& Params)
 	{
 		UObject* (*OuterFunc)() = Params.OuterFunc;
diff --git a/Engine/Source/Runtime/CoreUObject/Public/UObject/ObjectMacros.h b/Engine/Source/Runtime/CoreUObject/Public/UObject/ObjectMacros.h
--- a/Engine/Source/Runtime/CoreUObject/Public/UObject/ObjectMacros.h
+++ b/Engine/Source/Runtime/CoreUObject/Public/UObject/ObjectMacros.h
@@ -1568,9 +1568,9 @@ namespace UM
 #define DECLARE_FSTRUCTUREDARCHIVE_SERIALIZER( TClass, API ) virtual API void Serialize(FStructuredArchive::FRecord Record) override;
 
 #if WITH_RELOAD
-#define CONSTUCT_RELOAD_VERSION_INFO(VersionInfo, ...) VersionInfo { __VA_ARGS__ }
+#define CONSTRUCT_RELOAD_VERSION_INFO(VersionInfo, ...) VersionInfo { __VA_ARGS__ }
 #else
-#define CONSTUCT_RELOAD_VERSION_INFO(VersionInfo, ...) VersionInfo()
+#define CONSTRUCT_RELOAD_VERSION_INFO(VersionInfo, ...) VersionInfo()
 #endif
 
 /*-----------------------------------------------------------------------------
@@ -1744,25 +1744,19 @@ public: \
 #define DECLARE_WITHIN_UPACKAGE() \
 	DECLARE_WITHIN_INTERNAL( UPackage, true )
 
-// Register a class at startup time.
-#define IMPLEMENT_CLASS(TClass, TClassCrc) \
-	FClassRegistrationInfo& Z_Registration_Info_UClass_##TClass() \
-	{ \
-		static FClassRegistrationInfo info; \
-		return info; \
-	} \
-	/* Do not change the AutoInitialize_ without changing LC_SymbolPatterns */ \
-	static FRegisterCompiledInInfo AutoInitialize_##TClass(&Z_Construct_UClass_##TClass, TClass::StaticClass, TClass::StaticPackage(), TEXT(#TClass), Z_Registration_Info_UClass_##TClass(), CONSTUCT_RELOAD_VERSION_INFO(FClassReloadVersionInfo, sizeof(TClass), TClassCrc)); \
+// Implement the GetPrivateStaticClass and the registration info but do not auto register the class.  
+// This is primarily used by UnrealHeaderTool
+#define IMPLEMENT_CLASS_NO_AUTO_REGISTRATION(TClass) \
+	FClassRegistrationInfo Z_Registration_Info_UClass_##TClass; \
 	UClass* TClass::GetPrivateStaticClass() \
 	{ \
-		static UClass*& PrivateStaticClass = Z_Registration_Info_UClass_##TClass().InnerSingleton; \
-		if (!PrivateStaticClass) \
+		if (!Z_Registration_Info_UClass_##TClass.InnerSingleton) \
 		{ \
 			/* this could be handled with templates, but we want it external to avoid code bloat */ \
 			GetPrivateStaticClassBody( \
 				StaticPackage(), \
 				(TCHAR*)TEXT(#TClass) + 1 + ((StaticClassFlags & CLASS_Deprecated) ? 11 : 0), \
-				PrivateStaticClass, \
+				Z_Registration_Info_UClass_##TClass.InnerSingleton, \
 				StaticRegisterNatives##TClass, \
 				sizeof(TClass), \
 				alignof(TClass), \
@@ -1776,13 +1770,19 @@ public: \
 				&TClass::WithinClass::StaticClass \
 			); \
 		} \
-		return PrivateStaticClass; \
+		return Z_Registration_Info_UClass_##TClass.InnerSingleton; \
 	}
 
+// Register a class at startup time.
+#define IMPLEMENT_CLASS(TClass, TClassCrc) \
+	IMPLEMENT_CLASS_NO_AUTO_REGISTRATION(TClass) \
+	/* Do not change the AutoInitialize_ without changing LC_SymbolPatterns */ \
+	static FRegisterCompiledInInfo AutoInitialize_##TClass(&Z_Construct_UClass_##TClass, TClass::StaticClass, TClass::StaticPackage(), TEXT(#TClass), Z_Registration_Info_UClass_##TClass, CONSTRUCT_RELOAD_VERSION_INFO(FClassReloadVersionInfo, sizeof(TClass), TClassCrc));
+
 // Used for intrinsics, this sets up the boiler plate, plus an initialization singleton, which can create properties and GC tokens
 #define IMPLEMENT_INTRINSIC_CLASS(TClass, TRequiredAPI, TSuperClass, TSuperRequiredAPI, TPackage, InitCode) \
 	TRequiredAPI UClass* Z_Construct_UClass_##TClass(); \
-	extern FClassRegistrationInfo& Z_Registration_Info_UClass_##TClass(); \
+	extern FClassRegistrationInfo Z_Registration_Info_UClass_##TClass; \
 	struct Z_Construct_UClass_##TClass##_Statics \
 	{ \
 		static UClass* Construct() \
@@ -1799,13 +1799,12 @@ public: \
 	}; \
 	UClass* Z_Construct_UClass_##TClass() \
 	{ \
-		static UClass* Class = Z_Registration_Info_UClass_##TClass().OuterSingleton; \
-		if (!Class) \
+		if (!Z_Registration_Info_UClass_##TClass.OuterSingleton) \
 		{ \
-			Class = Z_Construct_UClass_##TClass##_Statics::Construct();\
+			Z_Registration_Info_UClass_##TClass.OuterSingleton = Z_Construct_UClass_##TClass##_Statics::Construct();\
 		} \
-		check(Class->GetClass()); \
-		return Class; \
+		check(Z_Registration_Info_UClass_##TClass.OuterSingleton->GetClass()); \
+		return Z_Registration_Info_UClass_##TClass.OuterSingleton; \
 	} \
 	IMPLEMENT_CLASS(TClass, 0)
 
diff --git a/Engine/Source/Runtime/CoreUObject/Public/UObject/UObjectBase.h b/Engine/Source/Runtime/CoreUObject/Public/UObject/UObjectBase.h
--- a/Engine/Source/Runtime/CoreUObject/Public/UObject/UObjectBase.h
+++ b/Engine/Source/Runtime/CoreUObject/Public/UObject/UObjectBase.h
@@ -292,6 +292,18 @@ struct FClassReloadVersionInfo
  */
 using FClassRegistrationInfo = TRegistrationInfo<UClass, FClassReloadVersionInfo>;
 
+/**
+ * Composite class register compiled in info
+ */
+struct FClassRegisterCompiledInInfo
+{
+	class UClass* (*OuterRegister)();
+	class UClass* (*InnerRegister)();
+	const TCHAR* Name;
+	FClassRegistrationInfo* Info;
+	FClassReloadVersionInfo VersionInfo;
+};
+
 /**
  * Adds a class registration and version information. The InInfo parameter must be static.
  */
@@ -318,6 +330,18 @@ struct FStructReloadVersionInfo
  */
 using FStructRegistrationInfo = TRegistrationInfo<UScriptStruct, FStructReloadVersionInfo>;
 
+/**
+ * Composite structures register compiled in info
+ */
+struct FStructRegisterCompiledInInfo
+{
+	class UScriptStruct* (*OuterRegister)();
+	void* (*CreateCppStructOps)();
+	const TCHAR* Name;
+	FStructRegistrationInfo* Info;
+	FStructReloadVersionInfo VersionInfo;
+};
+
 /**
  * Adds a struct registration and version information. The InInfo parameter must be static.
  */
@@ -354,6 +378,17 @@ struct FEnumReloadVersionInfo
  */
 using FEnumRegistrationInfo = TRegistrationInfo<UEnum, FEnumReloadVersionInfo>;
 
+/**
+ * Composite enumeration register compiled in info
+ */
+struct FEnumRegisterCompiledInInfo
+{
+	class UEnum* (*OuterRegister)();
+	const TCHAR* Name;
+	FEnumRegistrationInfo* Info;
+	FEnumReloadVersionInfo VersionInfo;
+};
+
 /**
  * Adds a static enum registration and version information. The InInfo parameter must be static.
  */
@@ -390,6 +425,12 @@ using FPackageRegistrationInfo = TRegistrationInfo<UPackage, FPackageReloadVersi
  */
 COREUOBJECT_API void RegisterCompiledInInfo(UPackage* (*InOuterRegister)(), const TCHAR* InPackageName, FPackageRegistrationInfo& InInfo, const FPackageReloadVersionInfo& InVersionInfo);
 
+
+/**
+ * Register compiled in information for multiple classes, structures, and enumerations
+ */
+COREUOBJECT_API void RegisterCompiledInInfo(const TCHAR* PackageName, const FClassRegisterCompiledInInfo* ClassInfo, size_t NumClassInfo, const FStructRegisterCompiledInInfo* StructInfo, size_t NumStructInfo, const FEnumRegisterCompiledInInfo* EnumInfo, size_t NumEnumInfo);
+
 /**
  * Base class for deferred native class registration
  */
diff --git a/Engine/Source/Runtime/CoreUObject/Public/UObject/UObjectGlobals.h b/Engine/Source/Runtime/CoreUObject/Public/UObject/UObjectGlobals.h
--- a/Engine/Source/Runtime/CoreUObject/Public/UObject/UObjectGlobals.h
+++ b/Engine/Source/Runtime/CoreUObject/Public/UObject/UObjectGlobals.h
@@ -3044,7 +3044,9 @@ namespace UECodeGen_Private
 #endif
 	};
 
-	COREUOBJECT_API void ConstructUFunction(UFunction*& OutFunction, const FFunctionParams& Params, UFunction** SingletonPtr = nullptr);
+	UE_DEPRECATED(5.0, "ConstructUFunction deprecated.  Please use the version of ConstructUFunction which retains the singleton pointer.")
+	COREUOBJECT_API void ConstructUFunction(UFunction*& OutFunction, const FFunctionParams& Params);
+	COREUOBJECT_API void ConstructUFunction(UFunction** SingletonPtr, const FFunctionParams& Params);
 	COREUOBJECT_API void ConstructUEnum(UEnum*& OutEnum, const FEnumParams& Params);
 	COREUOBJECT_API void ConstructUScriptStruct(UScriptStruct*& OutStruct, const FStructParams& Params);
 	COREUOBJECT_API void ConstructUPackage(UPackage*& OutPackage, const FPackageParams& Params);
