Fixing compile issue with generated headers that are sparse

#rnx

diff --git a/Engine/Source/Programs/UnrealHeaderTool/Private/CodeGenerator.cpp b/Engine/Source/Programs/UnrealHeaderTool/Private/CodeGenerator.cpp
--- a/Engine/Source/Programs/UnrealHeaderTool/Private/CodeGenerator.cpp
+++ b/Engine/Source/Programs/UnrealHeaderTool/Private/CodeGenerator.cpp
@@ -2154,7 +2154,7 @@ void FNativeClassHeaderGenerator::ExportNativeGeneratedInitCode(FOutputDevice& O
 		
 		for (const FString& SparseClassDataString : SparseClassDataTypes)
 		{
-			GeneratedClassRegisterFunctionText.Logf(TEXT("\t\t\tOuterClass->SetSparseClassDataStruct(F%s::StaticStruct());\r\n"), *SparseClassDataString);
+			GeneratedClassRegisterFunctionText.Logf(TEXT("\t\t\t%s->SetSparseClassDataStruct(F%s::StaticStruct());\r\n"), *OuterSingletonName, *SparseClassDataString);
 		}
 
 		if (bIsDynamic)
