Fix issue with UHT so it generates errors with duplicate file names again.

diff --git a/Engine/Source/Programs/UnrealHeaderTool/Private/ClassMaps.h b/Engine/Source/Programs/UnrealHeaderTool/Private/ClassMaps.h
--- a/Engine/Source/Programs/UnrealHeaderTool/Private/ClassMaps.h
+++ b/Engine/Source/Programs/UnrealHeaderTool/Private/ClassMaps.h
@@ -108,13 +108,14 @@ private:
 // Wrapper class around SourceFiles map so we can quickly get a list of source files for a given package
 struct FUnrealSourceFiles : public FFreezableContainer
 {
-	TSharedRef<FUnrealSourceFile>* AddByHash(uint32 Hash, FString&& Filename, TSharedRef<FUnrealSourceFile> SourceFile)
+	TSharedPtr<FUnrealSourceFile> AddByHash(uint32 Hash, FString&& Filename, TSharedRef<FUnrealSourceFile> SourceFile)
 	{
 		check(!bFrozen);
 		TSharedRef<FUnrealSourceFile>* Existing = SourceFilesByString.FindByHash(Hash, Filename);
+		TSharedPtr<FUnrealSourceFile> Return(Existing != nullptr ? TSharedPtr<FUnrealSourceFile>(*Existing) : TSharedPtr<FUnrealSourceFile>());
 		AllSourceFiles.Add(&SourceFile.Get());
 		SourceFilesByString.AddByHash(Hash, MoveTemp(Filename), MoveTemp(SourceFile));
-		return Existing;
+		return Return;
 	}
 	const TSharedRef<FUnrealSourceFile>* Find(const FString& Id) const 
 	{
diff --git a/Engine/Source/Programs/UnrealHeaderTool/Private/CodeGenerator.cpp b/Engine/Source/Programs/UnrealHeaderTool/Private/CodeGenerator.cpp
--- a/Engine/Source/Programs/UnrealHeaderTool/Private/CodeGenerator.cpp
+++ b/Engine/Source/Programs/UnrealHeaderTool/Private/CodeGenerator.cpp
@@ -6544,10 +6544,10 @@ void PrepareModules(TArray<FUnrealPackageDefinitionInfo*>& PackageDefs, const FS
 
 				FString CleanFilename = FPaths::GetCleanFilename(RawFilename);
 				uint32  CleanFilenameHash = GetTypeHash(CleanFilename);
-				if (const TSharedRef<FUnrealSourceFile>* ExistingSourceFile = GUnrealSourceFilesMap.AddByHash(CleanFilenameHash, MoveTemp(CleanFilename), UnrealSourceFile))
+				if (TSharedPtr<FUnrealSourceFile> ExistingSourceFile = GUnrealSourceFilesMap.AddByHash(CleanFilenameHash, MoveTemp(CleanFilename), UnrealSourceFile))
 				{
 					FString NormalizedFullFilename = FullFilename;
-					FString NormalizedExistingFilename = (*ExistingSourceFile)->GetFilename();
+					FString NormalizedExistingFilename = ExistingSourceFile->GetFilename();
 
 					FPaths::NormalizeFilename(NormalizedFullFilename);
 					FPaths::NormalizeFilename(NormalizedExistingFilename);
