UnrealHeaderTool crashes with Assertion when structure is missing the proper prefix.

#rnx

diff --git a/Engine/Source/Programs/UnrealHeaderTool/Private/HeaderParser.cpp b/Engine/Source/Programs/UnrealHeaderTool/Private/HeaderParser.cpp
--- a/Engine/Source/Programs/UnrealHeaderTool/Private/HeaderParser.cpp
+++ b/Engine/Source/Programs/UnrealHeaderTool/Private/HeaderParser.cpp
@@ -8240,7 +8240,7 @@ TSharedRef<FUnrealTypeDefinitionInfo> FHeaderPreParser::ParseStructDeclaration(c
 
 
 	// Create the structure definition
-	TSharedRef<FUnrealScriptStructDefinitionInfo> StructDef = MakeShareable(new FUnrealScriptStructDefinitionInfo(SourceFile, InLineNumber, MoveTemp(StructNameInScript), FName(StructNameStripped, FNAME_Add)));
+	TSharedRef<FUnrealScriptStructDefinitionInfo> StructDef = MakeShareable(new FUnrealScriptStructDefinitionInfo(SourceFile, InLineNumber, *StructNameInScript, FName(StructNameStripped, FNAME_Add)));
 
 	// Parse the inheritance list
 	ParseInheritance(TEXT("struct"), [this, &StructNameStripped, &StructDef = *StructDef](const TCHAR* StructName, bool bIsSuperClass)
@@ -8299,6 +8299,27 @@ TSharedRef<FUnrealTypeDefinitionInfo> FHeaderPreParser::ParseStructDeclaration(c
 		break;
 		}
 	}
+
+	// Check to make sure the syntactic native prefix was set-up correctly.
+	// If this check results in a false positive, it will be flagged as an identifier failure.
+	FString DeclaredPrefix = GetClassPrefix(StructNameInScript);
+	if (DeclaredPrefix == StructDef->GetPrefixCPP() || DeclaredPrefix == TEXT("T"))
+	{
+		// Found a prefix, do a basic check to see if it's valid
+		const TCHAR* ExpectedPrefixCPP = UHTConfig.StructsWithTPrefix.Contains(StructNameStripped) ? TEXT("T") : StructDef->GetPrefixCPP();
+		FString ExpectedStructName = FString::Printf(TEXT("%s%s"), ExpectedPrefixCPP, *StructNameStripped);
+		if (StructNameInScript != ExpectedStructName)
+		{
+			StructDef->Throwf(TEXT("Struct '%s' has an invalid Unreal prefix, expecting '%s'"), *StructNameInScript, *ExpectedStructName);
+		}
+	}
+	else
+	{
+		const TCHAR* ExpectedPrefixCPP = UHTConfig.StructsWithTPrefix.Contains(StructNameInScript) ? TEXT("T") : StructDef->GetPrefixCPP();
+		FString ExpectedStructName = FString::Printf(TEXT("%s%s"), ExpectedPrefixCPP, *StructNameInScript);
+		StructDef->Throwf(TEXT("Struct '%s' is missing a valid Unreal prefix, expecting '%s'"), *StructNameInScript, *ExpectedStructName);
+	}
+
 	return StructDef;
 }
 
