Fixed crash issue to provide better feedback as to the issue.

#rnx

diff --git a/Engine/Source/Programs/UnrealHeaderTool/Private/UnrealTypeDefinitionInfo.cpp b/Engine/Source/Programs/UnrealHeaderTool/Private/UnrealTypeDefinitionInfo.cpp
--- a/Engine/Source/Programs/UnrealHeaderTool/Private/UnrealTypeDefinitionInfo.cpp
+++ b/Engine/Source/Programs/UnrealHeaderTool/Private/UnrealTypeDefinitionInfo.cpp
@@ -2149,6 +2149,10 @@ void FUnrealClassDefinitionInfo::CreateUObjectEngineTypesInternal(ECreateEngineT
 	{
 	case ECreateEngineTypesPhase::Phase1:
 	{
+		if (ClassWithin == nullptr)
+		{
+			Throwf(TEXT("ClassWithin is not set.  This can be the result of the class's UCLASS not being parsed due to statement skipping preceding the UCLASS keyword."));
+		}
 		check(ClassWithin);
 		if (GetObjectSafe() == nullptr)
 		{
@@ -2244,6 +2248,10 @@ void FUnrealClassDefinitionInfo::PostParseFinalizeInternal(EPostParseFinalizePha
 	switch (Phase)
 	{
 	case EPostParseFinalizePhase::Phase1:
+		if (ClassWithin == nullptr)
+		{
+			Throwf(TEXT("ClassWithin is not set.  This can be the result of the class's UCLASS not being parsed due to statement skipping preceding the UCLASS keyword."));
+		}
 		break;
 
 	case EPostParseFinalizePhase::Phase2:
