Temporary fix to a long standing UHT issue made worse by the concurrency changes.  If an intermediate file in an include chain doesn't contain any U definitions, it will be omited by UBT causing UHT to have gaps in the include file graph.  This results in some delegates not being able to be found.

#rnx

diff --git a/Engine/Source/Programs/UnrealHeaderTool/Private/CodeGenerator.cpp b/Engine/Source/Programs/UnrealHeaderTool/Private/CodeGenerator.cpp
--- a/Engine/Source/Programs/UnrealHeaderTool/Private/CodeGenerator.cpp
+++ b/Engine/Source/Programs/UnrealHeaderTool/Private/CodeGenerator.cpp
@@ -64,6 +64,7 @@
 
 #include "UObject/WeakFieldPtr.h"
 #include "Templates/SubclassOf.h"
+#include <atomic>
 
 /////////////////////////////////////////////////////
 // Globals
@@ -88,6 +89,13 @@ FUnrealClassDefinitionInfo* GUObjectDef = nullptr;
 FUnrealClassDefinitionInfo* GUClassDef = nullptr;
 FUnrealClassDefinitionInfo* GUInterfaceDef = nullptr;
 
+// Busy wait support for holes in the include graph issues
+std::atomic<bool> GSourcesConcurrent = false;
+std::atomic<int> GSourcesToParse = 0;
+std::atomic<int> GSourcesParsing = 0;
+std::atomic<int> GSourcesCompleted = 0;
+std::atomic<int> GSourcesStalled = 0;
+
 struct FFindDelcarationResults
 {
 	bool bVirtualFound = false;
@@ -6861,7 +6869,10 @@ void ParseSourceFiles(TArray<FUnrealSourceFile*>& OrderedSourceFiles)
 	// Disable loading of objects outside of this package (or more exactly, objects which aren't UFields, CDO, or templates)
 	TGuardValue<bool> AutoRestoreVerifyObjectRefsFlag(GVerifyObjectReferencesOnly, true);
 
+	GSourcesToParse = (int)OrderedSourceFiles.Num();
+	GSourcesConcurrent = UHT_ENABLE_CONCURRENT_PARSING != 0;
 #if UHT_ENABLE_CONCURRENT_PARSING
+
 	/**
 	 * For every FUnrealSourceFile being processed, an instance of this class represents the data associated with generating the new output.
 	 */
@@ -6911,10 +6922,12 @@ void ParseSourceFiles(TArray<FUnrealSourceFile*>& OrderedSourceFiles)
 
 		auto ParseSource = [&ParsedCPP]()
 		{
+			++GSourcesParsing;
 			FResults::TryAlways([&ParsedCPP]()
 			{
 				FHeaderParser::Parse(ParsedCPP.PackageDef, ParsedCPP.SourceFile);
 			});
+			++GSourcesCompleted;
 		};
 
 		Includes.Reset();
@@ -6946,8 +6959,10 @@ void ParseSourceFiles(TArray<FUnrealSourceFile*>& OrderedSourceFiles)
 	{
 		FUnrealPackageDefinitionInfo& PackageDef = SourceFile->GetPackageDef();
 		FScopedDurationTimer SourceTimer(SourceFile->GetTime(ESourceFileTime::Parse));
+		++GSourcesParsing;
 		FResults::TryAlways([&PackageDef, SourceFile]() { FHeaderParser::Parse(PackageDef, *SourceFile); });
-	}
+		++GSourcesCompleted;
+}
 #endif
 }
 
diff --git a/Engine/Source/Programs/UnrealHeaderTool/Private/HeaderParser.cpp b/Engine/Source/Programs/UnrealHeaderTool/Private/HeaderParser.cpp
--- a/Engine/Source/Programs/UnrealHeaderTool/Private/HeaderParser.cpp
+++ b/Engine/Source/Programs/UnrealHeaderTool/Private/HeaderParser.cpp
@@ -25,6 +25,9 @@
 #include "Algo/Find.h"
 #include "Algo/FindSortedStringCaseInsensitive.h"
 #include "Misc/ScopeExit.h"
+#include "Misc/Timespan.h"
+#include "HAL/PlatformTime.h"
+#include <atomic>
 
 #include "Specifiers/CheckedMetadataSpecifiers.h"
 #include "Specifiers/EnumSpecifiers.h"
@@ -82,6 +85,13 @@ const FName FHeaderParserNames::NAME_PrioritizeCategories(TEXT("PrioritizeCatego
 TArray<FString> FHeaderParser::PropertyCPPTypesRequiringUIRanges = { TEXT("float"), TEXT("double") };
 TArray<FString> FHeaderParser::ReservedTypeNames = { TEXT("none") };
 
+// Busy wait support for holes in the include graph issues
+extern std::atomic<bool> GSourcesConcurrent;
+extern std::atomic<int> GSourcesToParse;
+extern std::atomic<int> GSourcesParsing;
+extern std::atomic<int> GSourcesCompleted;
+extern std::atomic<int> GSourcesStalled;
+
 /*-----------------------------------------------------------------------------
 	Utility functions.
 -----------------------------------------------------------------------------*/
@@ -681,6 +691,76 @@ namespace
 		TEXT("LAYOUT_FIELD_EDITORONLY"),
 		TEXT("LAYOUT_FIELD_INITIALIZED"),
 	};
+
+	FCriticalSection GlobalDelegatesLock;
+	TMap<FString, FUnrealFunctionDefinitionInfo*> GlobalDelegates;
+
+	template <typename LambdaType>
+	bool BusyWait(LambdaType&& Lambda)
+	{
+		constexpr double TimeoutInSeconds = 5.0;
+
+		// Try it first
+		if (Lambda())
+		{
+			return true;
+		}
+
+		// If we aren't doing concurrent processing, then there isn't any reason to continue, we will never make any progress.
+		if (!GSourcesConcurrent)
+		{
+			return false;
+		}
+
+		// Mark that we are stalled
+		ON_SCOPE_EXIT
+		{
+			--GSourcesStalled;
+		};
+		++GSourcesStalled;
+
+		int SourcesToParse = GSourcesToParse;
+		int SourcesParsing = GSourcesParsing;
+		int SourcesCompleted = GSourcesCompleted;
+		int SourcesStalled = GSourcesStalled;
+		double StartTime = FPlatformTime::Seconds();
+		for (;;)
+		{
+			FPlatformProcess::YieldCycles(10000);
+			if (Lambda())
+			{
+				return true;
+			}
+
+			int NewSourcesParsing = GSourcesParsing;
+			int NewSourcesCompleted = GSourcesCompleted;
+			int NewSourcesStalled = GSourcesStalled;
+
+			// If everything is stalled, then we have no hope
+			if (NewSourcesStalled + NewSourcesCompleted == SourcesToParse)
+			{
+				return false;
+			}
+
+			// If something has completed, then reset the timeout
+			double CurrentTime = FPlatformTime::Seconds();
+			if (NewSourcesCompleted != SourcesCompleted || NewSourcesParsing != SourcesParsing)
+			{
+				StartTime = CurrentTime;
+			}
+
+			// Otherwise, check for long timeout
+			else if (CurrentTime - StartTime > TimeoutInSeconds)
+			{
+				return false;
+			}
+
+			// Remember state
+			SourcesParsing = NewSourcesParsing;
+			SourcesCompleted = NewSourcesCompleted;
+			SourcesStalled = NewSourcesStalled;
+		}
+	}
 }
 
 void AddEditInlineMetaData(TMap<FName, FString>& MetaData)
@@ -2204,7 +2284,8 @@ void FHeaderParser::FixupDelegateProperties(FUnrealStructDefinitionInfo& StructD
 				}
 				if (SourceDelegateFunctionDef == nullptr)
 				{
-					SourceDelegateFunctionDef = UHTCast<FUnrealFunctionDefinitionInfo>(GTypeDefinitionInfoMap.FindByName(*NameOfDelegateFunction));
+					FScopeLock Lock(&GlobalDelegatesLock);
+					SourceDelegateFunctionDef = GlobalDelegates.FindRef(NameOfDelegateFunction);
 				}
 				if (SourceDelegateFunctionDef == nullptr)
 				{
@@ -3822,6 +3903,14 @@ void FHeaderParser::GetVarType(
 			{
 				if (FUnrealClassDefinitionInfo* LocalOwnerClassDef = FUnrealClassDefinitionInfo::FindClass(*IdentifierStripped))
 				{
+					if (!BusyWait([LocalOwnerClassDef]() 
+						{
+							return LocalOwnerClassDef->IsParsed() || LocalOwnerClassDef->GetUnrealSourceFile().IsParsed();
+						}))
+					{
+						Throwf(TEXT("Timeout waiting for '%s' to be parsed.  Make sure that '%s' directly references the include file or that all intermediate include files contain some type of UCLASS, USTRUCT, or UENUM."),
+							*LocalOwnerClassDef->GetUnrealSourceFile().GetFilename(), *SourceFile.GetFilename());
+					}
 					const FString DelegateIdentifierStripped = GetClassNameWithPrefixRemoved(FString(DelegateName.Value));
 					if (FUnrealFieldDefinitionInfo* FoundDef = LocalOwnerClassDef->GetScope()->FindTypeByName(*(DelegateIdentifierStripped + HEADER_GENERATED_DELEGATE_SIGNATURE_SUFFIX)))
 					{
@@ -4082,14 +4171,25 @@ void FHeaderParser::GetVarType(
 			// Resolve delegates declared in another class  //@TODO: UCREMOVAL: This seems extreme
 			if (!bHandledType)
 			{
-				if (FUnrealFunctionDefinitionInfo* ExternDelegateFuncDef = GTypeDefinitionInfoMap.FindByName<FUnrealFunctionDefinitionInfo>(*(IdentifierStripped + HEADER_GENERATED_DELEGATE_SIGNATURE_SUFFIX)))
-				{
-					SetDelegateType(*ExternDelegateFuncDef, IdentifierStripped);
-				}
+				FString FullName = IdentifierStripped + HEADER_GENERATED_DELEGATE_SIGNATURE_SUFFIX;
+				bHandledType = BusyWait([&SetDelegateType, &IdentifierStripped, &FullName]()
+					{
+						FUnrealFunctionDefinitionInfo* DelegateDef = nullptr;
+						{
+							FScopeLock Lock(&GlobalDelegatesLock);
+							DelegateDef = GlobalDelegates.FindRef(FullName);
+						}
+						if (DelegateDef)
+						{
+							SetDelegateType(*DelegateDef, IdentifierStripped);
+							return true;
+						}
+						return false;
+					});
 
 				if (!bHandledType)
 				{
-					Throwf(TEXT("Unrecognized type '%s' - type must be a UCLASS, USTRUCT or UENUM"), *VarType.GetTokenValue());
+					Throwf(TEXT("Unrecognized type '%s' - type must be a UCLASS, USTRUCT, UENUM, or global delegate."), *VarType.GetTokenValue());
 				}
 			}
 		}
@@ -9002,6 +9102,8 @@ FUnrealFunctionDefinitionInfo& FHeaderParser::CreateFunction(const TCHAR* FuncNa
 	else
 	{
 		SourceFile.AddDefinedFunction(FuncDefRef);
+		FScopeLock Lock(&GlobalDelegatesLock);
+		GlobalDelegates.Add(FuncDef.GetName(), &FuncDef);
 	}
 
 	return FuncDef;
diff --git a/Engine/Source/Programs/UnrealHeaderTool/Private/UnrealSourceFile.h b/Engine/Source/Programs/UnrealHeaderTool/Private/UnrealSourceFile.h
--- a/Engine/Source/Programs/UnrealHeaderTool/Private/UnrealSourceFile.h
+++ b/Engine/Source/Programs/UnrealHeaderTool/Private/UnrealSourceFile.h
@@ -8,6 +8,7 @@
 #include "HeaderProvider.h"
 #include "UnrealTypeDefinitionInfo.h"
 #include "GeneratedCodeVersion.h"
+#include <atomic>
 
 class UPackage;
 class FArchive;
@@ -476,6 +477,22 @@ public:
 		return bIsNoExportTypes;
 	}
 
+	/**
+	 * Return true if the source has completed parsing 
+	 */
+	bool IsParsed() const
+	{
+		return bParsed.load();
+	}
+
+	/**
+	 * Mark the source as completed parsing
+	 */
+	void MarkParsed() 
+	{
+		bParsed.store(false);
+	}
+
 private:
 
 	// File scope.
@@ -532,6 +549,9 @@ private:
 	// True if this is the NoExportTypes.h file
 	bool bIsNoExportTypes = false;
 
+	// True if this source file has completed parsing
+	std::atomic<bool> bParsed = false;
+
 	// Current topological sort state
 	ETopologicalState TopologicalState = ETopologicalState::Unmarked;
 
diff --git a/Engine/Source/Programs/UnrealHeaderTool/Private/UnrealTypeDefinitionInfo.h b/Engine/Source/Programs/UnrealHeaderTool/Private/UnrealTypeDefinitionInfo.h
--- a/Engine/Source/Programs/UnrealHeaderTool/Private/UnrealTypeDefinitionInfo.h
+++ b/Engine/Source/Programs/UnrealHeaderTool/Private/UnrealTypeDefinitionInfo.h
@@ -3474,7 +3474,7 @@ private:
 	bool bCustomVTableHelperConstructorDeclared = false;
 
 	/** True if the class definition has been parsed */
-	bool bParsed = false;
+	std::atomic<bool> bParsed = false;
 
 	/** True if the class has a custom constructor */
 	bool bCustomConstructor = false;
