Making it clear (by way of UHT error) that a subclass of a not-placeable class isn't allowed to be declared with the 'placeable' specifier.

diff --git a/Engine/Source/Programs/UnrealHeaderTool/Private/UnrealTypeDefinitionInfo.cpp b/Engine/Source/Programs/UnrealHeaderTool/Private/UnrealTypeDefinitionInfo.cpp
--- a/Engine/Source/Programs/UnrealHeaderTool/Private/UnrealTypeDefinitionInfo.cpp
+++ b/Engine/Source/Programs/UnrealHeaderTool/Private/UnrealTypeDefinitionInfo.cpp
@@ -2262,9 +2262,9 @@ void FUnrealClassDefinitionInfo::PostParseFinalizeInternal(EPostParseFinalizePha
 		{
 			if ((InheritClassFlags & CLASS_NotPlaceable) == 0)
 			{
-				if (!SuperClassDef->HasAnyClassFlags(CLASS_NotPlaceable))
+				if (SuperClassDef->HasAnyClassFlags(CLASS_NotPlaceable))
 				{
-					Throwf(TEXT("The 'placeable' specifier is only allowed on classes which have a base class that's marked as not placeable. Classes are assumed to be placeable by default."));
+					Throwf(TEXT("The 'placeable' specifier cannot override a 'nonplaceable' base class. Classes are assumed to be placeable by default. Consider whether using the 'abstract' specifier on the base class would work."));
 				}
 			}
 
