Properties that use a TSet with a struct will generate a compile time error if the struct doesn't have a hash value function.

diff --git a/Engine/Source/Programs/UnrealHeaderTool/Private/CodeGenerator.cpp b/Engine/Source/Programs/UnrealHeaderTool/Private/CodeGenerator.cpp
--- a/Engine/Source/Programs/UnrealHeaderTool/Private/CodeGenerator.cpp
+++ b/Engine/Source/Programs/UnrealHeaderTool/Private/CodeGenerator.cpp
@@ -1414,6 +1414,12 @@ void FNativeClassHeaderGenerator::OutputProperty(FOutputDevice& DeclOut, FOutput
 
 		DeclOut.Logf(TEXT("%sstatic const UECodeGen_Private::FSetPropertyParams %s;\r\n"), DeclSpaces, *NameWithoutScope);
 
+		if (ValueDef.IsStructOrStructStaticArray())
+		{
+			const FString& StructName = ValueDef.GetPropertyBase().ScriptStructDef->GetNameCPP();
+			Out.Logf(TEXT("%sstatic_assert(TModels<CGetTypeHashable, %s>::Value, \"The structure '%s' is used in a TSet but does not have a GetValueTypeHash defined\");\r\n"), Spaces, *StructName, *StructName);
+		}
+
 		Out.Logf(
 			TEXT("%sconst UECodeGen_Private::FSetPropertyParams %s = { %s, %s, (EPropertyFlags)0x%016llx, UECodeGen_Private::EPropertyGenFlags::Set, %s, %s, %s, %s };%s\r\n"),
 			Spaces,
