PR #8049: Trim spaces in ExpandEnumAsExecs meta parameter (Contributed by [REDACTED_PERSON])

diff --git a/Engine/Source/Editor/BlueprintGraph/Private/K2Node_CallFunction.cpp b/Engine/Source/Editor/BlueprintGraph/Private/K2Node_CallFunction.cpp
--- a/Engine/Source/Editor/BlueprintGraph/Private/K2Node_CallFunction.cpp
+++ b/Engine/Source/Editor/BlueprintGraph/Private/K2Node_CallFunction.cpp
@@ -956,8 +956,10 @@ void UK2Node_CallFunction::GetExpandEnumPinNames(const UFunction* Function, TArr
 
 	TArray<FString> RawGroupings;
 	EnumParamString.ParseIntoArray(RawGroupings, TEXT(","), false);
-	for (const FString& RawGroup : RawGroupings)
+	for (FString& RawGroup : RawGroupings)
 	{
+		RawGroup.TrimStartAndEndInline();
+
 		TArray<FString> IndividualEntries;
 		RawGroup.ParseIntoArray(IndividualEntries, TEXT("|"));
 
diff --git a/Engine/Source/Programs/UnrealHeaderTool/Private/HeaderParser.cpp b/Engine/Source/Programs/UnrealHeaderTool/Private/HeaderParser.cpp
--- a/Engine/Source/Programs/UnrealHeaderTool/Private/HeaderParser.cpp
+++ b/Engine/Source/Programs/UnrealHeaderTool/Private/HeaderParser.cpp
@@ -1160,13 +1160,16 @@ namespace
 					InValue.ParseIntoArray(RawGroupings, TEXT(","), false);
 
 					FProperty* FirstInput = nullptr;
-					for (const FString& RawGroup : RawGroupings)
+					for (FString& RawGroup : RawGroupings)
 					{
+						RawGroup.TrimStartAndEndInline();
+
 						TArray<FString> IndividualEntries;
 						RawGroup.ParseIntoArray(IndividualEntries, TEXT("|"));
 
-						for (const FString& Entry : IndividualEntries)
+						for (FString& Entry : IndividualEntries)
 						{
+							Entry.TrimStartAndEndInline();
 							if (Entry.IsEmpty())
 							{
 								continue;
