UHT: Remove suppression of a warning that was required for BP nativization (C4883: function size suppresses optimizations)

#rnx

diff --git a/Engine/Source/Programs/UnrealHeaderTool/Private/CodeGenerator.cpp b/Engine/Source/Programs/UnrealHeaderTool/Private/CodeGenerator.cpp
--- a/Engine/Source/Programs/UnrealHeaderTool/Private/CodeGenerator.cpp
+++ b/Engine/Source/Programs/UnrealHeaderTool/Private/CodeGenerator.cpp
@@ -5177,9 +5177,6 @@ bool FNativeClassHeaderGenerator::WriteSource(const FManifestModule& Module, FGe
 		}
 	}
 
-	static const TCHAR* DisableWarning4883 = TEXT("#ifdef _MSC_VER" LINE_TERMINATOR_ANSI "#pragma warning (push)" LINE_TERMINATOR_ANSI "#pragma warning (disable : 4883)" LINE_TERMINATOR_ANSI "#endif" LINE_TERMINATOR_ANSI);
-	static const TCHAR* EnableWarning4883 = TEXT("#ifdef _MSC_VER" LINE_TERMINATOR_ANSI "#pragma warning (pop)" LINE_TERMINATOR_ANSI "#endif" LINE_TERMINATOR_ANSI);
-
 	FUHTStringBuilder FileText;
 	FileText.Log(HeaderCopyright);
 	FileText.Log(RequiredCPPIncludes);
@@ -5188,7 +5185,6 @@ bool FNativeClassHeaderGenerator::WriteSource(const FManifestModule& Module, FGe
 	{
 		FileText.Logf(TEXT("#include \"%s\"\r\n"), *RelativeInclude);
 	}
-	FileText.Log(DisableWarning4883);
 	FileText.Log(DisableDeprecationWarnings);
 
 	FString CleanFilename = FPaths::GetCleanFilename(FileInfo.GetFilename());
@@ -5207,7 +5203,6 @@ bool FNativeClassHeaderGenerator::WriteSource(const FManifestModule& Module, FGe
 	}
 	FileText.Log(*InBodyText);
 	FileText.Log(EnableDeprecationWarnings);
-	FileText.Log(EnableWarning4883);
 
 	return SaveHeaderIfChanged(FileInfo, MoveTemp(FileText));
 }
