Added temporary code to log all files when UHT fails.
Tim.Smith

diff --git a/Engine/Source/Programs/UnrealHeaderTool/Private/CodeGenerator.cpp b/Engine/Source/Programs/UnrealHeaderTool/Private/CodeGenerator.cpp
--- a/Engine/Source/Programs/UnrealHeaderTool/Private/CodeGenerator.cpp
+++ b/Engine/Source/Programs/UnrealHeaderTool/Private/CodeGenerator.cpp
@@ -7019,6 +7019,23 @@ ECompilationResult::Type UnrealHeaderTool_Main(const FString& ModuleInfoFilename
 	// Finish all async file tasks before stopping the clock
 	FTaskGraphInterface::Get().WaitUntilTasksComplete(GAsyncFileTasks);
 
+	// TEMPORARY change to log all files when UHT fails
+	if (FResults::GetOverallResults() != 0)
+	{
+		UE_LOG(LogCompile, Log, TEXT("Source file listing due to UHT detected errors:"));
+		for (FUnrealPackageDefinitionInfo* PackageDef : PackageDefs)
+		{
+			const FManifestModule& Module = PackageDef->GetModule();
+
+			UE_LOG(LogCompile, Log, TEXT("Package %s sources"), *Module.Name);
+			TArray<TSharedRef<FUnrealSourceFile>>& SourceFiles = PackageDef->GetAllSourceFiles();
+			for (TSharedRef<FUnrealSourceFile>& SourceFile : SourceFiles)
+			{
+				UE_LOG(LogCompile, Log, TEXT("---- %s"), *SourceFile->GetFilename());
+			}
+		}
+	}
+
 	double TotalShutdownTime = FResults::TimedTry([]() { GTypeDefinitionInfoMap.Reset(); GUnrealSourceFilesMap.Reset(); });
 
 	MainTimer.Stop();
