If a ShowCategories overrides a parent's HideCategories, the correct list of categories will now be included in the meta data.

diff --git a/Engine/Source/Programs/UnrealHeaderTool/Private/UnrealTypeDefinitionInfo.cpp b/Engine/Source/Programs/UnrealHeaderTool/Private/UnrealTypeDefinitionInfo.cpp
--- a/Engine/Source/Programs/UnrealHeaderTool/Private/UnrealTypeDefinitionInfo.cpp
+++ b/Engine/Source/Programs/UnrealHeaderTool/Private/UnrealTypeDefinitionInfo.cpp
@@ -2833,15 +2833,32 @@ void FUnrealClassDefinitionInfo::SetAndValidateWithinClass()
 	}
 }
 
-void FUnrealClassDefinitionInfo::MergeCategoryMetaData(TMap<FName, FString>& InMetaData) const
+void FUnrealClassDefinitionInfo::MergeCatagoryMetaData(TMap<FName, FString>& InMetaData, FName InName, const TArray<FString>& InNames)
 {
+	// The meta data copied from the super is already on the class.  If the new collection is empty, then we must remove the
+	// meta data from the class existing meta data (which is from the parent).  
+	if (InNames.Num() > 0)
+	{
+		InMetaData.Add(InName, FString::Join(InNames, TEXT(" ")));
+	}
+	else
+	{
+		GetMetaDataMap().Remove(InName);
+	}
+}
+
+void FUnrealClassDefinitionInfo::MergeCategoryMetaData(TMap<FName, FString>& InMetaData)
+{
+	// These collections are merged by hand and if the list is empty, then we must remove anything from the parent
+	MergeCatagoryMetaData(InMetaData, FHeaderParserNames::NAME_AutoExpandCategories, AutoExpandCategories);
+	MergeCatagoryMetaData(InMetaData, FHeaderParserNames::NAME_AutoCollapseCategories, AutoCollapseCategories);
+	MergeCatagoryMetaData(InMetaData, FHeaderParserNames::NAME_HideCategories, HideCategories);
+	MergeCatagoryMetaData(InMetaData, FHeaderParserNames::NAME_ShowCategories, ShowSubCatgories);
+	MergeCatagoryMetaData(InMetaData, FHeaderParserNames::NAME_HideFunctions, HideFunctions);
+
+	// These collections overwrite the parent ONLY if they don't have any elements.
 	if (ClassGroupNames.Num()) { InMetaData.Add(NAME_ClassGroupNames, FString::Join(ClassGroupNames, TEXT(" "))); }
-	if (AutoCollapseCategories.Num()) { InMetaData.Add(FHeaderParserNames::NAME_AutoCollapseCategories, FString::Join(AutoCollapseCategories, TEXT(" "))); }
-	if (HideCategories.Num()) { InMetaData.Add(FHeaderParserNames::NAME_HideCategories, FString::Join(HideCategories, TEXT(" "))); }
-	if (ShowSubCatgories.Num()) { InMetaData.Add(FHeaderParserNames::NAME_ShowCategories, FString::Join(ShowSubCatgories, TEXT(" "))); }
 	if (SparseClassDataTypes.Num()) { InMetaData.Add(FHeaderParserNames::NAME_SparseClassDataTypes, FString::Join(SparseClassDataTypes, TEXT(" "))); }
-	if (HideFunctions.Num()) { InMetaData.Add(FHeaderParserNames::NAME_HideFunctions, FString::Join(HideFunctions, TEXT(" "))); }
-	if (AutoExpandCategories.Num()) { InMetaData.Add(FHeaderParserNames::NAME_AutoExpandCategories, FString::Join(AutoExpandCategories, TEXT(" "))); }
 	if (PrioritizeCategories.Num()) { InMetaData.Add(FHeaderParserNames::NAME_PrioritizeCategories, FString::Join(PrioritizeCategories, TEXT(" "))); }
 }
 
diff --git a/Engine/Source/Programs/UnrealHeaderTool/Private/UnrealTypeDefinitionInfo.h b/Engine/Source/Programs/UnrealHeaderTool/Private/UnrealTypeDefinitionInfo.h
--- a/Engine/Source/Programs/UnrealHeaderTool/Private/UnrealTypeDefinitionInfo.h
+++ b/Engine/Source/Programs/UnrealHeaderTool/Private/UnrealTypeDefinitionInfo.h
@@ -3078,7 +3078,9 @@ public:
 	/**
 	 * Add the category meta data
 	 */
-	void MergeCategoryMetaData(TMap<FName, FString>& InMetaData) const;
+	void MergeCategoryMetaData(TMap<FName, FString>& InMetaData);
+
+	void MergeCatagoryMetaData(TMap<FName, FString>& InMetaData, FName InName, const TArray<FString>& InNames);
 
 	void GetSparseClassDataTypes(TArray<FString>& OutSparseClassDataTypes) const;
 
