Improved assertion output for FHeaderParser to provide the name of the struct/class being compiled and a possible user error that might have caused the problem (missing a prefix)

diff --git a/Engine/Source/Programs/UnrealHeaderTool/Private/HeaderParser.cpp b/Engine/Source/Programs/UnrealHeaderTool/Private/HeaderParser.cpp
--- a/Engine/Source/Programs/UnrealHeaderTool/Private/HeaderParser.cpp
+++ b/Engine/Source/Programs/UnrealHeaderTool/Private/HeaderParser.cpp
@@ -8207,6 +8207,11 @@ TSharedRef<FUnrealTypeDefinitionInfo> FHeaderPreParser::ParseClassDeclaration(co
 
 	FString ClassNameWithoutPrefixStr = GetClassNameWithPrefixRemoved(ClassName);
 
+	if (ClassNameWithoutPrefixStr.IsEmpty())
+	{
+		Throwf(TEXT("When compiling class definition for '%s', attempting to strip prefix results in an empty name. Did you leave off a prefix?"), *ClassNameWithoutPrefixStr);
+	}
+
 	// Skip optional final keyword
 	MatchIdentifier(TEXT("final"), ESearchCase::CaseSensitive);
 
@@ -8349,6 +8354,11 @@ TSharedRef<FUnrealTypeDefinitionInfo> FHeaderPreParser::ParseStructDeclaration(c
 	// The struct name stripped of it's prefix
 	FString StructNameStripped = GetClassNameWithPrefixRemoved(StructNameInScript);
 
+	if (StructNameStripped.IsEmpty())
+	{
+		Throwf(TEXT("When compiling struct definition for '%s', attempting to strip prefix results in an empty name. Did you leave off a prefix?"), *StructNameInScript);
+	}
+
 
 	// Create the structure definition
 	TSharedRef<FUnrealScriptStructDefinitionInfo> StructDef = MakeShareable(new FUnrealScriptStructDefinitionInfo(SourceFile, InLineNumber, *StructNameInScript, FName(StructNameStripped, FNAME_Add)));
