UHT will now warn the user if any form of GENERATED_BODY is used in code being skipped.

diff --git a/Engine/Source/Programs/UnrealHeaderTool/Private/HeaderParser.cpp b/Engine/Source/Programs/UnrealHeaderTool/Private/HeaderParser.cpp
--- a/Engine/Source/Programs/UnrealHeaderTool/Private/HeaderParser.cpp
+++ b/Engine/Source/Programs/UnrealHeaderTool/Private/HeaderParser.cpp
@@ -5150,6 +5150,23 @@ bool FHeaderParser::CheckForSerialize(FUnrealStructDefinitionInfo& StructDef, co
 	
 	return false;
 }
+static const TArray<FStringView> GSkipDeclarationWarningStrings =
+{
+	TEXT("GENERATED_BODY"),
+	TEXT("GENERATED_IINTERFACE_BODY"),
+	TEXT("GENERATED_UCLASS_BODY"),
+	TEXT("GENERATED_UINTERFACE_BODY"),
+	TEXT("GENERATED_USTRUCT_BODY"),
+	// Leaving these disabled ATM since they can exist in the code without causing compile issues
+	//TEXT("RIGVM_METHOD"),
+	//TEXT("UCLASS"),
+	//TEXT("UDELEGATE"),
+	//TEXT("UENUM"),
+	//TEXT("UFUNCTION"),
+	//TEXT("UINTERFACE"),
+	//TEXT("UPROPERTY"),
+	//TEXT("USTRUCT"),
+};
 
 bool FHeaderParser::SkipDeclaration(FToken& Token)
 {
@@ -5183,6 +5200,19 @@ bool FHeaderParser::SkipDeclaration(FToken& Token)
 			break;
 		}
 
+		if (Token.IsIdentifier())
+		{
+			// Use a trivial prefilter to avoid doing the search on things that aren't UE keywords we care about
+			if (Token.Value[0] == 'G' || Token.Value[0] == 'R' || Token.Value[0] == 'U')
+			{
+				if (Algo::BinarySearch(GSkipDeclarationWarningStrings, Token.Value, 
+					[](const FStringView& Lhs, const FStringView& Rhs) { return Lhs.Compare(Rhs, ESearchCase::CaseSensitive) < 0; }) >= 0)
+				{
+					LogWarning(FString::Printf(TEXT("The identifier \'%s\' was detected in a block being skipped. Was this intentional?"), *FString(Token.Value)));
+				}
+			}
+		}
+
 		if (!bMacroDeclaration && Token.IsIdentifier(TEXT("PURE_VIRTUAL"), ESearchCase::CaseSensitive) && NestedScopes == 0)
 		{
 			OpeningBracket = TEXT('(');
