UnrealHeaderTool: emptying early a global array that stores tasks as their allocator can be already destroyed when the array gets statically destructed

diff --git a/Engine/Source/Programs/UnrealHeaderTool/Private/CodeGenerator.cpp b/Engine/Source/Programs/UnrealHeaderTool/Private/CodeGenerator.cpp
--- a/Engine/Source/Programs/UnrealHeaderTool/Private/CodeGenerator.cpp
+++ b/Engine/Source/Programs/UnrealHeaderTool/Private/CodeGenerator.cpp
@@ -6873,6 +6873,7 @@ ECompilationResult::Type UnrealHeaderTool_Main(const FString& ModuleInfoFilename
 
 	// Finish all async file tasks before stopping the clock
 	FTaskGraphInterface::Get().WaitUntilTasksComplete(GAsyncFileTasks);
+	GAsyncFileTasks.Reset();
 
 	// TEMPORARY change to log all files when UHT fails
 	if (FResults::GetOverallResults() != 0)
