Added support for native FProperty setters and getters.

diff --git a/Engine/Source/Programs/UnrealHeaderTool/Private/CodeGenerator.cpp b/Engine/Source/Programs/UnrealHeaderTool/Private/CodeGenerator.cpp
--- a/Engine/Source/Programs/UnrealHeaderTool/Private/CodeGenerator.cpp
+++ b/Engine/Source/Programs/UnrealHeaderTool/Private/CodeGenerator.cpp
@@ -755,6 +755,30 @@ void FNativeClassHeaderGenerator::GetPropertyTag(FUHTStringBuilder& Out, FUnreal
 	}
 }
 
+/** Helper function that generates property getter wrapper function name */
+inline FString GetPropertyGetterWrapperName(FUnrealPropertyDefinitionInfo& Prop)
+{
+	return FString::Printf(TEXT("Get%s_WrapperImpl"), *Prop.GetName());
+}
+
+/** Helper function that generates property setter wrapper function name */
+inline FString GetPropertySetterWrapperName(FUnrealPropertyDefinitionInfo& Prop)
+{
+	return FString::Printf(TEXT("Set%s_WrapperImpl"), *Prop.GetName());
+}
+
+/** Helper function that generates property getter wrapper function name for passing as a function pointer argument */
+inline FString GetPropertyGetterWrapperPtr(FUnrealPropertyDefinitionInfo& Prop, const TCHAR* SourceStruct)
+{
+	return !Prop.GetPropertyBase().bGetterFunctionFound ? TEXT("nullptr") : FString::Printf(TEXT("&%s::%s"), SourceStruct, *GetPropertyGetterWrapperName(Prop));
+}
+
+/** Helper function that generates property setter wrapper function name for passing as a function pointer argument */
+inline FString GetPropertySetterWrapperPtr(FUnrealPropertyDefinitionInfo& Prop, const TCHAR* SourceStruct)
+{
+	return !Prop.GetPropertyBase().bSetterFunctionFound ? TEXT("nullptr") : FString::Printf(TEXT("&%s::%s"), SourceStruct, *GetPropertySetterWrapperName(Prop));
+}
+
 void FNativeClassHeaderGenerator::OutputProperty(FOutputDevice& DeclOut, FOutputDevice& Out, FReferenceGatherers& OutReferenceGatherers, const TCHAR* Scope, TArray<FPropertyNamePointerPair>& PropertyNamesAndPointers, FUnrealPropertyDefinitionInfo& PropertyDef, const TCHAR* OffsetStr, FString&& Name, const TCHAR* DeclSpaces, const TCHAR* Spaces, const TCHAR* SourceStruct) const
 {
 	const FPropertyBase& PropertyBase   = PropertyDef.GetPropertyBase();
@@ -781,6 +805,9 @@ void FNativeClassHeaderGenerator::OutputProperty(FOutputDevice& DeclOut, FOutput
 		}
 	}
 
+	FString SetterFunc = GetPropertySetterWrapperPtr(PropertyDef, SourceStruct);
+	FString GetterFunc = GetPropertyGetterWrapperPtr(PropertyDef, SourceStruct);
+
 	auto OutputByteProperty = [&]()
 	{
 		FString MetaDataParams = OutputMetaDataCodeForObject(DeclOut, Out, PropertyDef, *FString::Printf(TEXT("%s_MetaData"), *Name), DeclSpaces, Spaces);
@@ -788,7 +815,7 @@ void FNativeClassHeaderGenerator::OutputProperty(FOutputDevice& DeclOut, FOutput
 		DeclOut.Logf(TEXT("%sstatic const UECodeGen_Private::FBytePropertyParams %s;\r\n"), DeclSpaces, *NameWithoutScope);
 
 		Out.Logf(
-			TEXT("%sconst UECodeGen_Private::FBytePropertyParams %s = { %s, %s, (EPropertyFlags)0x%016llx, UECodeGen_Private::EPropertyGenFlags::Byte, %s, %s, %s, %s, %s };%s\r\n"),
+			TEXT("%sconst UECodeGen_Private::FBytePropertyParams %s = { %s, %s, (EPropertyFlags)0x%016llx, UECodeGen_Private::EPropertyGenFlags::Byte, %s, %s, %s, %s, %s, %s, %s };%s\r\n"),
 			Spaces,
 			*Name,
 			*PropName,
@@ -796,6 +823,8 @@ void FNativeClassHeaderGenerator::OutputProperty(FOutputDevice& DeclOut, FOutput
 			PropFlags,
 			FPropertyObjectFlags,
 			*ArrayDim,
+			*SetterFunc,
+			*GetterFunc,
 			OffsetStr,
 			*GetSingletonNameFuncAddr(PropertyBase.EnumDef, OutReferenceGatherers.UniqueCrossModuleReferences),
 			*MetaDataParams,
@@ -818,7 +847,7 @@ void FNativeClassHeaderGenerator::OutputProperty(FOutputDevice& DeclOut, FOutput
 		DeclOut.Logf(TEXT("%sstatic const UECodeGen_Private::FInt8PropertyParams %s;\r\n"), DeclSpaces, *NameWithoutScope);
 
 		Out.Logf(
-			TEXT("%sconst UECodeGen_Private::FInt8PropertyParams %s = { %s, %s, (EPropertyFlags)0x%016llx, UECodeGen_Private::EPropertyGenFlags::Int8, %s, %s, %s, %s };%s\r\n"),
+			TEXT("%sconst UECodeGen_Private::FInt8PropertyParams %s = { %s, %s, (EPropertyFlags)0x%016llx, UECodeGen_Private::EPropertyGenFlags::Int8, %s, %s, %s, %s, %s, %s };%s\r\n"),
 			Spaces,
 			*Name,
 			*PropName,
@@ -826,6 +855,8 @@ void FNativeClassHeaderGenerator::OutputProperty(FOutputDevice& DeclOut, FOutput
 			PropFlags,
 			FPropertyObjectFlags,
 			*ArrayDim,
+			*SetterFunc,
+			*GetterFunc,
 			OffsetStr,
 			*MetaDataParams,
 			*PropTag
@@ -840,7 +871,7 @@ void FNativeClassHeaderGenerator::OutputProperty(FOutputDevice& DeclOut, FOutput
 		DeclOut.Logf(TEXT("%sstatic const UECodeGen_Private::FInt16PropertyParams %s;\r\n"), DeclSpaces, *NameWithoutScope);
 
 		Out.Logf(
-			TEXT("%sconst UECodeGen_Private::FInt16PropertyParams %s = { %s, %s, (EPropertyFlags)0x%016llx, UECodeGen_Private::EPropertyGenFlags::Int16, %s, %s, %s, %s };%s\r\n"),
+			TEXT("%sconst UECodeGen_Private::FInt16PropertyParams %s = { %s, %s, (EPropertyFlags)0x%016llx, UECodeGen_Private::EPropertyGenFlags::Int16, %s, %s, %s, %s, %s, %s };%s\r\n"),
 			Spaces,
 			*Name,
 			*PropName,
@@ -848,6 +879,8 @@ void FNativeClassHeaderGenerator::OutputProperty(FOutputDevice& DeclOut, FOutput
 			PropFlags,
 			FPropertyObjectFlags,
 			*ArrayDim,
+			*SetterFunc,
+			*GetterFunc,
 			OffsetStr,
 			*MetaDataParams,
 			*PropTag
@@ -864,7 +897,7 @@ void FNativeClassHeaderGenerator::OutputProperty(FOutputDevice& DeclOut, FOutput
 		DeclOut.Logf(TEXT("%sstatic const UECodeGen_Private::%s %s;\r\n"), DeclSpaces, PropTypeName, *NameWithoutScope);
 
 		Out.Logf(
-			TEXT("%sconst UECodeGen_Private::%s %s = { %s, %s, (EPropertyFlags)0x%016llx, UECodeGen_Private::EPropertyGenFlags::Int, %s, %s, %s, %s };%s\r\n"),
+			TEXT("%sconst UECodeGen_Private::%s %s = { %s, %s, (EPropertyFlags)0x%016llx, UECodeGen_Private::EPropertyGenFlags::Int, %s, %s, %s, %s, %s, %s };%s\r\n"),
 			Spaces,
 			PropTypeName,
 			*Name,
@@ -873,6 +906,8 @@ void FNativeClassHeaderGenerator::OutputProperty(FOutputDevice& DeclOut, FOutput
 			PropFlags,
 			FPropertyObjectFlags,
 			*ArrayDim,
+			*SetterFunc,
+			*GetterFunc,
 			OffsetStr,
 			*MetaDataParams,
 			*PropTag
@@ -887,7 +922,7 @@ void FNativeClassHeaderGenerator::OutputProperty(FOutputDevice& DeclOut, FOutput
 		DeclOut.Logf(TEXT("%sstatic const UECodeGen_Private::FInt64PropertyParams %s;\r\n"), DeclSpaces, *NameWithoutScope);
 
 		Out.Logf(
-			TEXT("%sconst UECodeGen_Private::FInt64PropertyParams %s = { %s, %s, (EPropertyFlags)0x%016llx, UECodeGen_Private::EPropertyGenFlags::Int64, %s, %s, %s, %s };%s\r\n"),
+			TEXT("%sconst UECodeGen_Private::FInt64PropertyParams %s = { %s, %s, (EPropertyFlags)0x%016llx, UECodeGen_Private::EPropertyGenFlags::Int64, %s, %s, %s, %s, %s, %s };%s\r\n"),
 			Spaces,
 			*Name,
 			*PropName,
@@ -895,6 +930,8 @@ void FNativeClassHeaderGenerator::OutputProperty(FOutputDevice& DeclOut, FOutput
 			PropFlags,
 			FPropertyObjectFlags,
 			*ArrayDim,
+			*SetterFunc,
+			*GetterFunc,
 			OffsetStr,
 			*MetaDataParams,
 			*PropTag
@@ -909,7 +946,7 @@ void FNativeClassHeaderGenerator::OutputProperty(FOutputDevice& DeclOut, FOutput
 		DeclOut.Logf(TEXT("%sstatic const UECodeGen_Private::FFInt16PropertyParams %s;\r\n"), DeclSpaces, *NameWithoutScope);
 
 		Out.Logf(
-			TEXT("%sconst UECodeGen_Private::FFInt16PropertyParams %s = { %s, %s, (EPropertyFlags)0x%016llx, UECodeGen_Private::EPropertyGenFlags::UInt16, %s, %s, %s, %s };%s\r\n"),
+			TEXT("%sconst UECodeGen_Private::FFInt16PropertyParams %s = { %s, %s, (EPropertyFlags)0x%016llx, UECodeGen_Private::EPropertyGenFlags::UInt16, %s, %s, %s, %s, %s, %s };%s\r\n"),
 			Spaces,
 			*Name,
 			*PropName,
@@ -917,6 +954,8 @@ void FNativeClassHeaderGenerator::OutputProperty(FOutputDevice& DeclOut, FOutput
 			PropFlags,
 			FPropertyObjectFlags,
 			*ArrayDim,
+			*SetterFunc,
+			*GetterFunc,
 			OffsetStr,
 			*MetaDataParams,
 			*PropTag
@@ -933,7 +972,7 @@ void FNativeClassHeaderGenerator::OutputProperty(FOutputDevice& DeclOut, FOutput
 		DeclOut.Logf(TEXT("%sstatic const UECodeGen_Private::%s %s;\r\n"), DeclSpaces, PropTypeName, *NameWithoutScope);
 
 		Out.Logf(
-			TEXT("%sconst UECodeGen_Private::%s %s = { %s, %s, (EPropertyFlags)0x%016llx, UECodeGen_Private::EPropertyGenFlags::UInt32, %s, %s, %s, %s };%s\r\n"),
+			TEXT("%sconst UECodeGen_Private::%s %s = { %s, %s, (EPropertyFlags)0x%016llx, UECodeGen_Private::EPropertyGenFlags::UInt32, %s, %s, %s, %s, %s, %s };%s\r\n"),
 			Spaces,
 			PropTypeName,
 			*Name,
@@ -942,6 +981,8 @@ void FNativeClassHeaderGenerator::OutputProperty(FOutputDevice& DeclOut, FOutput
 			PropFlags,
 			FPropertyObjectFlags,
 			*ArrayDim,
+			*SetterFunc,
+			*GetterFunc,
 			OffsetStr,
 			*MetaDataParams,
 			*PropTag
@@ -956,7 +997,7 @@ void FNativeClassHeaderGenerator::OutputProperty(FOutputDevice& DeclOut, FOutput
 		DeclOut.Logf(TEXT("%sstatic const UECodeGen_Private::FFInt64PropertyParams %s;\r\n"), DeclSpaces, *NameWithoutScope);
 
 		Out.Logf(
-			TEXT("%sconst UECodeGen_Private::FFInt64PropertyParams %s = { %s, %s, (EPropertyFlags)0x%016llx, UECodeGen_Private::EPropertyGenFlags::UInt64, %s, %s, %s, %s };%s\r\n"),
+			TEXT("%sconst UECodeGen_Private::FFInt64PropertyParams %s = { %s, %s, (EPropertyFlags)0x%016llx, UECodeGen_Private::EPropertyGenFlags::UInt64, %s, %s, %s, %s, %s, %s };%s\r\n"),
 			Spaces,
 			*Name,
 			*PropName,
@@ -964,6 +1005,8 @@ void FNativeClassHeaderGenerator::OutputProperty(FOutputDevice& DeclOut, FOutput
 			PropFlags,
 			FPropertyObjectFlags,
 			*ArrayDim,
+			*SetterFunc,
+			*GetterFunc,
 			OffsetStr,
 			*MetaDataParams,
 			*PropTag
@@ -978,7 +1021,7 @@ void FNativeClassHeaderGenerator::OutputProperty(FOutputDevice& DeclOut, FOutput
 		DeclOut.Logf(TEXT("%sstatic const UECodeGen_Private::FFloatPropertyParams %s;\r\n"), DeclSpaces, *NameWithoutScope);
 
 		Out.Logf(
-			TEXT("%sconst UECodeGen_Private::FFloatPropertyParams %s = { %s, %s, (EPropertyFlags)0x%016llx, UECodeGen_Private::EPropertyGenFlags::Float, %s, %s, %s, %s };%s\r\n"),
+			TEXT("%sconst UECodeGen_Private::FFloatPropertyParams %s = { %s, %s, (EPropertyFlags)0x%016llx, UECodeGen_Private::EPropertyGenFlags::Float, %s, %s, %s, %s, %s, %s };%s\r\n"),
 			Spaces,
 			*Name,
 			*PropName,
@@ -986,6 +1029,8 @@ void FNativeClassHeaderGenerator::OutputProperty(FOutputDevice& DeclOut, FOutput
 			PropFlags,
 			FPropertyObjectFlags,
 			*ArrayDim,
+			*SetterFunc,
+			*GetterFunc,
 			OffsetStr,
 			*MetaDataParams,
 			*PropTag
@@ -1000,7 +1045,7 @@ void FNativeClassHeaderGenerator::OutputProperty(FOutputDevice& DeclOut, FOutput
 		DeclOut.Logf(TEXT("%sstatic const UECodeGen_Private::FDoublePropertyParams %s;\r\n"), DeclSpaces, *NameWithoutScope);
 
 		Out.Logf(
-			TEXT("%sconst UECodeGen_Private::FDoublePropertyParams %s = { %s, %s, (EPropertyFlags)0x%016llx, UECodeGen_Private::EPropertyGenFlags::Double, %s, %s, %s, %s };%s\r\n"),
+			TEXT("%sconst UECodeGen_Private::FDoublePropertyParams %s = { %s, %s, (EPropertyFlags)0x%016llx, UECodeGen_Private::EPropertyGenFlags::Double, %s, %s, %s, %s, %s, %s };%s\r\n"),
 			Spaces,
 			*Name,
 			*PropName,
@@ -1008,6 +1053,8 @@ void FNativeClassHeaderGenerator::OutputProperty(FOutputDevice& DeclOut, FOutput
 			PropFlags,
 			FPropertyObjectFlags,
 			*ArrayDim,
+			*SetterFunc,
+			*GetterFunc,
 			OffsetStr,
 			*MetaDataParams,
 			*PropTag
@@ -1022,7 +1069,7 @@ void FNativeClassHeaderGenerator::OutputProperty(FOutputDevice& DeclOut, FOutput
 		DeclOut.Logf(TEXT("%sstatic const UECodeGen_Private::FLargeWorldCoordinatesRealPropertyParams %s;\r\n"), DeclSpaces, *NameWithoutScope);
 
 		Out.Logf(
-			TEXT("%sconst UECodeGen_Private::FLargeWorldCoordinatesRealPropertyParams %s = { %s, %s, (EPropertyFlags)0x%016llx, UECodeGen_Private::EPropertyGenFlags::LargeWorldCoordinatesReal, %s, %s, %s, %s };%s\r\n"),
+			TEXT("%sconst UECodeGen_Private::FLargeWorldCoordinatesRealPropertyParams %s = { %s, %s, (EPropertyFlags)0x%016llx, UECodeGen_Private::EPropertyGenFlags::LargeWorldCoordinatesReal, %s, %s, %s, %s, %s, %s };%s\r\n"),
 			Spaces,
 			*Name,
 			*PropName,
@@ -1030,6 +1077,8 @@ void FNativeClassHeaderGenerator::OutputProperty(FOutputDevice& DeclOut, FOutput
 			PropFlags,
 			FPropertyObjectFlags,
 			*ArrayDim,
+			*SetterFunc,
+			*GetterFunc,
 			OffsetStr,
 			*MetaDataParams,
 			*PropTag
@@ -1070,7 +1119,7 @@ void FNativeClassHeaderGenerator::OutputProperty(FOutputDevice& DeclOut, FOutput
 		DeclOut.Logf(TEXT("%sstatic const UECodeGen_Private::FBoolPropertyParams %s;\r\n"), DeclSpaces, *NameWithoutScope);
 
 		Out.Logf(
-			TEXT("%sconst UECodeGen_Private::FBoolPropertyParams %s = { %s, %s, (EPropertyFlags)0x%016llx, UECodeGen_Private::EPropertyGenFlags::Bool %s, %s, %s, sizeof(%s), %s, %s, %s };%s\r\n"),
+			TEXT("%sconst UECodeGen_Private::FBoolPropertyParams %s = { %s, %s, (EPropertyFlags)0x%016llx, UECodeGen_Private::EPropertyGenFlags::Bool %s, %s, %s, %s, %s, sizeof(%s), %s, %s, %s };%s\r\n"),
 			Spaces,
 			*Name,
 			*PropName,
@@ -1079,6 +1128,8 @@ void FNativeClassHeaderGenerator::OutputProperty(FOutputDevice& DeclOut, FOutput
 			bIsNativeBool ? TEXT("| UECodeGen_Private::EPropertyGenFlags::NativeBool") : TEXT(""),
 			FPropertyObjectFlags,
 			*ArrayDim,
+			*SetterFunc,
+			*GetterFunc,
 			*PropertyDef.GetCPPType(nullptr, 0),
 			*OuterSize,
 			*Setter,
@@ -1097,7 +1148,7 @@ void FNativeClassHeaderGenerator::OutputProperty(FOutputDevice& DeclOut, FOutput
 			DeclOut.Logf(TEXT("%sstatic const UECodeGen_Private::FClassPropertyParams %s;\r\n"), DeclSpaces, *NameWithoutScope);
 
 			Out.Logf(
-				TEXT("%sconst UECodeGen_Private::FClassPropertyParams %s = { %s, %s, (EPropertyFlags)0x%016llx, UECodeGen_Private::EPropertyGenFlags::Class, %s, %s, %s, %s, %s, %s };%s\r\n"),
+				TEXT("%sconst UECodeGen_Private::FClassPropertyParams %s = { %s, %s, (EPropertyFlags)0x%016llx, UECodeGen_Private::EPropertyGenFlags::Class, %s, %s, %s, %s, %s, %s, %s, %s };%s\r\n"),
 				Spaces,
 				*Name,
 				*PropName,
@@ -1105,6 +1156,8 @@ void FNativeClassHeaderGenerator::OutputProperty(FOutputDevice& DeclOut, FOutput
 				PropFlags,
 				FPropertyObjectFlags,
 				*ArrayDim,
+				*SetterFunc,
+				*GetterFunc,
 				OffsetStr,
 				*GetSingletonNameFuncAddr(PropertyBase.ClassDef, OutReferenceGatherers.UniqueCrossModuleReferences, false),
 				*GetSingletonNameFuncAddr(PropertyBase.MetaClassDef, OutReferenceGatherers.UniqueCrossModuleReferences, false),				
@@ -1117,7 +1170,7 @@ void FNativeClassHeaderGenerator::OutputProperty(FOutputDevice& DeclOut, FOutput
 			DeclOut.Logf(TEXT("%sstatic const UECodeGen_Private::FObjectPropertyParams %s;\r\n"), DeclSpaces, *NameWithoutScope);
 
 			Out.Logf(
-				TEXT("%sconst UECodeGen_Private::FObjectPropertyParams %s = { %s, %s, (EPropertyFlags)0x%016llx, UECodeGen_Private::EPropertyGenFlags::Object, %s, %s, %s, %s, %s };%s\r\n"),
+				TEXT("%sconst UECodeGen_Private::FObjectPropertyParams %s = { %s, %s, (EPropertyFlags)0x%016llx, UECodeGen_Private::EPropertyGenFlags::Object, %s, %s, %s, %s, %s, %s, %s };%s\r\n"),
 				Spaces,
 				*Name,
 				*PropName,
@@ -1125,6 +1178,8 @@ void FNativeClassHeaderGenerator::OutputProperty(FOutputDevice& DeclOut, FOutput
 				PropFlags,
 				FPropertyObjectFlags,
 				*ArrayDim,
+				*SetterFunc,
+				*GetterFunc,
 				OffsetStr,
 				*GetSingletonNameFuncAddr(PropertyBase.ClassDef, OutReferenceGatherers.UniqueCrossModuleReferences, false),
 				*MetaDataParams,
@@ -1143,7 +1198,7 @@ void FNativeClassHeaderGenerator::OutputProperty(FOutputDevice& DeclOut, FOutput
 			DeclOut.Logf(TEXT("%sstatic const UECodeGen_Private::FClassPtrPropertyParams %s;\r\n"), DeclSpaces, *NameWithoutScope);
 
 			Out.Logf(
-				TEXT("%sconst UECodeGen_Private::FClassPtrPropertyParams %s = { %s, %s, (EPropertyFlags)0x%016llx, UECodeGen_Private::EPropertyGenFlags::Class | UECodeGen_Private::EPropertyGenFlags::ObjectPtr, %s, %s, %s, %s, %s, %s };%s\r\n"),
+				TEXT("%sconst UECodeGen_Private::FClassPtrPropertyParams %s = { %s, %s, (EPropertyFlags)0x%016llx, UECodeGen_Private::EPropertyGenFlags::Class | UECodeGen_Private::EPropertyGenFlags::ObjectPtr, %s, %s, %s, %s, %s, %s, %s, %s };%s\r\n"),
 				Spaces,
 				*Name,
 				*PropName,
@@ -1151,6 +1206,8 @@ void FNativeClassHeaderGenerator::OutputProperty(FOutputDevice& DeclOut, FOutput
 				PropFlags,
 				FPropertyObjectFlags,
 				*ArrayDim,
+				*SetterFunc,
+				*GetterFunc,
 				OffsetStr,
 				*GetSingletonNameFuncAddr(PropertyBase.ClassDef, OutReferenceGatherers.UniqueCrossModuleReferences, false),
 				*GetSingletonNameFuncAddr(PropertyBase.MetaClassDef, OutReferenceGatherers.UniqueCrossModuleReferences, false),
@@ -1163,7 +1220,7 @@ void FNativeClassHeaderGenerator::OutputProperty(FOutputDevice& DeclOut, FOutput
 			DeclOut.Logf(TEXT("%sstatic const UECodeGen_Private::FObjectPtrPropertyParams %s;\r\n"), DeclSpaces, *NameWithoutScope);
 
 			Out.Logf(
-				TEXT("%sconst UECodeGen_Private::FObjectPtrPropertyParams %s = { %s, %s, (EPropertyFlags)0x%016llx, UECodeGen_Private::EPropertyGenFlags::Object | UECodeGen_Private::EPropertyGenFlags::ObjectPtr, %s, %s, %s, %s, %s };%s\r\n"),
+				TEXT("%sconst UECodeGen_Private::FObjectPtrPropertyParams %s = { %s, %s, (EPropertyFlags)0x%016llx, UECodeGen_Private::EPropertyGenFlags::Object | UECodeGen_Private::EPropertyGenFlags::ObjectPtr, %s, %s, %s, %s, %s, %s, %s };%s\r\n"),
 				Spaces,
 				*Name,
 				*PropName,
@@ -1171,6 +1228,8 @@ void FNativeClassHeaderGenerator::OutputProperty(FOutputDevice& DeclOut, FOutput
 				PropFlags,
 				FPropertyObjectFlags,
 				*ArrayDim,
+				*SetterFunc,
+				*GetterFunc,
 				OffsetStr,
 				*GetSingletonNameFuncAddr(PropertyBase.ClassDef, OutReferenceGatherers.UniqueCrossModuleReferences, false),
 				*MetaDataParams,
@@ -1189,7 +1248,7 @@ void FNativeClassHeaderGenerator::OutputProperty(FOutputDevice& DeclOut, FOutput
 			DeclOut.Logf(TEXT("%sstatic const UECodeGen_Private::FSoftClassPropertyParams %s;\r\n"), DeclSpaces, *NameWithoutScope);
 
 			Out.Logf(
-				TEXT("%sconst UECodeGen_Private::FSoftClassPropertyParams %s = { %s, %s, (EPropertyFlags)0x%016llx, UECodeGen_Private::EPropertyGenFlags::SoftClass, %s, %s, %s, %s, %s };%s\r\n"),
+				TEXT("%sconst UECodeGen_Private::FSoftClassPropertyParams %s = { %s, %s, (EPropertyFlags)0x%016llx, UECodeGen_Private::EPropertyGenFlags::SoftClass, %s, %s, %s, %s, %s, %s, %s };%s\r\n"),
 				Spaces,
 				*Name,
 				*PropName,
@@ -1197,6 +1256,8 @@ void FNativeClassHeaderGenerator::OutputProperty(FOutputDevice& DeclOut, FOutput
 				PropFlags,
 				FPropertyObjectFlags,
 				*ArrayDim,
+				*SetterFunc,
+				*GetterFunc,
 				OffsetStr,
 				*GetSingletonNameFuncAddr(PropertyBase.MetaClassDef, OutReferenceGatherers.UniqueCrossModuleReferences, false),
 				*MetaDataParams,
@@ -1208,7 +1269,7 @@ void FNativeClassHeaderGenerator::OutputProperty(FOutputDevice& DeclOut, FOutput
 			DeclOut.Logf(TEXT("%sstatic const UECodeGen_Private::FSoftObjectPropertyParams %s;\r\n"), DeclSpaces, *NameWithoutScope);
 
 			Out.Logf(
-				TEXT("%sconst UECodeGen_Private::FSoftObjectPropertyParams %s = { %s, %s, (EPropertyFlags)0x%016llx, UECodeGen_Private::EPropertyGenFlags::SoftObject, %s, %s, %s, %s, %s };%s\r\n"),
+				TEXT("%sconst UECodeGen_Private::FSoftObjectPropertyParams %s = { %s, %s, (EPropertyFlags)0x%016llx, UECodeGen_Private::EPropertyGenFlags::SoftObject, %s, %s, %s, %s, %s, %s, %s };%s\r\n"),
 				Spaces,
 				*Name,
 				*PropName,
@@ -1216,6 +1277,8 @@ void FNativeClassHeaderGenerator::OutputProperty(FOutputDevice& DeclOut, FOutput
 				PropFlags,
 				FPropertyObjectFlags,
 				*ArrayDim,
+				*SetterFunc,
+				*GetterFunc,
 				OffsetStr,
 				*GetSingletonNameFuncAddr(PropertyBase.ClassDef, OutReferenceGatherers.UniqueCrossModuleReferences, false),
 				*MetaDataParams,
@@ -1232,7 +1295,7 @@ void FNativeClassHeaderGenerator::OutputProperty(FOutputDevice& DeclOut, FOutput
 		DeclOut.Logf(TEXT("%sstatic const UECodeGen_Private::FWeakObjectPropertyParams %s;\r\n"), DeclSpaces, *NameWithoutScope);
 
 		Out.Logf(
-			TEXT("%sconst UECodeGen_Private::FWeakObjectPropertyParams %s = { %s, %s, (EPropertyFlags)0x%016llx, UECodeGen_Private::EPropertyGenFlags::WeakObject, %s, %s, %s, %s, %s };%s\r\n"),
+			TEXT("%sconst UECodeGen_Private::FWeakObjectPropertyParams %s = { %s, %s, (EPropertyFlags)0x%016llx, UECodeGen_Private::EPropertyGenFlags::WeakObject, %s, %s, %s, %s, %s, %s, %s };%s\r\n"),
 			Spaces,
 			*Name,
 			*PropName,
@@ -1240,6 +1303,8 @@ void FNativeClassHeaderGenerator::OutputProperty(FOutputDevice& DeclOut, FOutput
 			PropFlags,
 			FPropertyObjectFlags,
 			*ArrayDim,
+			*SetterFunc,
+			*GetterFunc,
 			OffsetStr,
 			*GetSingletonNameFuncAddr(PropertyBase.ClassDef, OutReferenceGatherers.UniqueCrossModuleReferences, false),
 			*MetaDataParams,
@@ -1255,7 +1320,7 @@ void FNativeClassHeaderGenerator::OutputProperty(FOutputDevice& DeclOut, FOutput
 		DeclOut.Logf(TEXT("%sstatic const UECodeGen_Private::FLazyObjectPropertyParams %s;\r\n"), DeclSpaces, *NameWithoutScope);
 
 		Out.Logf(
-			TEXT("%sconst UECodeGen_Private::FLazyObjectPropertyParams %s = { %s, %s, (EPropertyFlags)0x%016llx, UECodeGen_Private::EPropertyGenFlags::LazyObject, %s, %s, %s, %s, %s };%s\r\n"),
+			TEXT("%sconst UECodeGen_Private::FLazyObjectPropertyParams %s = { %s, %s, (EPropertyFlags)0x%016llx, UECodeGen_Private::EPropertyGenFlags::LazyObject, %s, %s, %s, %s, %s, %s, %s };%s\r\n"),
 			Spaces,
 			*Name,
 			*PropName,
@@ -1263,6 +1328,8 @@ void FNativeClassHeaderGenerator::OutputProperty(FOutputDevice& DeclOut, FOutput
 			PropFlags,
 			FPropertyObjectFlags,
 			*ArrayDim,
+			*SetterFunc,
+			*GetterFunc,
 			OffsetStr,
 			*GetSingletonNameFuncAddr(PropertyBase.ClassDef, OutReferenceGatherers.UniqueCrossModuleReferences, false),
 			*MetaDataParams,
@@ -1278,7 +1345,7 @@ void FNativeClassHeaderGenerator::OutputProperty(FOutputDevice& DeclOut, FOutput
 		DeclOut.Logf(TEXT("%sstatic const UECodeGen_Private::FInterfacePropertyParams %s;\r\n"), DeclSpaces, *NameWithoutScope);
 
 		Out.Logf(
-			TEXT("%sconst UECodeGen_Private::FInterfacePropertyParams %s = { %s, %s, (EPropertyFlags)0x%016llx, UECodeGen_Private::EPropertyGenFlags::Interface, %s, %s, %s, %s, %s };%s\r\n"),
+			TEXT("%sconst UECodeGen_Private::FInterfacePropertyParams %s = { %s, %s, (EPropertyFlags)0x%016llx, UECodeGen_Private::EPropertyGenFlags::Interface, %s, %s, %s, %s, %s, %s, %s };%s\r\n"),
 			Spaces,
 			*Name,
 			*PropName,
@@ -1286,6 +1353,8 @@ void FNativeClassHeaderGenerator::OutputProperty(FOutputDevice& DeclOut, FOutput
 			PropFlags,
 			FPropertyObjectFlags,
 			*ArrayDim,
+			*SetterFunc,
+			*GetterFunc,
 			OffsetStr,
 			*GetSingletonNameFuncAddr(PropertyBase.ClassDef, OutReferenceGatherers.UniqueCrossModuleReferences, false),
 			*MetaDataParams,
@@ -1301,7 +1370,7 @@ void FNativeClassHeaderGenerator::OutputProperty(FOutputDevice& DeclOut, FOutput
 		DeclOut.Logf(TEXT("%sstatic const UECodeGen_Private::FNamePropertyParams %s;\r\n"), DeclSpaces, *NameWithoutScope);
 
 		Out.Logf(
-			TEXT("%sconst UECodeGen_Private::FNamePropertyParams %s = { %s, %s, (EPropertyFlags)0x%016llx, UECodeGen_Private::EPropertyGenFlags::Name, %s, %s, %s, %s };%s\r\n"),
+			TEXT("%sconst UECodeGen_Private::FNamePropertyParams %s = { %s, %s, (EPropertyFlags)0x%016llx, UECodeGen_Private::EPropertyGenFlags::Name, %s, %s, %s, %s, %s, %s };%s\r\n"),
 			Spaces,
 			*Name,
 			*PropName,
@@ -1309,6 +1378,8 @@ void FNativeClassHeaderGenerator::OutputProperty(FOutputDevice& DeclOut, FOutput
 			PropFlags,
 			FPropertyObjectFlags,
 			*ArrayDim,
+			*SetterFunc,
+			*GetterFunc,
 			OffsetStr,
 			*MetaDataParams,
 			*PropTag
@@ -1323,7 +1394,7 @@ void FNativeClassHeaderGenerator::OutputProperty(FOutputDevice& DeclOut, FOutput
 		DeclOut.Logf(TEXT("%sstatic const UECodeGen_Private::FStrPropertyParams %s;\r\n"), DeclSpaces, *NameWithoutScope);
 
 		Out.Logf(
-			TEXT("%sconst UECodeGen_Private::FStrPropertyParams %s = { %s, %s, (EPropertyFlags)0x%016llx, UECodeGen_Private::EPropertyGenFlags::Str, %s, %s, %s, %s };%s\r\n"),
+			TEXT("%sconst UECodeGen_Private::FStrPropertyParams %s = { %s, %s, (EPropertyFlags)0x%016llx, UECodeGen_Private::EPropertyGenFlags::Str, %s, %s, %s, %s, %s, %s };%s\r\n"),
 			Spaces,
 			*Name,
 			*PropName,
@@ -1331,6 +1402,8 @@ void FNativeClassHeaderGenerator::OutputProperty(FOutputDevice& DeclOut, FOutput
 			PropFlags,
 			FPropertyObjectFlags,
 			*ArrayDim,
+			*SetterFunc,
+			*GetterFunc,
 			OffsetStr,
 			*MetaDataParams,
 			*PropTag
@@ -1351,7 +1424,7 @@ void FNativeClassHeaderGenerator::OutputProperty(FOutputDevice& DeclOut, FOutput
 		DeclOut.Logf(TEXT("%sstatic const UECodeGen_Private::FArrayPropertyParams %s;\r\n"), DeclSpaces, *NameWithoutScope);
 
 		Out.Logf(
-			TEXT("%sconst UECodeGen_Private::FArrayPropertyParams %s = { %s, %s, (EPropertyFlags)0x%016llx, UECodeGen_Private::EPropertyGenFlags::Array, %s, %s, %s, %s, %s };%s\r\n"),
+			TEXT("%sconst UECodeGen_Private::FArrayPropertyParams %s = { %s, %s, (EPropertyFlags)0x%016llx, UECodeGen_Private::EPropertyGenFlags::Array, %s, %s, %s, %s, %s, %s, %s };%s\r\n"),
 			Spaces,
 			*Name,
 			*PropName,
@@ -1359,6 +1432,8 @@ void FNativeClassHeaderGenerator::OutputProperty(FOutputDevice& DeclOut, FOutput
 			PropFlags,
 			FPropertyObjectFlags,
 			*ArrayDim,
+			*SetterFunc,
+			*GetterFunc,
 			OffsetStr,
 			PropertyDef.GetAllocatorType() == EAllocatorType::MemoryImage ? TEXT("EArrayPropertyFlags::UsesMemoryImageAllocator") : TEXT("EArrayPropertyFlags::None"),
 			*MetaDataParams,
@@ -1383,7 +1458,7 @@ void FNativeClassHeaderGenerator::OutputProperty(FOutputDevice& DeclOut, FOutput
 		DeclOut.Logf(TEXT("%sstatic const UECodeGen_Private::FMapPropertyParams %s;\r\n"), DeclSpaces, *NameWithoutScope);
 
 		Out.Logf(
-			TEXT("%sconst UECodeGen_Private::FMapPropertyParams %s = { %s, %s, (EPropertyFlags)0x%016llx, UECodeGen_Private::EPropertyGenFlags::Map, %s, %s, %s, %s, %s };%s\r\n"),
+			TEXT("%sconst UECodeGen_Private::FMapPropertyParams %s = { %s, %s, (EPropertyFlags)0x%016llx, UECodeGen_Private::EPropertyGenFlags::Map, %s, %s, %s, %s, %s, %s, %s };%s\r\n"),
 			Spaces,
 			*Name,
 			*PropName,
@@ -1391,6 +1466,8 @@ void FNativeClassHeaderGenerator::OutputProperty(FOutputDevice& DeclOut, FOutput
 			PropFlags,
 			FPropertyObjectFlags,
 			*ArrayDim,
+			*SetterFunc,
+			*GetterFunc,
 			OffsetStr,
 			PropertyDef.GetAllocatorType() == EAllocatorType::MemoryImage ? TEXT("EMapPropertyFlags::UsesMemoryImageAllocator") : TEXT("EMapPropertyFlags::None"),
 			*MetaDataParams,
@@ -1418,7 +1495,7 @@ void FNativeClassHeaderGenerator::OutputProperty(FOutputDevice& DeclOut, FOutput
 		}
 
 		Out.Logf(
-			TEXT("%sconst UECodeGen_Private::FSetPropertyParams %s = { %s, %s, (EPropertyFlags)0x%016llx, UECodeGen_Private::EPropertyGenFlags::Set, %s, %s, %s, %s };%s\r\n"),
+			TEXT("%sconst UECodeGen_Private::FSetPropertyParams %s = { %s, %s, (EPropertyFlags)0x%016llx, UECodeGen_Private::EPropertyGenFlags::Set, %s, %s, %s, %s, %s, %s };%s\r\n"),
 			Spaces,
 			*Name,
 			*PropName,
@@ -1426,6 +1503,8 @@ void FNativeClassHeaderGenerator::OutputProperty(FOutputDevice& DeclOut, FOutput
 			PropFlags,
 			FPropertyObjectFlags,
 			*ArrayDim,
+			*SetterFunc,
+			*GetterFunc,
 			OffsetStr,
 			*MetaDataParams,
 			*PropTag
@@ -1440,7 +1519,7 @@ void FNativeClassHeaderGenerator::OutputProperty(FOutputDevice& DeclOut, FOutput
 		DeclOut.Logf(TEXT("%sstatic const UECodeGen_Private::FStructPropertyParams %s;\r\n"), DeclSpaces, *NameWithoutScope);
 
 		Out.Logf(
-			TEXT("%sconst UECodeGen_Private::FStructPropertyParams %s = { %s, %s, (EPropertyFlags)0x%016llx, UECodeGen_Private::EPropertyGenFlags::Struct, %s, %s, %s, %s, %s };%s\r\n"),
+			TEXT("%sconst UECodeGen_Private::FStructPropertyParams %s = { %s, %s, (EPropertyFlags)0x%016llx, UECodeGen_Private::EPropertyGenFlags::Struct, %s, %s, %s, %s, %s, %s, %s };%s\r\n"),
 			Spaces,
 			*Name,
 			*PropName,
@@ -1448,6 +1527,8 @@ void FNativeClassHeaderGenerator::OutputProperty(FOutputDevice& DeclOut, FOutput
 			PropFlags,
 			FPropertyObjectFlags,
 			*ArrayDim,
+			*SetterFunc,
+			*GetterFunc,
 			OffsetStr,
 			*GetSingletonNameFuncAddr(PropertyBase.ScriptStructDef, OutReferenceGatherers.UniqueCrossModuleReferences),
 			*MetaDataParams,
@@ -1463,7 +1544,7 @@ void FNativeClassHeaderGenerator::OutputProperty(FOutputDevice& DeclOut, FOutput
 		DeclOut.Logf(TEXT("%sstatic const UECodeGen_Private::FDelegatePropertyParams %s;\r\n"), DeclSpaces, *NameWithoutScope);
 
 		Out.Logf(
-			TEXT("%sconst UECodeGen_Private::FDelegatePropertyParams %s = { %s, %s, (EPropertyFlags)0x%016llx, UECodeGen_Private::EPropertyGenFlags::Delegate, %s, %s, %s, %s, %s };%s\r\n"),
+			TEXT("%sconst UECodeGen_Private::FDelegatePropertyParams %s = { %s, %s, (EPropertyFlags)0x%016llx, UECodeGen_Private::EPropertyGenFlags::Delegate, %s, %s, %s, %s, %s, %s, %s };%s\r\n"),
 			Spaces,
 			*Name,
 			*PropName,
@@ -1471,6 +1552,8 @@ void FNativeClassHeaderGenerator::OutputProperty(FOutputDevice& DeclOut, FOutput
 			PropFlags,
 			FPropertyObjectFlags,
 			*ArrayDim,
+			*SetterFunc,
+			*GetterFunc,
 			OffsetStr,
 			*GetSingletonNameFuncAddr(PropertyBase.FunctionDef, OutReferenceGatherers.UniqueCrossModuleReferences),
 			*MetaDataParams,
@@ -1488,7 +1571,7 @@ void FNativeClassHeaderGenerator::OutputProperty(FOutputDevice& DeclOut, FOutput
 		DeclOut.Logf(TEXT("%sstatic const UECodeGen_Private::FMulticastDelegatePropertyParams %s;\r\n"), DeclSpaces, *NameWithoutScope);
 
 		Out.Logf(
-			TEXT("%sconst UECodeGen_Private::FMulticastDelegatePropertyParams %s = { %s, %s, (EPropertyFlags)0x%016llx, UECodeGen_Private::EPropertyGenFlags::%sMulticastDelegate, %s, %s, %s, %s, %s };%s\r\n"),
+			TEXT("%sconst UECodeGen_Private::FMulticastDelegatePropertyParams %s = { %s, %s, (EPropertyFlags)0x%016llx, UECodeGen_Private::EPropertyGenFlags::%sMulticastDelegate, %s, %s, %s, %s, %s, %s, %s };%s\r\n"),
 			Spaces,
 			*Name,
 			*PropName,
@@ -1497,6 +1580,8 @@ void FNativeClassHeaderGenerator::OutputProperty(FOutputDevice& DeclOut, FOutput
 			!bIsSparse ? TEXT("Inline") : TEXT("Sparse"),
 			FPropertyObjectFlags,
 			*ArrayDim,
+			*SetterFunc,
+			*GetterFunc,
 			OffsetStr,
 			*GetSingletonNameFuncAddr(PropertyBase.FunctionDef, OutReferenceGatherers.UniqueCrossModuleReferences),
 			*MetaDataParams,
@@ -1512,7 +1597,7 @@ void FNativeClassHeaderGenerator::OutputProperty(FOutputDevice& DeclOut, FOutput
 		DeclOut.Logf(TEXT("%sstatic const UECodeGen_Private::FTextPropertyParams %s;\r\n"), DeclSpaces, *NameWithoutScope);
 
 		Out.Logf(
-			TEXT("%sconst UECodeGen_Private::FTextPropertyParams %s = { %s, %s, (EPropertyFlags)0x%016llx, UECodeGen_Private::EPropertyGenFlags::Text, %s, %s, %s, %s };%s\r\n"),
+			TEXT("%sconst UECodeGen_Private::FTextPropertyParams %s = { %s, %s, (EPropertyFlags)0x%016llx, UECodeGen_Private::EPropertyGenFlags::Text, %s, %s, %s, %s, %s, %s };%s\r\n"),
 			Spaces,
 			*Name,
 			*PropName,
@@ -1520,6 +1605,8 @@ void FNativeClassHeaderGenerator::OutputProperty(FOutputDevice& DeclOut, FOutput
 			PropFlags,
 			FPropertyObjectFlags,
 			*ArrayDim,
+			*SetterFunc,
+			*GetterFunc,
 			OffsetStr,
 			*MetaDataParams,
 			*PropTag
@@ -1544,7 +1631,7 @@ void FNativeClassHeaderGenerator::OutputProperty(FOutputDevice& DeclOut, FOutput
 			DeclOut.Logf(TEXT("%sstatic const UECodeGen_Private::FEnumPropertyParams %s;\r\n"), DeclSpaces, *NameWithoutScope);
 
 			Out.Logf(
-				TEXT("%sconst UECodeGen_Private::FEnumPropertyParams %s = { %s, %s, (EPropertyFlags)0x%016llx, UECodeGen_Private::EPropertyGenFlags::Enum, %s, %s, %s, %s, %s };%s\r\n"),
+				TEXT("%sconst UECodeGen_Private::FEnumPropertyParams %s = { %s, %s, (EPropertyFlags)0x%016llx, UECodeGen_Private::EPropertyGenFlags::Enum, %s, %s, %s, %s, %s, %s, %s };%s\r\n"),
 				Spaces,
 				*Name,
 				*PropName,
@@ -1552,6 +1639,8 @@ void FNativeClassHeaderGenerator::OutputProperty(FOutputDevice& DeclOut, FOutput
 				PropFlags,
 				FPropertyObjectFlags,
 				*ArrayDim,
+				*SetterFunc,
+				*GetterFunc,
 				OffsetStr,
 				*GetSingletonNameFuncAddr(PropertyBase.EnumDef, OutReferenceGatherers.UniqueCrossModuleReferences),
 				*MetaDataParams,
@@ -1568,7 +1657,7 @@ void FNativeClassHeaderGenerator::OutputProperty(FOutputDevice& DeclOut, FOutput
 		DeclOut.Logf(TEXT("%sstatic const UECodeGen_Private::FFieldPathPropertyParams %s;\r\n"), DeclSpaces, *NameWithoutScope);
 
 		Out.Logf(
-			TEXT("%sconst UECodeGen_Private::FFieldPathPropertyParams %s = { %s, %s, (EPropertyFlags)0x%016llx, UECodeGen_Private::EPropertyGenFlags::FieldPath, %s, %s, %s, %s, %s };%s\r\n"),
+			TEXT("%sconst UECodeGen_Private::FFieldPathPropertyParams %s = { %s, %s, (EPropertyFlags)0x%016llx, UECodeGen_Private::EPropertyGenFlags::FieldPath, %s, %s, %s, %s, %s, %s, %s };%s\r\n"),
 			Spaces,
 			*Name,
 			*PropName,
@@ -1576,6 +1665,8 @@ void FNativeClassHeaderGenerator::OutputProperty(FOutputDevice& DeclOut, FOutput
 			PropFlags,
 			FPropertyObjectFlags,
 			*ArrayDim,
+			*SetterFunc,
+			*GetterFunc,
 			OffsetStr,
 			*FString::Printf(TEXT("&F%s::StaticClass"), *PropertyBase.FieldClassName.ToString()),
 			*MetaDataParams,
@@ -4638,6 +4729,8 @@ struct FNativeFunctionStringBuilder
 	FUHTStringBuilder AutogeneratedBlueprintFunctionDeclarationsOnlyNotDeclared;
 	FUHTStringBuilder AutogeneratedStaticData;
 	FUHTStringBuilder AutogeneratedStaticDataFuncs;
+	FUHTStringBuilder AccessorWrappers;
+	FUHTStringBuilder AccessorWrapperImplementations;
 };
 
 void FNativeClassHeaderGenerator::ExportNativeFunctions(FOutputDevice& OutGeneratedHeaderText, FOutputDevice& OutGeneratedCPPText, FOutputDevice& OutMacroCalls, FOutputDevice& OutNoPureDeclsMacroCalls, FReferenceGatherers& OutReferenceGatherers, const FUnrealSourceFile& SourceFile, FUnrealClassDefinitionInfo& ClassDef) const
@@ -4886,6 +4979,69 @@ void FNativeClassHeaderGenerator::ExportNativeFunctions(FOutputDevice& OutGenera
 		OutGeneratedCPPText.Log(EditorStringBuilders.RPCImplementations);
 		OutGeneratedCPPText.Log(EndEditorOnlyGuard);
 	}
+
+	// Property setters and getters
+	{
+		TArray<TSharedRef<FUnrealPropertyDefinitionInfo>>& Properties = ClassDef.GetProperties();
+		
+		for (TSharedRef<FUnrealPropertyDefinitionInfo>& Prop : Properties)
+		{
+			const bool bEditorOnlyProperty = Prop->IsEditorOnlyProperty();
+
+			if (!Prop->GetPropertyBase().GetterName.IsEmpty() && Prop->GetPropertyBase().bGetterFunctionFound)
+			{
+				const FString GetterWrapperFunctionName = GetPropertyGetterWrapperName(*Prop);
+				RuntimeStringBuilders.AccessorWrappers.Logf(TEXT("static void %s(const void* Object, void* OutValue);\r\n"), *GetterWrapperFunctionName);
+				RuntimeStringBuilders.AccessorWrapperImplementations.Logf(TEXT("\tvoid %s::%s(const void* Object, void* OutValue)\r\n"), *ClassCPPName, *GetterWrapperFunctionName);
+				RuntimeStringBuilders.AccessorWrapperImplementations += TEXT("\t{\r\n");
+				if (bEditorOnlyProperty)
+				{
+					RuntimeStringBuilders.AccessorWrapperImplementations += TEXT("\t#if WITH_EDITORONLY_DATA\r\n");
+				}
+				RuntimeStringBuilders.AccessorWrapperImplementations.Logf(TEXT("\t\tconst %s* Obj = (const %s*)Object;\r\n"), *ClassCPPName, *ClassCPPName);
+				FString ExtendedType;
+				FString PropertyType = Prop->GetCPPType(&ExtendedType, CPPF_Implementation | CPPF_ArgumentOrReturnValue);
+				PropertyType += ExtendedType;
+				RuntimeStringBuilders.AccessorWrapperImplementations.Logf(TEXT("\t\t%s& Result = *(%s*)OutValue;\r\n"), *PropertyType, *PropertyType);
+				RuntimeStringBuilders.AccessorWrapperImplementations.Logf(TEXT("\t\tResult = (%s)Obj->%s();\r\n"), *PropertyType, *Prop->GetPropertyBase().GetterName);
+				if (bEditorOnlyProperty)
+				{
+					RuntimeStringBuilders.AccessorWrapperImplementations += TEXT("\t#endif // WITH_EDITORONLY_DATA\r\n");
+				}
+				RuntimeStringBuilders.AccessorWrapperImplementations += TEXT("\t}\r\n");
+			}
+			if (!Prop->GetPropertyBase().SetterName.IsEmpty() && Prop->GetPropertyBase().bSetterFunctionFound)
+			{
+				const FString SetterWrapperFunctionName = GetPropertySetterWrapperName(*Prop);
+				RuntimeStringBuilders.AccessorWrappers.Logf(TEXT("static void %s(void* Object, const void* InValue);\r\n"), *SetterWrapperFunctionName);
+				RuntimeStringBuilders.AccessorWrapperImplementations.Logf(TEXT("\tvoid %s::%s(void* Object, const void* InValue)\r\n"), *ClassCPPName, *SetterWrapperFunctionName);
+				RuntimeStringBuilders.AccessorWrapperImplementations += TEXT("\t{\r\n");
+				if (bEditorOnlyProperty)
+				{
+					RuntimeStringBuilders.AccessorWrapperImplementations += TEXT("\t#if WITH_EDITORONLY_DATA\r\n");
+				}
+				RuntimeStringBuilders.AccessorWrapperImplementations.Logf(TEXT("\t\t%s* Obj = (%s*)Object;\r\n"), *ClassCPPName, *ClassCPPName);
+				FString ExtendedType;
+				FString PropertyType = Prop->GetCPPType(&ExtendedType, CPPF_Implementation | CPPF_ArgumentOrReturnValue);
+				PropertyType += ExtendedType;
+				RuntimeStringBuilders.AccessorWrapperImplementations.Logf(TEXT("\t\t%s& Value = *(%s*)InValue;\r\n"), *PropertyType, *PropertyType);
+				RuntimeStringBuilders.AccessorWrapperImplementations.Logf(TEXT("\t\tObj->%s(Value);\r\n"), *Prop->GetPropertyBase().SetterName);
+				if (bEditorOnlyProperty)
+				{
+					RuntimeStringBuilders.AccessorWrapperImplementations += TEXT("\t#endif // WITH_EDITORONLY_DATA\r\n");
+				}
+				RuntimeStringBuilders.AccessorWrapperImplementations += TEXT("\t}\r\n");
+			}
+		}
+
+		FString MacroName = SourceFile.GetGeneratedMacroName(ClassDef.GetGeneratedBodyLine(), TEXT("_ACCESSORS"));
+		WriteMacro(OutGeneratedHeaderText, MacroName, RuntimeStringBuilders.AccessorWrappers);
+
+		OutMacroCalls.Logf(TEXT("\t%s\r\n"), *MacroName);
+		OutNoPureDeclsMacroCalls.Logf(TEXT("\t%s\r\n"), *MacroName);
+
+		OutGeneratedCPPText.Log(RuntimeStringBuilders.AccessorWrapperImplementations);
+	}
 }
 
 /**
diff --git a/Engine/Source/Programs/UnrealHeaderTool/Private/HeaderParser.cpp b/Engine/Source/Programs/UnrealHeaderTool/Private/HeaderParser.cpp
--- a/Engine/Source/Programs/UnrealHeaderTool/Private/HeaderParser.cpp
+++ b/Engine/Source/Programs/UnrealHeaderTool/Private/HeaderParser.cpp
@@ -783,6 +783,7 @@ FUHTConfig::FUHTConfig()
 	const FName EngineObjectPtrMemberBehaviorKey(TEXT("EngineObjectPtrMemberBehavior"));
 	const FName NonEngineNativePointerMemberBehaviorKey(TEXT("NonEngineNativePointerMemberBehavior"));
 	const FName NonEngineObjectPtrMemberBehaviorKey(TEXT("NonEngineObjectPtrMemberBehavior"));
+	const FName AutomaticSettersAndGettersKey(TEXT("AutomaticSettersAndGetters"));
 
 	FConfigSection* ConfigSection = GConfig->GetSectionPrivate(TEXT("UnrealHeaderTool"), false, true, GEngineIni);
 	if (ConfigSection)
@@ -831,6 +832,10 @@ FUHTConfig::FUHTConfig()
 			{
 				NonEngineObjectPtrMemberBehavior = ToPointerMemberBehavior(It.Value().GetValue());
 			}
+			else if (It.Key() == AutomaticSettersAndGettersKey)
+			{
+				bAllowAutomaticSettersAndGetters = It.Value().GetValue().ToBool();
+			}
 		}
 	}
 }
@@ -2609,8 +2614,6 @@ void FHeaderParser::VerifyPropertyMarkups(FUnrealClassDefinitionInfo& TargetClas
 			return nullptr;
 		};
 
-		FPropertyBase& PropertyToken = PropertyDef->GetPropertyBase();
-
 		TGuardValue<int32> GuardedInputPos(InputPos, PropertyDef->GetParsePosition());
 		TGuardValue<int32> GuardedInputLine(InputLine, PropertyDef->GetLineNumber());
 
@@ -2652,6 +2655,23 @@ void FHeaderParser::VerifyPropertyMarkups(FUnrealClassDefinitionInfo& TargetClas
 				}
 			}
 		}
+
+		// Verify if native property setter and getter functions could actually be found while parsing class header
+		FPropertyBase& PropertyToken = PropertyDef->GetPropertyBase();
+		if (!PropertyToken.SetterName.IsEmpty())
+		{
+			if (PropertyToken.SetterName != TEXT("None") && !PropertyToken.bSetterFunctionFound)
+			{
+				LogError(TEXT("Property %s setter function %s not found"), *PropertyDef->GetName(), *PropertyToken.SetterName);
+			}
+		}
+		if (!PropertyToken.GetterName.IsEmpty())
+		{
+			if (PropertyToken.GetterName != TEXT("None") && !PropertyToken.bGetterFunctionFound)
+			{
+				LogError(TEXT("Property %s getter function %s not found"), *PropertyDef->GetName(), *PropertyToken.GetterName);
+			}
+		}
 	}
 }
 
@@ -2819,6 +2839,8 @@ void FHeaderParser::GetVarType(
 	FUnrealScriptStructDefinitionInfo* OwnerScriptStructDef = UHTCast<FUnrealScriptStructDefinitionInfo>(OwnerStructDef);
 	FUnrealClassDefinitionInfo* OwnerClassDef = UHTCast<FUnrealClassDefinitionInfo>(OwnerStructDef);
 	FName RepCallbackName = FName(NAME_None);
+	FString SetterName;
+	FString GetterName;
 
 	// Get flags.
 	EPropertyFlags Flags        = CPF_None;
@@ -3243,6 +3265,18 @@ void FHeaderParser::GetVarType(
 				}
 				break;
 
+				case EVariableSpecifier::Setter:
+				{
+					SetterName = RequireExactlyOneSpecifierValue(*this, Specifier);
+				}
+				break;
+
+				case EVariableSpecifier::Getter:
+				{
+					GetterName = RequireExactlyOneSpecifierValue(*this, Specifier);
+				}
+				break;
+
 				default:
 				{
 					LogError(TEXT("Unknown variable specifier '%s'"), *Specifier.Key);
@@ -4332,6 +4366,9 @@ void FHeaderParser::GetVarType(
 		}
 	}
 
+	VarProperty.SetterName = SetterName;
+	VarProperty.GetterName = GetterName;
+
 	// Perform some more specific validation on the property flags
 	if (VarProperty.PropertyFlags & CPF_PersistentInstance)
 	{
@@ -4916,6 +4953,12 @@ bool FHeaderParser::CompileDeclaration(TArray<FUnrealFunctionDefinitionInfo*>& D
 			if (CheckForSerialize(StructDef, Declaration))
 			{
 			}
+			else if (CheckForPropertySetterFunction(StructDef, Declaration))
+			{
+			}
+			else if (CheckForPropertyGetterFunction(StructDef, Declaration))
+			{
+			}
 		}
 	}
 	return Result;
@@ -5190,6 +5233,390 @@ static const TArray<FStringView> GSkipDeclarationWarningStrings =
 	//TEXT("USTRUCT"),
 };
 
+static bool ParseAccessorType(FString& OutAccessorType, FTokenReplay& Tokens)
+{
+	FToken Token;
+	Tokens.GetToken(Token);
+	if (!Token.IsIdentifier())
+	{
+		return false;
+	}
+	if (Token.IsIdentifier(TEXT("const"), ESearchCase::CaseSensitive))
+	{
+		// skip const. We allow const in paramater and return values even if the property is not const
+		Tokens.GetToken(Token);
+		if (!Token.IsIdentifier())
+		{
+			return false;
+		}
+	}
+
+	OutAccessorType = Token.Value;
+
+	Tokens.GetToken(Token);
+	if (Token.IsSymbol(TEXT("<"), ESearchCase::CaseSensitive))
+	{
+		// template type - add everything until the matching '>' into the type string
+		OutAccessorType += Token.Value;
+		int32 TemplateNestCount = 1;
+		while (TemplateNestCount > 0)
+		{
+			Tokens.GetToken(Token);
+			if (Token.TokenType == ETokenType::None)
+			{
+				return false;
+			}
+
+			OutAccessorType += Token.Value;
+			if (Token.IsSymbol(TEXT("<"), ESearchCase::CaseSensitive))
+			{
+				TemplateNestCount++;
+			}
+			else if (Token.IsSymbol(TEXT(">"), ESearchCase::CaseSensitive))
+			{
+				TemplateNestCount--;
+			}
+		}
+		Tokens.GetToken(Token);
+	}
+	// Skip '&' and add any '*'
+	for (; !Token.IsIdentifier(); Tokens.GetToken(Token))
+	{
+		if (Token.TokenType == ETokenType::None)
+		{
+			return false;
+		}
+
+		// Support for passing values by ref for setters
+		if (!Token.IsSymbol(TEXT("&"), ESearchCase::CaseSensitive))
+		{
+			OutAccessorType += Token.Value;
+		}
+	}
+	// Return the last token since it's either ')' in case of setters or the getter name which will be parsed by the caller
+	Tokens.UngetToken(Token);
+
+	return true;
+}
+
+bool FHeaderParser::CheckForPropertySetterFunction(FUnrealStructDefinitionInfo& StructDef, const FDeclaration& Declaration)
+{
+	FTokenReplay Tokens(Declaration.Tokens);
+	FUnrealClassDefinitionInfo& ClassDef = StructDef.AsClassChecked();
+
+	FToken Token;
+	if (!Tokens.GetToken(Token))
+	{
+		return false;
+	}
+
+	// Skip virtual keyword or any API macros
+	while (Token.IsIdentifier(TEXT("virtual"), ESearchCase::CaseSensitive) || Token.Value.EndsWith(TEXT("_API"), ESearchCase::CaseSensitive))
+	{
+		Tokens.GetToken(Token);
+	}
+
+	// Setters can only return void
+	if (!Token.IsIdentifier(TEXT("void"), ESearchCase::CaseSensitive))
+	{
+		return false;
+	}
+
+	FToken SetterNameToken;
+	Tokens.GetToken(SetterNameToken);
+	if (!SetterNameToken.IsIdentifier())
+	{
+		return false;
+	}
+
+	if (SetterNameToken.Value.Compare(TEXT("SetHardwareCursors"), ESearchCase::CaseSensitive) == 0)
+	{
+		static volatile int32 xx = 0;
+		xx++;
+	}
+
+	FUnrealPropertyDefinitionInfo* PropertyWithSetter = nullptr;
+	FString SetterFunctionName;
+	TArray<TSharedRef<FUnrealPropertyDefinitionInfo>>& Properties = ClassDef.GetProperties();
+	bool bExplicitSetter = false;
+
+	// First check if there are any properties that already specify this function name as a setter
+	for (TSharedRef<FUnrealPropertyDefinitionInfo>& Prop : Properties)
+	{
+		FPropertyBase& BaseProp = Prop->GetPropertyBase();
+		if (!BaseProp.bSetterFunctionFound && !BaseProp.SetterName.IsEmpty() && BaseProp.SetterName != TEXT("None"))
+		{
+			if (SetterNameToken.Value.Compare(Prop->GetPropertyBase().SetterName, ESearchCase::CaseSensitive) == 0)
+			{
+				SetterFunctionName = Prop->GetPropertyBase().SetterName;
+				PropertyWithSetter = &Prop.Get();
+				bExplicitSetter = true;
+				break;
+			}
+		}
+	}
+
+	// If not, try with autogenerated setter names
+	if (!PropertyWithSetter && FUHTConfig::Get().bAllowAutomaticSettersAndGetters && SetterNameToken.ValueStartsWith(TEXT("Set"), ESearchCase::CaseSensitive))
+	{
+		for (TSharedRef<FUnrealPropertyDefinitionInfo>& Prop : Properties)
+		{
+			FPropertyBase& BaseProp = Prop->GetPropertyBase();
+			if (!BaseProp.bSetterFunctionFound && BaseProp.SetterName.IsEmpty())
+			{
+				SetterFunctionName.Reset();
+				SetterFunctionName += TEXT("Set");
+				SetterFunctionName += Prop->GetName();
+
+				if (SetterNameToken.Value.Compare(SetterFunctionName, ESearchCase::CaseSensitive) == 0)
+				{
+					// We found a matching function name so we can move onto verifying if the declaration matches
+					PropertyWithSetter = &Prop.Get();
+					break;
+				}
+			}
+		}
+	}
+
+	if (!PropertyWithSetter)
+	{
+		return false;
+	}
+
+	Tokens.GetToken(Token);
+	if (!Token.IsSymbol(TEXT('(')))
+	{
+		if (bExplicitSetter)
+		{
+			LogError(TEXT("Expected '(' when parsing property %s setter function %s"), *PropertyWithSetter->GetName(), *SetterFunctionName);
+		}
+		return false;
+	}
+
+	FString ParameterType;
+	if (!ParseAccessorType(ParameterType, Tokens))
+	{
+		if (bExplicitSetter)
+		{
+			LogError(TEXT("Error when parsing parameter type for property %s setter function %s"), *PropertyWithSetter->GetName(), *SetterFunctionName);
+		}
+		return false;
+	}
+
+	// Skip parameter name
+	Tokens.GetToken(Token);
+	if (!Token.IsIdentifier())
+	{
+		if (bExplicitSetter)
+		{
+			LogError(TEXT("Expected parameter name when parsing property %s setter function %s"), *PropertyWithSetter->GetName(), *SetterFunctionName);
+		}
+		return false;
+	}
+
+	// Setter has only one parameter
+	Tokens.GetToken(Token);
+	if (!Token.IsSymbol(TEXT(')')))
+	{
+		if (bExplicitSetter)
+		{
+			LogError(TEXT("Expected ')' when parsing property %s setter function %s"), *PropertyWithSetter->GetName(), *SetterFunctionName);
+		}
+		return false;
+	}
+
+	// Parameter type and property type must match
+	FString ExtendedPropertyType;
+	FString PropertyType = PropertyWithSetter->GetCPPType(&ExtendedPropertyType, CPPF_Implementation | CPPF_ArgumentOrReturnValue | CPPF_NoRef);
+	ExtendedPropertyType.RemoveSpacesInline();
+	PropertyType += ExtendedPropertyType;
+	if (ParameterType.Compare(PropertyType, ESearchCase::CaseSensitive) != 0)
+	{
+		if (bExplicitSetter)
+		{
+			LogError(TEXT("Paramater type %s unsupported for property %s setter function %s (expected: %s)"),
+				*ParameterType, *PropertyWithSetter->GetName(), *SetterFunctionName, *PropertyType);
+		}
+		return false;
+	}
+
+	if ((GetCurrentCompilerDirective() & ECompilerDirective::WithEditor) != 0)
+	{
+		if (bExplicitSetter)
+		{
+			LogError(TEXT("Property %s setter function %s cannot be declared within WITH_EDITOR block. Use WITH_EDITORONLY_DATA instead."),
+				*PropertyWithSetter->GetName(), *SetterFunctionName);
+		}
+		return false;
+	}
+
+	if ((GetCurrentCompilerDirective() & ECompilerDirective::WithEditorOnlyData) != 0 && (PropertyWithSetter->GetPropertyFlags() & CPF_EditorOnly) == 0)
+	{
+		if (bExplicitSetter)
+		{
+			LogError(TEXT("Property %s is not editor-only but its setter function %s is."),
+				*PropertyWithSetter->GetName(), *SetterFunctionName);
+		}
+		return false;
+	}
+
+	PropertyWithSetter->GetPropertyBase().SetterName = *SetterFunctionName;
+	PropertyWithSetter->GetPropertyBase().bSetterFunctionFound = true;
+
+	return true;
+}
+
+bool FHeaderParser::CheckForPropertyGetterFunction(FUnrealStructDefinitionInfo& StructDef, const FDeclaration& Declaration)
+{
+	FTokenReplay Tokens(Declaration.Tokens);
+	FUnrealClassDefinitionInfo& ClassDef = StructDef.AsClassChecked();
+
+	FToken Token;
+	if (!Tokens.GetToken(Token))
+	{
+		return false;
+	}
+
+	// Skip virtual keyword and any API macros
+	while (Token.IsIdentifier(TEXT("virtual"), ESearchCase::CaseSensitive) || Token.Value.EndsWith(TEXT("_API"), ESearchCase::CaseSensitive))
+	{
+		Tokens.GetToken(Token);
+	}
+
+	Tokens.UngetToken(Token);
+	FString GetterType;
+	if (!ParseAccessorType(GetterType, Tokens))
+	{
+		// Unable to print a meaningful error before parsing the Getter name
+		return false;
+	}
+
+	FToken GetterNameToken;
+	Tokens.GetToken(GetterNameToken);
+	if (!GetterNameToken.IsIdentifier())
+	{
+		// Unable to print a meaningful error before parsing the Getter name
+		return false;
+	}
+
+	FUnrealPropertyDefinitionInfo* PropertyWithGetter = nullptr;
+	TArray<TSharedRef<FUnrealPropertyDefinitionInfo>>& Properties = ClassDef.GetProperties();
+	FString GetterFunctionName;
+	bool bExplicitGetter = false;
+
+	// First check if there are any properties that already specify this function name as a getter
+	for (TSharedRef<FUnrealPropertyDefinitionInfo>& Prop : Properties)
+	{
+		FPropertyBase& BaseProp = Prop->GetPropertyBase();
+		if (!BaseProp.bGetterFunctionFound && !BaseProp.GetterName.IsEmpty() && BaseProp.GetterName != TEXT("None"))
+		{
+			if (GetterNameToken.Value.Compare(Prop->GetPropertyBase().GetterName, ESearchCase::CaseSensitive) == 0)
+			{
+				GetterFunctionName = Prop->GetPropertyBase().GetterName;
+				PropertyWithGetter = &Prop.Get();
+				bExplicitGetter = true;
+				break;
+			}
+		}
+	}
+
+	// If not, try with autogenerated getter names
+	if (!PropertyWithGetter && FUHTConfig::Get().bAllowAutomaticSettersAndGetters && GetterNameToken.ValueStartsWith(TEXT("Get"), ESearchCase::CaseSensitive))
+	{
+		for (TSharedRef<FUnrealPropertyDefinitionInfo>& Prop : Properties)
+		{
+			FPropertyBase& BaseProp = Prop->GetPropertyBase();			
+			if (!BaseProp.bGetterFunctionFound && BaseProp.GetterName != TEXT("None"))
+			{
+				GetterFunctionName.Reset();
+				GetterFunctionName += TEXT("Get");
+				GetterFunctionName += Prop->GetName();
+				if (GetterNameToken.Value.Compare(GetterFunctionName, ESearchCase::CaseSensitive) == 0)
+				{
+					PropertyWithGetter = &Prop.Get();
+					break;
+				}
+			}
+		}
+	}
+
+	if (!PropertyWithGetter)
+	{
+		return false;
+	}
+
+	FString ExtendedPropertyType;
+	FString PropertyType = PropertyWithGetter->GetCPPType(&ExtendedPropertyType, CPPF_Implementation | CPPF_ArgumentOrReturnValue);
+	PropertyType += ExtendedPropertyType;
+	if (GetterType.Compare(PropertyType, ESearchCase::CaseSensitive) != 0)
+	{
+		if (bExplicitGetter)
+		{
+			LogError(TEXT("Paramater type %s unsupported for property %s getter function %s (expected: %s)"),
+				*GetterType, *PropertyWithGetter->GetName(), *GetterFunctionName, *PropertyType);
+		}
+		return false;
+	}
+
+	// Getter is a function that takes no arguments
+	Tokens.GetToken(Token);
+	if (!Token.IsSymbol(TEXT('(')))
+	{
+		if (bExplicitGetter)
+		{
+			LogError(TEXT("Expected '(' when parsing property %s getter function %s"), *PropertyWithGetter->GetName(), *GetterFunctionName);
+		}
+		return false;
+	}
+
+	Tokens.GetToken(Token);
+	if (!Token.IsSymbol(TEXT(')')))
+	{
+		if (bExplicitGetter)
+		{
+			LogError(TEXT("Expected ')' when parsing property %s getter function %s"), *PropertyWithGetter->GetName(), *GetterFunctionName);
+		}
+		return false;
+	}
+
+	// Getters should be const functions
+	Tokens.GetToken(Token);
+	if (!Token.IsIdentifier(TEXT("const"), ESearchCase::CaseSensitive))
+	{
+		if (bExplicitGetter)
+		{
+			LogError(TEXT("Property %s getter function %s must be const"), *PropertyWithGetter->GetName(), *GetterFunctionName);
+		}
+		return false;
+	}
+
+	if ((GetCurrentCompilerDirective() & ECompilerDirective::WithEditor) != 0)
+	{
+		if (bExplicitGetter)
+		{
+			LogError(TEXT("Property %s setter function %s cannot be declared within WITH_EDITOR block. Use WITH_EDITORONLY_DATA instead."),
+				*PropertyWithGetter->GetName(), *GetterFunctionName);
+		}
+		return false;
+	}
+
+	if ((GetCurrentCompilerDirective() & ECompilerDirective::WithEditorOnlyData) != 0 && (PropertyWithGetter->GetPropertyFlags() & CPF_EditorOnly) == 0)
+	{
+		if (bExplicitGetter)
+		{
+			LogError(TEXT("Property %s is not editor-only but its getter function %s is."),
+				*PropertyWithGetter->GetName(), *GetterFunctionName);
+		}
+		return false;
+	}
+
+	PropertyWithGetter->GetPropertyBase().GetterName = GetterFunctionName;
+	PropertyWithGetter->GetPropertyBase().bGetterFunctionFound = true;
+
+	return true;
+}
+
+
 bool FHeaderParser::SkipDeclaration(FToken& Token)
 {
 	// Store the current value of PrevComment so it can be restored after we parsed everything.
diff --git a/Engine/Source/Programs/UnrealHeaderTool/Private/HeaderParser.h b/Engine/Source/Programs/UnrealHeaderTool/Private/HeaderParser.h
--- a/Engine/Source/Programs/UnrealHeaderTool/Private/HeaderParser.h
+++ b/Engine/Source/Programs/UnrealHeaderTool/Private/HeaderParser.h
@@ -626,6 +626,12 @@ private:
 	// Check to see if the declaration is a serialize
 	static bool CheckForSerialize(FUnrealStructDefinitionInfo& StructDef, const FDeclaration& Declaration);
 
+	// Check to see if the declaration is a property getter function
+	bool CheckForPropertyGetterFunction(FUnrealStructDefinitionInfo& StructDef, const FDeclaration& Declaration);
+
+	// Check to see if the declaration is a property setter function
+	bool CheckForPropertySetterFunction(FUnrealStructDefinitionInfo& StructDef, const FDeclaration& Declaration);
+
 	// Names that cannot be used enums, UStructs, or UClasses
 	static TArray<FString> ReservedTypeNames;
 
diff --git a/Engine/Source/Programs/UnrealHeaderTool/Private/ParserHelper.h b/Engine/Source/Programs/UnrealHeaderTool/Private/ParserHelper.h
--- a/Engine/Source/Programs/UnrealHeaderTool/Private/ParserHelper.h
+++ b/Engine/Source/Programs/UnrealHeaderTool/Private/ParserHelper.h
@@ -208,6 +208,15 @@ public:
 	EPointerType::Type PointerType = EPointerType::None;
 	EIntType IntType = EIntType::None;
 
+	/** Setter function name. Either specified in the UPROPERTY macro or autogenerated while parsing class declaration */
+	FString SetterName;	
+	/** Getter function name. Either specified in the UPROPERTY macro or autogenerated while parsing class declaration */
+	FString GetterName;
+	/** True if setter function declaration was found while parsing class header */
+	bool bSetterFunctionFound = false;
+	/** True if getter function declaration was found while parsing class header */
+	bool bGetterFunctionFound = false;
+
 public:
 	/** @name Constructors */
 	//@{
diff --git a/Engine/Source/Programs/UnrealHeaderTool/Private/Specifiers/VariableSpecifiers.def b/Engine/Source/Programs/UnrealHeaderTool/Private/Specifiers/VariableSpecifiers.def
--- a/Engine/Source/Programs/UnrealHeaderTool/Private/Specifiers/VariableSpecifiers.def
+++ b/Engine/Source/Programs/UnrealHeaderTool/Private/Specifiers/VariableSpecifiers.def
@@ -22,6 +22,7 @@ VARIABLE_SPECIFIER(EditFixedSize)
 VARIABLE_SPECIFIER(EditInline)
 VARIABLE_SPECIFIER(EditInstanceOnly)
 VARIABLE_SPECIFIER(Export)
+VARIABLE_SPECIFIER(Getter)
 VARIABLE_SPECIFIER(GlobalConfig)
 VARIABLE_SPECIFIER(Instanced)
 VARIABLE_SPECIFIER(Interp)
@@ -36,6 +37,7 @@ VARIABLE_SPECIFIER(Replicated)
 VARIABLE_SPECIFIER(ReplicatedUsing)
 VARIABLE_SPECIFIER(RepRetry)
 VARIABLE_SPECIFIER(SaveGame)
+VARIABLE_SPECIFIER(Setter)
 VARIABLE_SPECIFIER(SimpleDisplay)
 VARIABLE_SPECIFIER(SkipSerialization)
 VARIABLE_SPECIFIER(TextExportTransient)
diff --git a/Engine/Source/Runtime/AIModule/Private/BehaviorTree/BlueprintNodeHelpers.cpp b/Engine/Source/Runtime/AIModule/Private/BehaviorTree/BlueprintNodeHelpers.cpp
--- a/Engine/Source/Runtime/AIModule/Private/BehaviorTree/BlueprintNodeHelpers.cpp
+++ b/Engine/Source/Runtime/AIModule/Private/BehaviorTree/BlueprintNodeHelpers.cpp
@@ -120,7 +120,7 @@ namespace BlueprintNodeHelpers
 		}
 		else
 		{
-			Prop->ExportTextItem(ExportedStringValue, PropertyAddr, NULL, NULL, PPF_PropertyWindow, NULL);
+			Prop->ExportTextItem_Direct(ExportedStringValue, PropertyAddr, NULL, NULL, PPF_PropertyWindow, NULL);
 		}
 
 		const bool bIsBool = Prop->IsA(FBoolProperty::StaticClass());
diff --git a/Engine/Source/Runtime/CoreUObject/Private/UObject/EnumProperty.cpp b/Engine/Source/Runtime/CoreUObject/Private/UObject/EnumProperty.cpp
--- a/Engine/Source/Runtime/CoreUObject/Private/UObject/EnumProperty.cpp
+++ b/Engine/Source/Runtime/CoreUObject/Private/UObject/EnumProperty.cpp
@@ -250,7 +250,7 @@ FString FEnumProperty::GetCPPType(FString* ExtendedTypeText, uint32 CPPExportFla
 	return EnumName;
 }
 
-void FEnumProperty::ExportTextItem(FString& ValueStr, const void* PropertyValue, const void* DefaultValue, UObject* Parent, int32 PortFlags, UObject* ExportRootScope) const
+void FEnumProperty::ExportText_Internal(FString& ValueStr, const void* PropertyValueOrContainer, EPropertyPointerType PropertyPointerType, const void* DefaultValue, UObject* Parent, int32 PortFlags, UObject* ExportRootScope) const
 {
 	if (Enum == nullptr)
 	{
@@ -265,8 +265,20 @@ void FEnumProperty::ExportTextItem(FString& ValueStr, const void* PropertyValue,
 
 	check(UnderlyingProp);
 
+	int64 LocalValue = 0;
+	void* PropertyValue = nullptr;
 	FNumericProperty* LocalUnderlyingProp = UnderlyingProp;
 
+	if (PropertyPointerType == EPropertyPointerType::Container && HasGetter())
+	{
+		PropertyValue = &LocalValue;
+		GetValue_InContainer(PropertyValueOrContainer, PropertyValue);
+	}
+	else
+	{
+		PropertyValue = PointerToValuePtr(PropertyValueOrContainer, PropertyPointerType);
+	}
+
 	if (PortFlags & PPF_ExportCpp)
 	{
 		const int64 ActualValue = LocalUnderlyingProp->GetSignedIntPropertyValue(PropertyValue);
@@ -291,7 +303,7 @@ void FEnumProperty::ExportTextItem(FString& ValueStr, const void* PropertyValue,
 
 	if (PortFlags & PPF_ConsoleVariable)
 	{
-		UnderlyingProp->ExportTextItem(ValueStr, PropertyValue, DefaultValue, Parent, PortFlags, ExportRootScope);
+		UnderlyingProp->ExportText_Internal(ValueStr, PropertyValue, EPropertyPointerType::Direct, DefaultValue, Parent, PortFlags, ExportRootScope);
 		return;
 	}
 
@@ -320,7 +332,7 @@ void FEnumProperty::ExportTextItem(FString& ValueStr, const void* PropertyValue,
 	}
 }
 
-const TCHAR* FEnumProperty::ImportText_Internal(const TCHAR* InBuffer, void* Data, int32 PortFlags, UObject* Parent, FOutputDevice* ErrorText) const
+const TCHAR* FEnumProperty::ImportText_Internal(const TCHAR* InBuffer, void* ContainerOrPropertyPtr, EPropertyPointerType PropertyPointerType, UObject* Parent, int32 PortFlags, FOutputDevice* ErrorText) const
 {
 	check(Enum);
 	check(UnderlyingProp);
@@ -339,7 +351,15 @@ const TCHAR* FEnumProperty::ImportText_Internal(const TCHAR* InBuffer, void* Dat
 			}
 			if (EnumIndex != INDEX_NONE)
 			{
-				UnderlyingProp->SetIntPropertyValue(Data, Enum->GetValueByIndex(EnumIndex));
+				int64 EnumValue = Enum->GetValueByIndex(EnumIndex);
+				if (PropertyPointerType == EPropertyPointerType::Container && HasSetter())
+				{
+					SetValue_InContainer(ContainerOrPropertyPtr, &EnumValue);
+				}
+				else
+				{
+					UnderlyingProp->SetIntPropertyValue(PointerToValuePtr(ContainerOrPropertyPtr, PropertyPointerType), EnumValue);
+				}
 				return Buffer;
 			}
 
@@ -359,7 +379,8 @@ const TCHAR* FEnumProperty::ImportText_Internal(const TCHAR* InBuffer, void* Dat
 		}
 	}
 
-	const TCHAR* Result = UnderlyingProp->ImportText(InBuffer, Data, PortFlags, Parent, ErrorText);
+	// UnderlyingProp has a 0 offset so we need to make sure we convert the container pointer to the actual value pointer
+	const TCHAR* Result = UnderlyingProp->ImportText_Internal(InBuffer, PointerToValuePtr(ContainerOrPropertyPtr, PropertyPointerType), EPropertyPointerType::Direct, Parent, PortFlags, ErrorText);
 	return Result;
 }
 
diff --git a/Engine/Source/Runtime/CoreUObject/Private/UObject/FieldPathProperty.cpp b/Engine/Source/Runtime/CoreUObject/Private/UObject/FieldPathProperty.cpp
--- a/Engine/Source/Runtime/CoreUObject/Private/UObject/FieldPathProperty.cpp
+++ b/Engine/Source/Runtime/CoreUObject/Private/UObject/FieldPathProperty.cpp
@@ -87,9 +87,18 @@ void FFieldPathProperty::SerializeItem(FStructuredArchive::FSlot Slot, void* Val
 	Slot << *FieldPtr;
 }
 
-void FFieldPathProperty::ExportTextItem( FString& ValueStr, const void* PropertyValue, const void* DefaultValue, UObject* Parent, int32 PortFlags, UObject* ExportRootScope ) const
+void FFieldPathProperty::ExportText_Internal( FString& ValueStr, const void* PropertyValueOrContainer, EPropertyPointerType PropertyPointerType, const void* DefaultValue, UObject* Parent, int32 PortFlags, UObject* ExportRootScope ) const
 {
-	const FFieldPath& Value = GetPropertyValue(PropertyValue);
+	FFieldPath Value;
+	
+	if (PropertyPointerType == EPropertyPointerType::Container && HasGetter())
+	{
+		GetValue_InContainer(PropertyValueOrContainer, &Value);
+	}
+	else
+	{
+		Value = GetPropertyValue(PointerToValuePtr(PropertyValueOrContainer, PropertyPointerType));
+	}
 
 	if (PortFlags & PPF_ExportCpp)
 	{
@@ -116,11 +125,10 @@ void FFieldPathProperty::ExportTextItem( FString& ValueStr, const void* Property
 	}
 }
 
-const TCHAR* FFieldPathProperty::ImportText_Internal( const TCHAR* Buffer, void* Data, int32 PortFlags, UObject* Parent, FOutputDevice* ErrorText ) const
+const TCHAR* FFieldPathProperty::ImportText_Internal( const TCHAR* Buffer, void* ContainerOrPropertyPtr, EPropertyPointerType PropertyPointerType, UObject* Parent, int32 PortFlags, FOutputDevice* ErrorText ) const
 {
 	check(Buffer);
-	FFieldPath* PathPtr = GetPropertyValuePtr(Data);
-	check(PathPtr);
+	FFieldPath ImportedPath;
 	FString PathName;
 
 	if (!(PortFlags & PPF_Delimited))
@@ -157,8 +165,17 @@ const TCHAR* FFieldPathProperty::ImportText_Internal( const TCHAR* Buffer, void*
 			}
 			PathName = MoveTemp(UnquotedPathName);
 		}
-		PathPtr->Generate(*PathName);
-	}
+		ImportedPath.Generate(*PathName);
+		if (PropertyPointerType == EPropertyPointerType::Container && HasSetter())
+		{
+			SetValue_InContainer(ContainerOrPropertyPtr, ImportedPath);
+		}
+		else
+		{
+			FFieldPath* PathPtr = GetPropertyValuePtr(PointerToValuePtr(ContainerOrPropertyPtr, PropertyPointerType));
+			*PathPtr = ImportedPath;
+		}
+	}	
 
 	return Buffer;
 }
diff --git a/Engine/Source/Runtime/CoreUObject/Private/UObject/Obj.cpp b/Engine/Source/Runtime/CoreUObject/Private/UObject/Obj.cpp
--- a/Engine/Source/Runtime/CoreUObject/Private/UObject/Obj.cpp
+++ b/Engine/Source/Runtime/CoreUObject/Private/UObject/Obj.cpp
@@ -1884,7 +1884,7 @@ static void GetAssetRegistryTagFromProperty(const void* BaseMemoryLocation, cons
 		else
 		{
 			FString PropertyStr;
-			Prop->ExportTextItem(PropertyStr, Bundles, Bundles, nullptr, PPF_None);
+			Prop->ExportTextItem_Direct(PropertyStr, Bundles, Bundles, nullptr, PPF_None);
 			OutTags.Add(UObject::FAssetRegistryTag(GAssetBundleDataName, MoveTemp(PropertyStr), UObject::FAssetRegistryTag::ETagType::TT_Alphabetical));
 		}
 	}
@@ -1926,7 +1926,7 @@ static void GetAssetRegistryTagFromProperty(const void* BaseMemoryLocation, cons
 
 		FString PropertyStr;
 		const uint8* PropertyAddr = Prop->ContainerPtrToValuePtr<uint8>(BaseMemoryLocation);
-		Prop->ExportTextItem(PropertyStr, PropertyAddr, PropertyAddr, nullptr, PPF_None);
+		Prop->ExportTextItem_Direct(PropertyStr, PropertyAddr, PropertyAddr, nullptr, PPF_None);
 
 		OutTags.Add(UObject::FAssetRegistryTag(Prop->GetFName(), MoveTemp(PropertyStr), TagType));
 	}
@@ -2504,7 +2504,7 @@ void UObject::LoadConfig( UClass* ConfigClass/*=NULL*/, const TCHAR* InFilename/
 
 				if (bFoundValue)
 				{
-					if (Property->ImportText(*Value, Property->ContainerPtrToValuePtr<uint8>(this, i), PortFlags, this) == NULL)
+					if (Property->ImportText_Direct(*Value, Property->ContainerPtrToValuePtr<uint8>(this, i), this, PortFlags) == NULL)
 					{
 						// this should be an error as the properties from the .ini / .int file are not correctly being read in and probably are affecting things in subtle ways
 						UE_LOG(LogObj, Error, TEXT("LoadConfig (%s): import failed for %s in: %s"), *GetPathName(), *Property->GetName(), *Value);
@@ -2553,7 +2553,7 @@ void UObject::LoadConfig( UClass* ConfigClass/*=NULL*/, const TCHAR* InFilename/
 					ArrayHelper.EmptyAndAddValues(List.Num());
 					for( int32 i=List.Num()-1,c=0; i>=0; i--,c++ )
 					{
-						Array->Inner->ImportText( *List[i].GetValue(), ArrayHelper.GetRawPtr(c), PortFlags, this );
+						Array->Inner->ImportText_Direct( *List[i].GetValue(), ArrayHelper.GetRawPtr(c), this, PortFlags );
 					}
 				}
 				else
@@ -2578,7 +2578,7 @@ void UObject::LoadConfig( UClass* ConfigClass/*=NULL*/, const TCHAR* InFilename/
 						{
 							// expand the array if necessary so that Index is a valid element
 							ArrayHelper.ExpandForIndex(Index);
-							Array->Inner->ImportText(*ElementValue->GetValue(), ArrayHelper.GetRawPtr(Index), PortFlags, this);
+							Array->Inner->ImportText_Direct(*ElementValue->GetValue(), ArrayHelper.GetRawPtr(Index), this, PortFlags);
 						}
 
 						Index++;
@@ -2719,7 +2719,7 @@ void UObject::SaveConfig( uint64 Flags, const TCHAR* InFilename, FConfigCacheIni
 					for( int32 i=0; i<ArrayHelper.Num(); i++ )
 					{
 						FString	Buffer;
-						Array->Inner->ExportTextItem( Buffer, ArrayHelper.GetRawPtr(i), ArrayHelper.GetRawPtr(i), this, PortFlags );
+						Array->Inner->ExportTextItem_Direct( Buffer, ArrayHelper.GetRawPtr(i), ArrayHelper.GetRawPtr(i), this, PortFlags );
 						Sec->Add(*CompleteKey, *Buffer);
 					}
 				}
@@ -3216,7 +3216,7 @@ void UObject::ParseParms( const TCHAR* Parms )
 			FString Value;
 			if( FParse::Value(Parms,*(FString(*It->GetName())+TEXT("=")),Value) )
 			{
-				It->ImportText( *Value, It->ContainerPtrToValuePtr<uint8>(this), 0, this );
+				It->ImportText_InContainer( *Value, this, this, 0 );
 			}
 		}
 	}
@@ -3360,7 +3360,7 @@ static void PerformSetCommand( const TCHAR* Str, FOutputDevice& Ar, bool bNotify
 						Object->PreEditChange(Property);
 					}
 #endif // WITH_EDITOR
-					Property->ImportText(Str, Property->ContainerPtrToValuePtr<uint8>(Object), 0, Object);
+					Property->ImportText_InContainer(Str, Object, Object, 0);
 #if WITH_EDITOR
 					if (!Object->HasAnyFlags(RF_ClassDefaultObject) && bNotifyObjectOfChange)
 					{
@@ -3778,7 +3778,7 @@ bool StaticExec( UWorld* InWorld, const TCHAR* Cmd, FOutputDevice& Ar )
 									{
 										FString ResultStr;
 										uint8* ElementData = BaseData + ArrayIndex * ElementSize;
-										ExportProperty->ExportTextItem(ResultStr, ElementData, NULL, CurrentObject, PPF_IncludeTransient);
+										ExportProperty->ExportTextItem_Direct(ResultStr, ElementData, NULL, CurrentObject, PPF_IncludeTransient);
 
 										if (bShowDetailedInfo)
 										{
diff --git a/Engine/Source/Runtime/CoreUObject/Private/UObject/Property.cpp b/Engine/Source/Runtime/CoreUObject/Private/UObject/Property.cpp
--- a/Engine/Source/Runtime/CoreUObject/Private/UObject/Property.cpp
+++ b/Engine/Source/Runtime/CoreUObject/Private/UObject/Property.cpp
@@ -1076,16 +1076,17 @@ bool FProperty::ExportText_Direct
 	const void*	Data,
 	const void*	Delta,
 	UObject*	Parent,
-	int32			PortFlags,
+	int32		PortFlags,
 	UObject*	ExportRootScope
 	) const
 {
 	if( Data==Delta || !Identical(Data,Delta,PortFlags) )
 	{
-		ExportTextItem
+		ExportText_Internal
 			(
 			ValueStr,
 			(uint8*)Data,
+			EPropertyPointerType::Direct,
 			(uint8*)Delta,
 			Parent,
 			PortFlags,
@@ -1548,7 +1549,7 @@ const TCHAR* FProperty::ImportSingleProperty( const TCHAR* Str, void* DestData,
 				{
 					int32	Index = ArrayHelper.AddValue();
 
-					const TCHAR* Result = ArrayProperty->Inner->ImportText(Str, ArrayHelper.GetRawPtr(Index), PortFlags, SubobjectOuter, Warn);
+					const TCHAR* Result = ArrayProperty->Inner->ImportText_Direct(Str, ArrayHelper.GetRawPtr(Index), SubobjectOuter, PortFlags, Warn);
 					if ( Result == NULL || Result == Str )
 					{
 						Warn->Logf(ELogVerbosity::Warning, TEXT("Unable to parse parameter value '%s' in defaultproperties array operation: %s"), Str, Start);
@@ -1567,7 +1568,7 @@ const TCHAR* FProperty::ImportSingleProperty( const TCHAR* Str, void* DestData,
 					ArrayProperty->Inner->InitializeValue(Temp);
 							
 					// export the value specified to a temporary buffer
-					const TCHAR* Result = ArrayProperty->Inner->ImportText(Str, Temp, PortFlags, SubobjectOuter, Warn);
+					const TCHAR* Result = ArrayProperty->Inner->ImportText_Direct(Str, Temp, SubobjectOuter, PortFlags, Warn);
 					if ( Result == NULL || Result == Str )
 					{
 						Warn->Logf(ELogVerbosity::Error, TEXT("Unable to parse parameter value '%s' in defaultproperties array operation: %s"), Str, Start);
@@ -1687,7 +1688,7 @@ const TCHAR* FProperty::ImportSingleProperty( const TCHAR* Str, void* DestData,
 					ArrayHelper.ExpandForIndex(Index);
 
 					FStringOutputDevice ImportError;
-					const TCHAR* Result = ArrayProperty->Inner->ImportText(Str, ArrayHelper.GetRawPtr(Index), PortFlags, SubobjectOuter, &ImportError);
+					const TCHAR* Result = ArrayProperty->Inner->ImportText_Direct(Str, ArrayHelper.GetRawPtr(Index), SubobjectOuter, PortFlags, &ImportError);
 					// Spit any error we had while importing property
 					if (ImportError.Len() > 0)
 					{
@@ -1718,7 +1719,7 @@ const TCHAR* FProperty::ImportSingleProperty( const TCHAR* Str, void* DestData,
 
 					FStringOutputDevice ImportError;
 
-					const TCHAR* Result = Property->ImportText(Str, Property->ContainerPtrToValuePtr<void>(DestData, Index), PortFlags, SubobjectOuter, &ImportError);
+					const TCHAR* Result = Property->ImportText_Direct(Str, Property->ContainerPtrToValuePtr<void>(DestData, Index), SubobjectOuter, PortFlags, &ImportError);
 					
 					// Spit any error we had while importing property
 					if (ImportError.Len() > 0)
@@ -1826,9 +1827,9 @@ UPropertyWrapper* FProperty::GetUPropertyWrapper()
 }
 #endif //  WITH_EDITORONLY_DATA
 
-void FFloatProperty::ExportTextItem(FString& ValueStr, const void* PropertyValue, const void* DefaultValue, UObject* Parent, int32 PortFlags, UObject* ExportRootScope) const
+void FFloatProperty::ExportText_Internal(FString& ValueStr, const void* PropertyValueOrContainer, EPropertyPointerType PropertyPointerType, const void* DefaultValue, UObject* Parent, int32 PortFlags, UObject* ExportRootScope) const
 {
-	Super::ExportTextItem(ValueStr, PropertyValue, DefaultValue, Parent, PortFlags, ExportRootScope);
+	Super::ExportText_Internal(ValueStr, PropertyValueOrContainer, PropertyPointerType, DefaultValue, Parent, PortFlags, ExportRootScope);
 
 	if (0 != (PortFlags & PPF_ExportCpp))
 	{
diff --git a/Engine/Source/Runtime/CoreUObject/Private/UObject/PropertyArray.cpp b/Engine/Source/Runtime/CoreUObject/Private/UObject/PropertyArray.cpp
--- a/Engine/Source/Runtime/CoreUObject/Private/UObject/PropertyArray.cpp
+++ b/Engine/Source/Runtime/CoreUObject/Private/UObject/PropertyArray.cpp
@@ -461,7 +461,7 @@ FString FArrayProperty::GetCPPMacroType( FString& ExtendedTypeText ) const
 	ExtendedTypeText = Inner->GetCPPType();
 	return TEXT("TARRAY");
 }
-void FArrayProperty::ExportTextItem( FString& ValueStr, const void* PropertyValue, const void* DefaultValue, UObject* Parent, int32 PortFlags, UObject* ExportRootScope ) const
+void FArrayProperty::ExportText_Internal( FString& ValueStr, const void* ContainerOrPropertyPtr, EPropertyPointerType PropertyPointerType, const void* DefaultValue, UObject* Parent, int32 PortFlags, UObject* ExportRootScope ) const
 {
 	checkSlow(Inner);
 
@@ -473,7 +473,34 @@ void FArrayProperty::ExportTextItem( FString& ValueStr, const void* PropertyValu
 		return;
 	}
 
-	FScriptArrayHelper ArrayHelper(this, PropertyValue);
+	uint8* TempArrayStorage = nullptr;
+	void* PropertyValuePtr = nullptr;
+	if (PropertyPointerType == EPropertyPointerType::Container && HasGetter())
+	{
+		// Allocate temporary map as we first need to initialize it with the value provided by the getter function and then export it
+		TempArrayStorage = (uint8*)FMemory::MallocZeroed(GetSize());
+		if (!HasAnyPropertyFlags(CPF_ZeroConstructor)) // this stuff is already zero
+		{
+			InitializeValue(TempArrayStorage);
+		}
+		PropertyValuePtr = TempArrayStorage;
+		FProperty::GetValue_InContainer(ContainerOrPropertyPtr, PropertyValuePtr);
+	}
+	else
+	{
+		PropertyValuePtr = PointerToValuePtr(ContainerOrPropertyPtr, PropertyPointerType);
+	}
+
+	ON_SCOPE_EXIT
+	{
+		if (TempArrayStorage)
+		{
+			DestroyValue(TempArrayStorage);
+			FMemory::Free(TempArrayStorage);
+		}
+	};
+
+	FScriptArrayHelper ArrayHelper(this, PropertyValuePtr);
 
 	int32 DefaultSize = 0;
 	if (DefaultValue)
@@ -551,7 +578,7 @@ void FArrayProperty::ExportTextInnerItem(FString& ValueStr, const FProperty* Inn
 			}
 		}
 
-		Inner->ExportTextItem( ValueStr, PropData, PropDefault, Parent, PortFlags|PPF_Delimited, ExportRootScope );
+		Inner->ExportTextItem_Direct( ValueStr, PropData, PropDefault, Parent, PortFlags|PPF_Delimited, ExportRootScope );
 	}
 
 	if ((Count > 0) && !bReadableForm)
@@ -565,11 +592,39 @@ void FArrayProperty::ExportTextInnerItem(FString& ValueStr, const FProperty* Inn
 	}
 }
 
-const TCHAR* FArrayProperty::ImportText_Internal(const TCHAR* Buffer, void* Data, int32 PortFlags, UObject* OwnerObject, FOutputDevice* ErrorText) const
+const TCHAR* FArrayProperty::ImportText_Internal(const TCHAR* Buffer, void* ContainerOrPropertyPtr, EPropertyPointerType PropertyPointerType, UObject* OwnerObject, int32 PortFlags, FOutputDevice* ErrorText) const
 {
-	FScriptArrayHelper ArrayHelper(this, Data);
+	uint8* TempArrayStorage = nullptr;
+	ON_SCOPE_EXIT
+	{
+		if (TempArrayStorage)
+		{
+			// TempArrayStorage is used by property setter so if it was allocated call the setter now
+			FProperty::SetValue_InContainer(ContainerOrPropertyPtr, TempArrayStorage);
+
+			// Destroy and free the temp array used by property setter
+			DestroyValue(TempArrayStorage);
+			FMemory::Free(TempArrayStorage);
+		}
+	};
+
+	void* ArrayPtr = nullptr;
+	if (PropertyPointerType == EPropertyPointerType::Container && HasSetter())
+	{
+		// Allocate temporary map as we first need to initialize it with the parsed items and then use the setter to update the property
+		TempArrayStorage = (uint8*)FMemory::MallocZeroed(GetSize());
+		if (!HasAnyPropertyFlags(CPF_ZeroConstructor)) // this stuff is already zero
+		{
+			InitializeValue(TempArrayStorage);
+		}
+	}
+	else
+	{
+		ArrayPtr = PointerToValuePtr(ContainerOrPropertyPtr, PropertyPointerType);
+	}
 
-	return ImportTextInnerItem(Buffer, Inner, Data, PortFlags, OwnerObject, &ArrayHelper, ErrorText);
+	FScriptArrayHelper ArrayHelper(this, ArrayPtr);
+	return ImportTextInnerItem(Buffer, Inner, ArrayPtr, PortFlags, OwnerObject, &ArrayHelper, ErrorText);
 }
 
 const TCHAR* FArrayProperty::ImportTextInnerItem( const TCHAR* Buffer, const FProperty* Inner, void* Data, int32 PortFlags, UObject* Parent, FScriptArrayHelper* ArrayHelper, FOutputDevice* ErrorText )
@@ -609,7 +664,8 @@ const TCHAR* FArrayProperty::ImportTextInnerItem( const TCHAR* Buffer, const FPr
 		{
 			uint8* Address = ArrayHelper ? ArrayHelper->GetRawPtr(Index) : ((uint8*)Data + Inner->ElementSize * Index);
 			// Parse the item
-			Buffer = Inner->ImportText(Buffer, Address, PortFlags | PPF_Delimited, Parent, ErrorText);
+			check(Inner->GetOffset_ForInternal() == 0);
+			Buffer = Inner->ImportText_Direct(Buffer, Address, Parent, PortFlags | PPF_Delimited, ErrorText);
 
 			if(!Buffer)
 			{
diff --git a/Engine/Source/Runtime/CoreUObject/Private/UObject/PropertyBaseObject.cpp b/Engine/Source/Runtime/CoreUObject/Private/UObject/PropertyBaseObject.cpp
--- a/Engine/Source/Runtime/CoreUObject/Private/UObject/PropertyBaseObject.cpp
+++ b/Engine/Source/Runtime/CoreUObject/Private/UObject/PropertyBaseObject.cpp
@@ -228,9 +228,18 @@ FString FObjectPropertyBase::GetExportPath(const UObject* Object, const UObject*
 	return FString::Printf( TEXT("%s'%s'"), *Object->GetClass()->GetName(), *PathName );
 }
 
-void FObjectPropertyBase::ExportTextItem( FString& ValueStr, const void* PropertyValue, const void* DefaultValue, UObject* Parent, int32 PortFlags, UObject* ExportRootScope ) const
+void FObjectPropertyBase::ExportText_Internal( FString& ValueStr, const void* PropertyValueOrContainer, EPropertyPointerType PropertyPointerType, const void* DefaultValue, UObject* Parent, int32 PortFlags, UObject* ExportRootScope ) const
 {
-	UObject* Temp = GetObjectPropertyValue(PropertyValue);
+	UObject* Temp = nullptr;
+
+	if (PropertyPointerType == EPropertyPointerType::Container && HasGetter())
+	{
+		GetValue_InContainer(PropertyValueOrContainer, &Temp);
+	}
+	else
+	{
+		Temp = GetObjectPropertyValue(PointerToValuePtr(PropertyValueOrContainer, PropertyPointerType));
+	}
 
 	if (0 != (PortFlags & PPF_ExportCpp))
 	{
@@ -365,7 +374,7 @@ bool FObjectPropertyBase::ParseObjectPropertyValue(const FProperty* Property, UO
 	return true;
 }
 
-const TCHAR* FObjectPropertyBase::ImportText_Internal( const TCHAR* InBuffer, void* Data, int32 PortFlags, UObject* Parent, FOutputDevice* ErrorText ) const
+const TCHAR* FObjectPropertyBase::ImportText_Internal( const TCHAR* InBuffer, void* ContainerOrPropertyPtr, EPropertyPointerType PropertyPointerType, UObject* Parent, int32 PortFlags, FOutputDevice* ErrorText ) const
 {
 	const TCHAR* Buffer = InBuffer;
 	UObject* Result = nullptr;
@@ -395,7 +404,14 @@ const TCHAR* FObjectPropertyBase::ImportText_Internal( const TCHAR* InBuffer, vo
 		}
 	}
 
-	SetObjectPropertyValue(Data, Result);
+	if (PropertyPointerType == EPropertyPointerType::Container && HasSetter())
+	{
+		SetValue_InContainer(ContainerOrPropertyPtr, &Result); // @todo: need to make sure SetValue follows whatever SetObjectPropertyValue does for object properties
+	}
+	else
+	{
+		SetObjectPropertyValue(PointerToValuePtr(ContainerOrPropertyPtr, PropertyPointerType), Result);
+	}
 	return Buffer;
 }
 
diff --git a/Engine/Source/Runtime/CoreUObject/Private/UObject/PropertyBool.cpp b/Engine/Source/Runtime/CoreUObject/Private/UObject/PropertyBool.cpp
--- a/Engine/Source/Runtime/CoreUObject/Private/UObject/PropertyBool.cpp
+++ b/Engine/Source/Runtime/CoreUObject/Private/UObject/PropertyBool.cpp
@@ -344,11 +344,19 @@ void FBoolProperty::AppendSchemaHash(FBlake3& Builder, bool bSkipEditorOnly) con
 #endif
 
 
-void FBoolProperty::ExportTextItem( FString& ValueStr, const void* PropertyValue, const void* DefaultValue, UObject* Parent, int32 PortFlags, UObject* ExportRootScope ) const
+void FBoolProperty::ExportText_Internal( FString& ValueStr, const void* ContainerOrPropertyPtr, EPropertyPointerType PropertyPointerType, const void* DefaultValue, UObject* Parent, int32 PortFlags, UObject* ExportRootScope ) const
 {
 	check(FieldSize != 0);
-	const uint8* ByteValue = (uint8*)PropertyValue + ByteOffset;
-	const bool bValue = 0 != ((*ByteValue) & FieldMask);
+	uint8 LocalByteValue = 0;
+	if (PropertyPointerType == EPropertyPointerType::Container && HasGetter())
+	{
+		GetValue_InContainer(ContainerOrPropertyPtr, &LocalByteValue);
+	}
+	else
+	{
+		LocalByteValue = *((uint8*)PointerToValuePtr(ContainerOrPropertyPtr, PropertyPointerType) + ByteOffset);
+	}
+	const bool bValue = 0 != (LocalByteValue & FieldMask);
 	const TCHAR* Temp = nullptr;
 	if (0 != (PortFlags & PPF_ExportCpp))
 	{
@@ -360,7 +368,7 @@ void FBoolProperty::ExportTextItem( FString& ValueStr, const void* PropertyValue
 	}
 	ValueStr += FString::Printf( TEXT("%s"), Temp );
 }
-const TCHAR* FBoolProperty::ImportText_Internal( const TCHAR* Buffer, void* Data, int32 PortFlags, UObject* Parent, FOutputDevice* ErrorText ) const
+const TCHAR* FBoolProperty::ImportText_Internal( const TCHAR* Buffer, void* ContainerOrPropertyPtr, EPropertyPointerType PropertyPointerType, UObject* Parent, int32 PortFlags, FOutputDevice* ErrorText ) const
 {
 	FString Temp; 
 	Buffer = FPropertyHelpers::ReadToken( Buffer, Temp );
@@ -370,23 +378,41 @@ const TCHAR* FBoolProperty::ImportText_Internal( const TCHAR* Buffer, void* Data
 	}
 
 	check(FieldSize != 0);
-	uint8* ByteValue = (uint8*)Data + ByteOffset;
+	uint8 LocalByteValue = 0;
+	if (PropertyPointerType == EPropertyPointerType::Container && HasGetter())
+	{
+		GetValue_InContainer(ContainerOrPropertyPtr, &LocalByteValue);
+	}
+	else
+	{
+		LocalByteValue = *((uint8*)PointerToValuePtr(ContainerOrPropertyPtr, PropertyPointerType) + ByteOffset);
+	}
 
 	const FCoreTexts& CoreTexts = FCoreTexts::Get();
 	if( Temp==TEXT("1") || Temp==TEXT("True") || Temp==*CoreTexts.True.ToString() || Temp == TEXT("Yes") || Temp == *CoreTexts.Yes.ToString() )
 	{
-		*ByteValue |= ByteMask;
+		LocalByteValue |= ByteMask;
 	}
 	else 
 	if( Temp==TEXT("0") || Temp==TEXT("False") || Temp==*CoreTexts.False.ToString() || Temp == TEXT("No") || Temp == *CoreTexts.No.ToString() )
 	{
-		*ByteValue &= ~FieldMask;
+		LocalByteValue &= ~FieldMask;
 	}
 	else
 	{
 		//UE_LOG(LogProperty, Log,  "Import: Failed to get bool" );
 		return NULL;
 	}
+
+	if (PropertyPointerType == EPropertyPointerType::Container && HasSetter())
+	{
+		SetValue_InContainer(ContainerOrPropertyPtr, &LocalByteValue);
+	}
+	else
+	{
+		*((uint8*)PointerToValuePtr(ContainerOrPropertyPtr, PropertyPointerType) + ByteOffset) = LocalByteValue;
+	}
+
 	return Buffer;
 }
 bool FBoolProperty::Identical( const void* A, const void* B, uint32 PortFlags ) const
diff --git a/Engine/Source/Runtime/CoreUObject/Private/UObject/PropertyByte.cpp b/Engine/Source/Runtime/CoreUObject/Private/UObject/PropertyByte.cpp
--- a/Engine/Source/Runtime/CoreUObject/Private/UObject/PropertyByte.cpp
+++ b/Engine/Source/Runtime/CoreUObject/Private/UObject/PropertyByte.cpp
@@ -349,13 +349,23 @@ void FByteProperty::AppendSchemaHash(FBlake3& Builder, bool bSkipEditorOnly) con
 }
 #endif
 
-void FByteProperty::ExportTextItem( FString& ValueStr, const void* PropertyValue, const void* DefaultValue, UObject* Parent, int32 PortFlags, UObject* ExportRootScope ) const
+void FByteProperty::ExportText_Internal( FString& ValueStr, const void* PropertyValueOrContainer, EPropertyPointerType PropertyPointerType, const void* DefaultValue, UObject* Parent, int32 PortFlags, UObject* ExportRootScope ) const
 {
 	if (0 != (PortFlags & PPF_ExportCpp))
 	{
 		if (Enum)
 		{
-			const int32 ActualValue = *(const uint8*)PropertyValue;
+			int64 ActualValue = 0;
+			if (PropertyPointerType == EPropertyPointerType::Container && HasGetter())
+			{
+				uint8 Value = 0;
+				GetValue_InContainer(PropertyValueOrContainer, &Value);
+				ActualValue = Value;
+			}
+			else
+			{
+				ActualValue = *(const uint8*)PointerToValuePtr(PropertyValueOrContainer, PropertyPointerType);
+			}
 			const int64 MaxValue = Enum->GetMaxEnumValue();
 			const int64 GoodValue = Enum->IsValidEnumValue(ActualValue) ? ActualValue : MaxValue;
 			const bool bNonNativeEnum = Enum->GetClass() != UEnum::StaticClass();
@@ -375,31 +385,42 @@ void FByteProperty::ExportTextItem( FString& ValueStr, const void* PropertyValue
 		}
 		else
 		{
-			Super::ExportTextItem(ValueStr, PropertyValue, DefaultValue, Parent, PortFlags, ExportRootScope);
+			Super::ExportText_Internal(ValueStr, PropertyValueOrContainer, PropertyPointerType, DefaultValue, Parent, PortFlags, ExportRootScope);
 		}
 		return;
 	}
 
 	if( Enum && (PortFlags & PPF_ConsoleVariable) == 0 )
 	{
+		int64 ActualValue = 0;
+		if (PropertyPointerType == EPropertyPointerType::Container && HasGetter())
+		{
+			uint8 Value = 0;
+			GetValue_InContainer(PropertyValueOrContainer, &Value);
+			ActualValue = Value;
+		}
+		else
+		{
+			ActualValue = *(const uint8*)PointerToValuePtr(PropertyValueOrContainer, PropertyPointerType);
+		}
 		// if the value is the max value (the autogenerated *_MAX value), export as "INVALID", unless we're exporting text for copy/paste (for copy/paste,
 		// the property text value must actually match an entry in the enum's names array)
-		bool bIsValid = Enum->IsValidEnumValue(*(const uint8*)PropertyValue);
-		bool bIsMax = *(const uint8*)PropertyValue == Enum->GetMaxEnumValue();
+		bool bIsValid = Enum->IsValidEnumValue(ActualValue);
+		bool bIsMax = ActualValue == Enum->GetMaxEnumValue();
 		if (bIsValid && (!bIsMax || (PortFlags & PPF_Copy)))
 		{
 			// We do not want to export the enum text for non-display uses, localization text is very dynamic and would cause issues on import
 			if (PortFlags & PPF_PropertyWindow)
 			{
-				ValueStr += Enum->GetDisplayNameTextByValue(*(const uint8*)PropertyValue).ToString();
+				ValueStr += Enum->GetDisplayNameTextByValue(ActualValue).ToString();
 			}
 			else if (PortFlags & PPF_ExternalEditor)
 			{
-				ValueStr += Enum->GetAuthoredNameStringByValue(*(const uint8*)PropertyValue);
+				ValueStr += Enum->GetAuthoredNameStringByValue(ActualValue);
 			}
 			else
 			{
-				ValueStr += Enum->GetNameStringByValue(*(const uint8*)PropertyValue);
+				ValueStr += Enum->GetNameStringByValue(ActualValue);
 			}
 		}
 		else
@@ -409,10 +430,10 @@ void FByteProperty::ExportTextItem( FString& ValueStr, const void* PropertyValue
 	}
 	else
 	{
-		Super::ExportTextItem(ValueStr, PropertyValue, DefaultValue, Parent, PortFlags, ExportRootScope);
+		Super::ExportText_Internal(ValueStr, PropertyValueOrContainer, PropertyPointerType, DefaultValue, Parent, PortFlags, ExportRootScope);
 	}
 }
-const TCHAR* FByteProperty::ImportText_Internal( const TCHAR* InBuffer, void* Data, int32 PortFlags, UObject* Parent, FOutputDevice* ErrorText ) const
+const TCHAR* FByteProperty::ImportText_Internal( const TCHAR* InBuffer, void* ContainerOrPropertyPtr, EPropertyPointerType PropertyPointerType, UObject* Parent, int32 PortFlags, FOutputDevice* ErrorText ) const
 {
 	if( Enum && (PortFlags & PPF_ConsoleVariable) == 0 )
 	{
@@ -428,7 +449,15 @@ const TCHAR* FByteProperty::ImportText_Internal( const TCHAR* InBuffer, void* Da
 			}
 			if (EnumIndex != INDEX_NONE)
 			{
-				*(uint8*)Data = Enum->GetValueByIndex(EnumIndex);
+				int64 EnumValue = Enum->GetValueByIndex(EnumIndex);
+				if (PropertyPointerType == EPropertyPointerType::Container && HasSetter())
+				{
+					SetValue_InContainer(ContainerOrPropertyPtr, EnumValue);
+				}
+				else
+				{
+					*(uint8*)PointerToValuePtr(ContainerOrPropertyPtr, PropertyPointerType) = EnumValue;
+				}
 				return Buffer;
 			}
 
@@ -458,18 +487,34 @@ const TCHAR* FByteProperty::ImportText_Internal( const TCHAR* InBuffer, void* Da
 
 			if (Temp == TEXT("True") || Temp == *(CoreTexts.True.ToString()))
 			{
-				SetIntPropertyValue(Data, 1ull);
+				uint64 TrueValue = 1ull;
+				if (PropertyPointerType == EPropertyPointerType::Container && HasSetter())
+				{
+					SetValue_InContainer(ContainerOrPropertyPtr, TrueValue);
+				}
+				else
+				{
+					SetIntPropertyValue(PointerToValuePtr(ContainerOrPropertyPtr, PropertyPointerType), TrueValue);
+				}
 				return Buffer;
 			}
 			else if (Temp == TEXT("False") || Temp == *(CoreTexts.False.ToString()))
 			{
-				SetIntPropertyValue(Data, 0ull);
+				uint64 FalseValue = 0ull;
+				if (PropertyPointerType == EPropertyPointerType::Container && HasSetter())
+				{
+					SetValue_InContainer(ContainerOrPropertyPtr, FalseValue);
+				}
+				else
+				{
+					SetIntPropertyValue(PointerToValuePtr(ContainerOrPropertyPtr, PropertyPointerType), FalseValue);
+				}
 				return Buffer;
 			}
 		}
 	}
 
-	return Super::ImportText_Internal( InBuffer, Data, PortFlags, Parent, ErrorText );
+	return Super::ImportText_Internal( InBuffer, ContainerOrPropertyPtr, PropertyPointerType, Parent, PortFlags, ErrorText );
 }
 
 UEnum* FByteProperty::GetIntPropertyEnum() const
diff --git a/Engine/Source/Runtime/CoreUObject/Private/UObject/PropertyClass.cpp b/Engine/Source/Runtime/CoreUObject/Private/UObject/PropertyClass.cpp
--- a/Engine/Source/Runtime/CoreUObject/Private/UObject/PropertyClass.cpp
+++ b/Engine/Source/Runtime/CoreUObject/Private/UObject/PropertyClass.cpp
@@ -106,11 +106,12 @@ void FClassProperty::AddReferencedObjects(FReferenceCollector& Collector)
 	Super::AddReferencedObjects( Collector );
 }
 
-const TCHAR* FClassProperty::ImportText_Internal( const TCHAR* Buffer, void* Data, int32 PortFlags, UObject* Parent, FOutputDevice* ErrorText ) const
+const TCHAR* FClassProperty::ImportText_Internal( const TCHAR* Buffer, void* ContainerOrPropertyPtr, EPropertyPointerType PropertyPointerType, UObject* Parent, int32 PortFlags, FOutputDevice* ErrorText ) const
 {
-	const TCHAR* Result = FObjectProperty::ImportText_Internal( Buffer, Data, PortFlags, Parent, ErrorText );
+	const TCHAR* Result = FObjectProperty::ImportText_Internal( Buffer, ContainerOrPropertyPtr, PropertyPointerType, Parent, PortFlags, ErrorText );
 	if( Result )
 	{
+		void* Data = PointerToValuePtr(ContainerOrPropertyPtr, PropertyPointerType);
 		if (UClass* AssignedPropertyClass = dynamic_cast<UClass*>(GetObjectPropertyValue(Data)))
 		{
 #if USE_CIRCULAR_DEPENDENCY_LOAD_DEFERRING
@@ -135,7 +136,15 @@ const TCHAR* FClassProperty::ImportText_Internal( const TCHAR* Buffer, void* Dat
 			{
 				// the object we imported doesn't implement our interface class
 				ErrorText->Logf(TEXT("Invalid object '%s' specified for property '%s'"), *AssignedPropertyClass->GetFullName(), *GetName());
-				SetObjectPropertyValue(Data, NULL);
+				UObject* NullObj = nullptr;
+				if (PropertyPointerType == EPropertyPointerType::Container && HasSetter())
+				{
+					SetValue_InContainer(ContainerOrPropertyPtr, NullObj);
+				}
+				else
+				{
+					SetObjectPropertyValue(PointerToValuePtr(ContainerOrPropertyPtr, PropertyPointerType), NullObj);
+				}
 				Result = NULL;
 			}
 		}
diff --git a/Engine/Source/Runtime/CoreUObject/Private/UObject/PropertyDelegate.cpp b/Engine/Source/Runtime/CoreUObject/Private/UObject/PropertyDelegate.cpp
--- a/Engine/Source/Runtime/CoreUObject/Private/UObject/PropertyDelegate.cpp
+++ b/Engine/Source/Runtime/CoreUObject/Private/UObject/PropertyDelegate.cpp
@@ -143,7 +143,7 @@ FString FDelegateProperty::GetCPPTypeForwardDeclaration() const
 	return FString();
 }
 
-void FDelegateProperty::ExportTextItem( FString& ValueStr, const void* PropertyValue, const void* DefaultValue, UObject* Parent, int32 PortFlags, UObject* ExportRootScope ) const
+void FDelegateProperty::ExportText_Internal( FString& ValueStr, const void* PropertyValueOrContainer, EPropertyPointerType PropertyPointerType, const void* DefaultValue, UObject* Parent, int32 PortFlags, UObject* ExportRootScope ) const
 {
 	if (0 != (PortFlags & PPF_ExportCpp))
 	{
@@ -151,18 +151,45 @@ void FDelegateProperty::ExportTextItem( FString& ValueStr, const void* PropertyV
 		return;
 	}
 
-	FScriptDelegate* ScriptDelegate = (FScriptDelegate*)PropertyValue;
-	check(ScriptDelegate != NULL);
-	bool bDelegateHasValue = ScriptDelegate->GetFunctionName() != NAME_None;
-	ValueStr += FString::Printf( TEXT("%s.%s"),
-		ScriptDelegate->GetUObject() != NULL ? *ScriptDelegate->GetUObject()->GetName() : TEXT("(null)"),
-		*ScriptDelegate->GetFunctionName().ToString() );
+	auto ExportDelegateAsText = [&ValueStr](FScriptDelegate* ScriptDelegate)
+	{
+		check(ScriptDelegate != NULL);
+		bool bDelegateHasValue = ScriptDelegate->GetFunctionName() != NAME_None;
+		ValueStr += FString::Printf(TEXT("%s.%s"),
+			ScriptDelegate->GetUObject() != NULL ? *ScriptDelegate->GetUObject()->GetName() : TEXT("(null)"),
+			*ScriptDelegate->GetFunctionName().ToString());
+	};
+	
+	if (PropertyPointerType == EPropertyPointerType::Container && HasGetter())
+	{
+		FScriptDelegate LocalScriptDelegate;
+		GetValue_InContainer(PropertyValueOrContainer, &LocalScriptDelegate);
+		ExportDelegateAsText(&LocalScriptDelegate);
+	}
+	else
+	{
+		FScriptDelegate* ScriptDelegate = (FScriptDelegate*)PointerToValuePtr(PropertyValueOrContainer, PropertyPointerType);
+		ExportDelegateAsText(ScriptDelegate);
+	}
 }
 
 
-const TCHAR* FDelegateProperty::ImportText_Internal( const TCHAR* Buffer, void* PropertyValue, int32 PortFlags, UObject* Parent, FOutputDevice* ErrorText ) const
+const TCHAR* FDelegateProperty::ImportText_Internal( const TCHAR* Buffer, void* ContainerOrPropertyPtr, EPropertyPointerType PropertyPointerType, UObject* Parent, int32 PortFlags, FOutputDevice* ErrorText ) const
 {
-	return DelegatePropertyTools::ImportDelegateFromText( *(FScriptDelegate*)PropertyValue, SignatureFunction, Buffer, Parent, ErrorText );
+	const TCHAR* Result = nullptr;
+	if (PropertyPointerType == EPropertyPointerType::Container && HasSetter())
+	{
+		FScriptDelegate LocalScriptDelegate;
+		Result = DelegatePropertyTools::ImportDelegateFromText(*(FScriptDelegate*)&LocalScriptDelegate, SignatureFunction, Buffer, Parent, ErrorText);
+		SetValue_InContainer(ContainerOrPropertyPtr, LocalScriptDelegate);
+	}
+	else
+	{
+		void* PropertyValue = PointerToValuePtr(ContainerOrPropertyPtr, PropertyPointerType);
+		Result = DelegatePropertyTools::ImportDelegateFromText(*(FScriptDelegate*)PropertyValue, SignatureFunction, Buffer, Parent, ErrorText);
+	}
+
+	return Result;
 }
 
 void FDelegateProperty::Serialize( FArchive& Ar )
diff --git a/Engine/Source/Runtime/CoreUObject/Private/UObject/PropertyInterface.cpp b/Engine/Source/Runtime/CoreUObject/Private/UObject/PropertyInterface.cpp
--- a/Engine/Source/Runtime/CoreUObject/Private/UObject/PropertyInterface.cpp
+++ b/Engine/Source/Runtime/CoreUObject/Private/UObject/PropertyInterface.cpp
@@ -182,11 +182,20 @@ bool FInterfaceProperty::NetSerializeItem( FArchive& Ar, UPackageMap* Map, void*
 	return false;
 }
 
-void FInterfaceProperty::ExportTextItem( FString& ValueStr, const void* PropertyValue, const void* DefaultValue, UObject* Parent, int32 PortFlags, UObject* ExportRootScope ) const
+void FInterfaceProperty::ExportText_Internal( FString& ValueStr, const void* PropertyValueOrContainer, EPropertyPointerType PropertyPointerType, const void* DefaultValue, UObject* Parent, int32 PortFlags, UObject* ExportRootScope ) const
 {
-	FScriptInterface* InterfaceValue = (FScriptInterface*)PropertyValue;
-
-	UObject* Temp = InterfaceValue->GetObject();
+	UObject* Temp = nullptr;
+	if (PropertyPointerType == EPropertyPointerType::Container && HasGetter())
+	{
+		FScriptInterface InterfaceValue;
+		GetValue_InContainer(PropertyValueOrContainer, &InterfaceValue);
+		Temp = InterfaceValue.GetObject();
+	}
+	else
+	{
+		FScriptInterface* InterfaceValue = (FScriptInterface*)PointerToValuePtr(PropertyValueOrContainer, PropertyPointerType);
+		Temp = InterfaceValue->GetObject();
+	}
 
 	if (0 != (PortFlags & PPF_ExportCpp))
 	{
@@ -230,18 +239,35 @@ void FInterfaceProperty::ExportTextItem( FString& ValueStr, const void* Property
 	}
 }
 
-const TCHAR* FInterfaceProperty::ImportText_Internal( const TCHAR* InBuffer, void* Data, int32 PortFlags, UObject* Parent, FOutputDevice* ErrorText) const
+const TCHAR* FInterfaceProperty::ImportText_Internal( const TCHAR* InBuffer, void* ContainerOrPropertyPtr, EPropertyPointerType PropertyPointerType, UObject* Parent, int32 PortFlags, FOutputDevice* ErrorText) const
 {
-	FScriptInterface* InterfaceValue = (FScriptInterface*)Data;
+	FScriptInterface* InterfaceValue = (FScriptInterface*)PointerToValuePtr(ContainerOrPropertyPtr, PropertyPointerType);
 	UObject* ResolvedObject = InterfaceValue->GetObject();
-	void* InterfaceAddress = InterfaceValue->GetInterface();
+
+	auto SetInterfaceValue = [InterfaceValue, ContainerOrPropertyPtr, PropertyPointerType, this](UObject* NewObject, void* NewInterfaceAddress)
+	{
+		if (PropertyPointerType == EPropertyPointerType::Container && HasSetter())
+		{
+			FScriptInterface LocalInterfaceValue;
+			LocalInterfaceValue.SetObject(NewObject);
+			LocalInterfaceValue.SetInterface(NewInterfaceAddress);
+			SetValue_InContainer(ContainerOrPropertyPtr, LocalInterfaceValue);
+		}
+		else
+		{
+			InterfaceValue->SetObject(NewObject);
+			if (NewObject)
+			{
+				// if NewObject was set to nullptr, SetObject will take care of clearing the interface address too
+				InterfaceValue->SetInterface(NewInterfaceAddress);
+			}
+		}
+	};
 
 	const TCHAR* Buffer = InBuffer;
 	if (!FObjectPropertyBase::ParseObjectPropertyValue(this, Parent, UObject::StaticClass(), PortFlags, Buffer, ResolvedObject))
 	{
-		// we only need to call SetObject here - if ObjectAddress was not modified, then InterfaceValue should not be modified either
-		// if it was set to nullptr, SetObject will take care of clearing the interface address too
-		InterfaceValue->SetObject(ResolvedObject);
+		SetInterfaceValue(ResolvedObject, InterfaceValue->GetInterface());
 		return nullptr;
 	}
 
@@ -250,7 +276,7 @@ const TCHAR* FInterfaceProperty::ImportText_Internal( const TCHAR* InBuffer, voi
 	{
 		// if ParseObjectPropertyValue returned true but ResolvedObject is nullptr, the imported text was "None".  Make sure the interface pointer
 		// is cleared, then stop
-		InterfaceValue->SetObject(nullptr);
+		SetInterfaceValue(nullptr, nullptr);
 		return Buffer;
 	}
 
@@ -260,8 +286,7 @@ const TCHAR* FInterfaceProperty::ImportText_Internal( const TCHAR* InBuffer, voi
 		// If this is a bp implementation of a native interface, set object but clear the interface
 		if (ResolvedObject->GetClass()->ImplementsInterface(InterfaceClass))
 		{
-			InterfaceValue->SetObject(ResolvedObject);
-			InterfaceValue->SetInterface(nullptr);
+			SetInterfaceValue(ResolvedObject, nullptr);
 			return Buffer;
 		}
 
@@ -272,8 +297,7 @@ const TCHAR* FInterfaceProperty::ImportText_Internal( const TCHAR* InBuffer, voi
 		return nullptr;
 	}
 
-	InterfaceValue->SetObject(ResolvedObject);
-	InterfaceValue->SetInterface(NewInterfaceAddress);
+	SetInterfaceValue(ResolvedObject, NewInterfaceAddress);
 	return Buffer;
 }
 
