It's no longer required to provide an explicit setter or getter name if Setter or Getter tag are present in UPROPERTY macro. If no accessor function name is specified it will be autogenerated.

diff --git a/Engine/Source/Programs/UnrealHeaderTool/Private/HeaderParser.cpp b/Engine/Source/Programs/UnrealHeaderTool/Private/HeaderParser.cpp
--- a/Engine/Source/Programs/UnrealHeaderTool/Private/HeaderParser.cpp
+++ b/Engine/Source/Programs/UnrealHeaderTool/Private/HeaderParser.cpp
@@ -2841,6 +2841,8 @@ void FHeaderParser::GetVarType(
 	FName RepCallbackName = FName(NAME_None);
 	FString SetterName;
 	FString GetterName;
+	bool bHasSetterTag = false;
+	bool bHasGetterTag = false;
 
 	// Get flags.
 	EPropertyFlags Flags        = CPF_None;
@@ -3267,13 +3269,30 @@ void FHeaderParser::GetVarType(
 
 				case EVariableSpecifier::Setter:
 				{
-					SetterName = RequireExactlyOneSpecifierValue(*this, Specifier);
+					bHasSetterTag = true;
+					if (Specifier.Values.Num() == 1)
+					{
+						SetterName = Specifier.Values[0];
+					}
+					else if (Specifier.Values.Num() > 1)
+					{
+						Throwf(TEXT("The specifier '%s' must be given exactly one value"), *Specifier.Key);
+					}
+					
 				}
 				break;
 
 				case EVariableSpecifier::Getter:
 				{
-					GetterName = RequireExactlyOneSpecifierValue(*this, Specifier);
+					bHasGetterTag = true;
+					if (Specifier.Values.Num() == 1)
+					{
+						GetterName = Specifier.Values[0];
+					}
+					else if (Specifier.Values.Num() > 1)
+					{
+						Throwf(TEXT("The specifier '%s' must be given exactly one value"), *Specifier.Key);
+					}
 				}
 				break;
 
@@ -4366,7 +4385,9 @@ void FHeaderParser::GetVarType(
 		}
 	}
 
+	VarProperty.bSetterTagFound = bHasSetterTag;
 	VarProperty.SetterName = SetterName;
+	VarProperty.bGetterTagFound = bHasGetterTag;
 	VarProperty.GetterName = GetterName;
 
 	// Perform some more specific validation on the property flags
@@ -5357,12 +5378,13 @@ bool FHeaderParser::CheckForPropertySetterFunction(FUnrealStructDefinitionInfo&
 	}
 
 	// If not, try with autogenerated setter names
-	if (!PropertyWithSetter && FUHTConfig::Get().bAllowAutomaticSettersAndGetters && SetterNameToken.ValueStartsWith(TEXT("Set"), ESearchCase::CaseSensitive))
+	if (!PropertyWithSetter && SetterNameToken.ValueStartsWith(TEXT("Set"), ESearchCase::CaseSensitive))
 	{
 		for (TSharedRef<FUnrealPropertyDefinitionInfo>& Prop : Properties)
 		{
 			FPropertyBase& BaseProp = Prop->GetPropertyBase();
-			if (!BaseProp.bSetterFunctionFound && BaseProp.SetterName.IsEmpty())
+			if (!BaseProp.bSetterFunctionFound && BaseProp.SetterName.IsEmpty() &&
+				(BaseProp.bSetterTagFound || FUHTConfig::Get().bAllowAutomaticSettersAndGetters))
 			{
 				SetterFunctionName.Reset();
 				SetterFunctionName += TEXT("Set");
@@ -5371,6 +5393,7 @@ bool FHeaderParser::CheckForPropertySetterFunction(FUnrealStructDefinitionInfo&
 				if (SetterNameToken.Value.Compare(SetterFunctionName, ESearchCase::CaseSensitive) == 0)
 				{
 					// We found a matching function name so we can move onto verifying if the declaration matches
+					bExplicitSetter = BaseProp.bSetterTagFound;
 					PropertyWithSetter = &Prop.Get();
 					break;
 				}
@@ -5521,18 +5544,20 @@ bool FHeaderParser::CheckForPropertyGetterFunction(FUnrealStructDefinitionInfo&
 	}
 
 	// If not, try with autogenerated getter names
-	if (!PropertyWithGetter && FUHTConfig::Get().bAllowAutomaticSettersAndGetters && GetterNameToken.ValueStartsWith(TEXT("Get"), ESearchCase::CaseSensitive))
+	if (!PropertyWithGetter && GetterNameToken.ValueStartsWith(TEXT("Get"), ESearchCase::CaseSensitive))
 	{
 		for (TSharedRef<FUnrealPropertyDefinitionInfo>& Prop : Properties)
 		{
 			FPropertyBase& BaseProp = Prop->GetPropertyBase();			
-			if (!BaseProp.bGetterFunctionFound && BaseProp.GetterName != TEXT("None"))
+			if (!BaseProp.bGetterFunctionFound && BaseProp.GetterName != TEXT("None") && 
+				(BaseProp.bGetterTagFound || FUHTConfig::Get().bAllowAutomaticSettersAndGetters))
 			{
 				GetterFunctionName.Reset();
 				GetterFunctionName += TEXT("Get");
 				GetterFunctionName += Prop->GetName();
 				if (GetterNameToken.Value.Compare(GetterFunctionName, ESearchCase::CaseSensitive) == 0)
 				{
+					bExplicitGetter = BaseProp.bGetterTagFound;
 					PropertyWithGetter = &Prop.Get();
 					break;
 				}
diff --git a/Engine/Source/Programs/UnrealHeaderTool/Private/ParserHelper.h b/Engine/Source/Programs/UnrealHeaderTool/Private/ParserHelper.h
--- a/Engine/Source/Programs/UnrealHeaderTool/Private/ParserHelper.h
+++ b/Engine/Source/Programs/UnrealHeaderTool/Private/ParserHelper.h
@@ -212,6 +212,10 @@ public:
 	FString SetterName;	
 	/** Getter function name. Either specified in the UPROPERTY macro or autogenerated while parsing class declaration */
 	FString GetterName;
+	/** True if the property was marked as having a setter */
+	bool bSetterTagFound = false;
+	/** True if the property was marked as having a gatter */
+	bool bGetterTagFound = false;
 	/** True if setter function declaration was found while parsing class header */
 	bool bSetterFunctionFound = false;
 	/** True if getter function declaration was found while parsing class header */
