Add internal use command line options to UHT.  One to enable debugger code generation outout  and one to disable the concurrent processing.

#rnx

diff --git a/Engine/Source/Programs/UnrealHeaderTool/Private/CodeGenerator.cpp b/Engine/Source/Programs/UnrealHeaderTool/Private/CodeGenerator.cpp
--- a/Engine/Source/Programs/UnrealHeaderTool/Private/CodeGenerator.cpp
+++ b/Engine/Source/Programs/UnrealHeaderTool/Private/CodeGenerator.cpp
@@ -76,6 +76,8 @@ double GMacroizeTime = 0.0;
 static TArray<FString> ChangeMessages;
 static bool bWriteContents = false;
 static bool bVerifyContents = false;
+static bool bIncludeDebugOutput = false;
+bool bGoWide = true;
 
 void ProcessParsedClass(FUnrealClassDefinitionInfo& ClassDef);
 void ProcessParsedEnum(FUnrealEnumDefinitionInfo& EnumDef);
@@ -90,7 +92,6 @@ FUnrealClassDefinitionInfo* GUClassDef = nullptr;
 FUnrealClassDefinitionInfo* GUInterfaceDef = nullptr;
 
 // Busy wait support for holes in the include graph issues
-std::atomic<bool> GSourcesConcurrent = false;
 std::atomic<int> GSourcesToParse = 0;
 std::atomic<int> GSourcesParsing = 0;
 std::atomic<int> GSourcesCompleted = 0;
@@ -1922,11 +1923,12 @@ void FNativeClassHeaderGenerator::ExportGeneratedPackageInitCode(FOutputDevice&
 		return AName < BName;
 	});
 
-#if UHT_ENABLE_EXTRA_HASH_OUTPUT
-	Out.Log(TEXT("#if 0\r\n"));
-	Out.Log(InDeclarations);
-	Out.Log(TEXT("#endif\r\n"));
-#endif
+	if (bIncludeDebugOutput)
+	{
+		Out.Log(TEXT("#if 0\r\n"));
+		Out.Log(InDeclarations);
+		Out.Log(TEXT("#endif\r\n"));
+	}
 
 	for (FUnrealFieldDefinitionInfo* FieldDef : Singletons)
 	{
@@ -2299,11 +2301,12 @@ void FNativeClassHeaderGenerator::ExportNativeGeneratedInitCode(FOutputDevice& O
 		}
 	}
 
-#if UHT_ENABLE_EXTRA_HASH_OUTPUT
-	Out.Log(TEXT("#if 0\r\n"));
-	Out.Log(HashBuilder);
-	Out.Log(TEXT("#endif\r\n"));
-#endif
+	if (bIncludeDebugOutput)
+	{
+		Out.Log(TEXT("#if 0\r\n"));
+		Out.Log(HashBuilder);
+		Out.Log(TEXT("#endif\r\n"));
+	}
 
 	GeneratedClassRegisterFunctionText.Log(HashBuilder);
 
@@ -5278,13 +5281,16 @@ bool FNativeClassHeaderGenerator::LoadSourceFile(FGeneratedCPP& GeneratedCPP)
 	FString HeaderPath = (Module.GeneratedIncludeDirectory / StrippedName) + TEXT(".generated.h");
 	FString GeneratedSourceFilename = (Module.GeneratedIncludeDirectory / StrippedName) + TEXT(".gen.cpp");
 
-#if UHT_ENABLE_CONCURRENT_CODE_GENERATION
-	GeneratedCPP.Header.StartLoad(MoveTemp(HeaderPath));
-	GeneratedCPP.Source.StartLoad(MoveTemp(GeneratedSourceFilename));
-#else
-	GeneratedCPP.Header.Load(MoveTemp(HeaderPath));
-	GeneratedCPP.Source.Load(MoveTemp(GeneratedSourceFilename));
-#endif
+	if (bGoWide)
+	{
+		GeneratedCPP.Header.StartLoad(MoveTemp(HeaderPath));
+		GeneratedCPP.Source.StartLoad(MoveTemp(GeneratedSourceFilename));
+	}
+	else
+	{
+		GeneratedCPP.Header.Load(MoveTemp(HeaderPath));
+		GeneratedCPP.Source.Load(MoveTemp(GeneratedSourceFilename));
+	}
 	return true;
 }
 
@@ -5557,11 +5563,12 @@ void FNativeClassHeaderGenerator::GenerateSourceFile(FGeneratedCPP& GeneratedCPP
 			{
 				GeneratedFunctionDeclarations.Logf(TEXT("%s\r\n"), *Line);
 			}
-#if UHT_ENABLE_EXTRA_HASH_OUTPUT
-			GeneratedCPPText.Logf(TEXT("#if 0\r\n"));
-			GeneratedCPPText.Log(GeneratedFunctionDeclarations);
-			GeneratedCPPText.Logf(TEXT("#endif\r\n"));
-#endif
+			if (bIncludeDebugOutput)
+			{
+				GeneratedCPPText.Logf(TEXT("#if 0\r\n"));
+				GeneratedCPPText.Log(GeneratedFunctionDeclarations);
+				GeneratedCPPText.Logf(TEXT("#endif\r\n"));
+			}
 		}
 	);
 }
@@ -5598,83 +5605,85 @@ void FNativeClassHeaderGenerator::GenerateSourceFiles(
 	TArray<FGeneratedCPP>& GeneratedCPPs
 )
 {
-#if UHT_ENABLE_CONCURRENT_CODE_GENERATION
-	TSet<FUnrealSourceFile*> Includes;
-	Includes.Reserve(GeneratedCPPs.Num());
-	FGraphEventArray TempTasks;
-	TempTasks.Reserve(3);
-	for (FGeneratedCPP& GeneratedCPP : GeneratedCPPs)
+	if (bGoWide)
 	{
-		if (!LoadSourceFile(GeneratedCPP))
+		TSet<FUnrealSourceFile*> Includes;
+		Includes.Reserve(GeneratedCPPs.Num());
+		FGraphEventArray TempTasks;
+		TempTasks.Reserve(3);
+		for (FGeneratedCPP& GeneratedCPP : GeneratedCPPs)
 		{
-			continue;
-		}
+			if (!LoadSourceFile(GeneratedCPP))
+			{
+				continue;
+			}
 
-		TempTasks.Reset();
-		GeneratedCPP.Header.AddLoadTaskRef(TempTasks);
-		GeneratedCPP.Source.AddLoadTaskRef(TempTasks);
+			TempTasks.Reset();
+			GeneratedCPP.Header.AddLoadTaskRef(TempTasks);
+			GeneratedCPP.Source.AddLoadTaskRef(TempTasks);
 
-		auto GenerateSource = [&GeneratedCPP]()
-		{
-			GenerateSourceFile(GeneratedCPP);
-		};
+			auto GenerateSource = [&GeneratedCPP]()
+			{
+				GenerateSourceFile(GeneratedCPP);
+			};
 
-		auto WriteGenerated = [&GeneratedCPP]()
-		{
-			WriteSourceFile(GeneratedCPP);
-		};
+			auto WriteGenerated = [&GeneratedCPP]()
+			{
+				WriteSourceFile(GeneratedCPP);
+			};
 
-		Includes.Reset();
-		for (FHeaderProvider& Header : GeneratedCPP.SourceFile.GetIncludes())
-		{
-			if (FUnrealSourceFile* Include = Header.Resolve(GeneratedCPP.SourceFile))
+			Includes.Reset();
+			for (FHeaderProvider& Header : GeneratedCPP.SourceFile.GetIncludes())
+			{
+				if (FUnrealSourceFile* Include = Header.Resolve(GeneratedCPP.SourceFile))
+				{
+					Includes.Add(Include);
+				}
+			}
+
+			// Our generation must wait on all of our includes generation to complete
+			TempTasks.Reset();
+			for (FUnrealSourceFile* Include : Includes)
 			{
-				Includes.Add(Include);
+				FGeneratedCPP& IncludeCPP = GeneratedCPPs[Include->GetOrderedIndex()];
+				IncludeCPP.AddGenerateTaskRef(TempTasks);
 			}
+			GeneratedCPP.GenerateTaskRef = FFunctionGraphTask::CreateAndDispatchWhenReady(MoveTemp(GenerateSource), TStatId(), &TempTasks);
+
+			// Our compare and save must wait on generation and loading of the current header and source
+			TempTasks.Reset();
+			TempTasks.Add(GeneratedCPP.GenerateTaskRef);
+			GeneratedCPP.Header.AddLoadTaskRef(TempTasks);
+			GeneratedCPP.Source.AddLoadTaskRef(TempTasks);
+			GeneratedCPP.ExportTaskRef = FFunctionGraphTask::CreateAndDispatchWhenReady(MoveTemp(WriteGenerated), TStatId(), &TempTasks);
 		}
 
-		// Our generation must wait on all of our includes generation to complete
-		TempTasks.Reset();
-		for (FUnrealSourceFile* Include : Includes)
+		// When this task fires, all of the exports for this package have completed
+		FGraphEventArray ExportSourceTasks;
+		ExportSourceTasks.Reserve(GeneratedCPPs.Num());
+		for (FGeneratedCPP& GeneratedCPP : GeneratedCPPs)
 		{
-			FGeneratedCPP& IncludeCPP = GeneratedCPPs[Include->GetOrderedIndex()];
-			IncludeCPP.AddGenerateTaskRef(TempTasks);
+			GeneratedCPP.AddExportTaskRef(ExportSourceTasks);
 		}
-		GeneratedCPP.GenerateTaskRef = FFunctionGraphTask::CreateAndDispatchWhenReady(MoveTemp(GenerateSource), TStatId(), &TempTasks);
 
-		// Our compare and save must wait on generation and loading of the current header and source
-		TempTasks.Reset();
-		TempTasks.Add(GeneratedCPP.GenerateTaskRef);
-		GeneratedCPP.Header.AddLoadTaskRef(TempTasks);
-		GeneratedCPP.Source.AddLoadTaskRef(TempTasks);
-		GeneratedCPP.ExportTaskRef = FFunctionGraphTask::CreateAndDispatchWhenReady(MoveTemp(WriteGenerated), TStatId(), &TempTasks);
-	}
+		// This is strange, but flushing the log can take a long time.  Without an explicit flush, we were getting a long stall in the UE_LOG message
+		// during the detection of script plugins.  By doing it here, the flush should easily complete during code generation.
+		GLog->FlushThreadedLogs();
 
-	// When this task fires, all of the exports for this package have completed
-	FGraphEventArray ExportSourceTasks;
-	ExportSourceTasks.Reserve(GeneratedCPPs.Num());
-	for (FGeneratedCPP& GeneratedCPP : GeneratedCPPs)
-	{
-		GeneratedCPP.AddExportTaskRef(ExportSourceTasks);
+		FTaskGraphInterface::Get().WaitUntilTasksComplete(ExportSourceTasks);
 	}
-
-	// This is strange, but flushing the log can take a long time.  Without an explicit flush, we were getting a long stall in the UE_LOG message
-	// during the detection of script plugins.  By doing it here, the flush should easily complete during code generation.
-	GLog->FlushThreadedLogs();
-
-	FTaskGraphInterface::Get().WaitUntilTasksComplete(ExportSourceTasks);
-#else
-	for (FGeneratedCPP& GeneratedCPP : GeneratedCPPs)
+	else
 	{
-		if (!LoadSourceFile(GeneratedCPP))
+		for (FGeneratedCPP& GeneratedCPP : GeneratedCPPs)
 		{
-			continue;
+			if (!LoadSourceFile(GeneratedCPP))
+			{
+				continue;
+			}
+			GenerateSourceFile(GeneratedCPP);
+			WriteSourceFile(GeneratedCPP);
 		}
-
-		GenerateSourceFile(GeneratedCPP);
-		WriteSourceFile(GeneratedCPP);
 	}
-#endif
 	FResults::WaitForErrorTasks();
 }
 
@@ -6456,44 +6465,48 @@ void PreparseSource(FUnrealSourceFile& SourceFile)
 
 void PreparseSources(TArray<FUnrealPackageDefinitionInfo*>& PackageDefs, const FString& ModuleInfoPath)
 {
-#if UHT_ENABLE_CONCURRENT_PREPARSING
-	FGraphEventArray LoadTasks;
-	LoadTasks.Reserve(1024); // Fairly arbitrary number
-	for (FUnrealPackageDefinitionInfo* PackageDef : PackageDefs)
+	if (bGoWide)
 	{
-		for (TSharedRef<FUnrealSourceFile>& SourceFile : PackageDef->GetAllSourceFiles())
+		FGraphEventArray LoadTasks;
+		LoadTasks.Reserve(1024); // Fairly arbitrary number
+		for (FUnrealPackageDefinitionInfo* PackageDef : PackageDefs)
 		{
-
-			// Phase #1: Load the file
-			auto LoadLambda = [&SourceFile = *SourceFile, &ModuleInfoPath]()
+			for (TSharedRef<FUnrealSourceFile>& SourceFile : PackageDef->GetAllSourceFiles())
 			{
-				LoadSource(SourceFile, ModuleInfoPath);
-			};
 
-			// Phase #2: Perform simplified class parse (can run concurrenrtly)
-			auto PreProcessLambda = [&SourceFile = *SourceFile]()
-			{
-				PreparseSource(SourceFile);
-			};
+				// Phase #1: Load the file
+				auto LoadLambda = [&SourceFile = *SourceFile, &ModuleInfoPath]()
+				{
+					LoadSource(SourceFile, ModuleInfoPath);
+				};
 
-			FGraphEventRef LoadTask = FFunctionGraphTask::CreateAndDispatchWhenReady(MoveTemp(LoadLambda), TStatId());
-			FGraphEventRef PreProcessTask = FFunctionGraphTask::CreateAndDispatchWhenReady(MoveTemp(PreProcessLambda), TStatId(), LoadTask);
-			LoadTasks.Add(MoveTemp(PreProcessTask));
+				// Phase #2: Perform simplified class parse (can run concurrenrtly)
+				auto PreProcessLambda = [&SourceFile = *SourceFile]()
+				{
+					PreparseSource(SourceFile);
+				};
+
+				FGraphEventRef LoadTask = FFunctionGraphTask::CreateAndDispatchWhenReady(MoveTemp(LoadLambda), TStatId());
+				FGraphEventRef PreProcessTask = FFunctionGraphTask::CreateAndDispatchWhenReady(MoveTemp(PreProcessLambda), TStatId(), LoadTask);
+				LoadTasks.Add(MoveTemp(PreProcessTask));
+			}
 		}
-	}
 
-	// Wait for all the loading and preparsing to complete
-	FTaskGraphInterface::Get().WaitUntilTasksComplete(LoadTasks);
-#else
-	for (FUnrealPackageDefinitionInfo* PackageDef : PackageDefs)
+		// Wait for all the loading and preparsing to complete
+		FTaskGraphInterface::Get().WaitUntilTasksComplete(LoadTasks);
+	}
+	else
 	{
-		for (TSharedRef<FUnrealSourceFile>& SourceFile : PackageDef->GetAllSourceFiles())
+		for (FUnrealPackageDefinitionInfo* PackageDef : PackageDefs)
 		{
-			LoadSource(*SourceFile, ModuleInfoPath);
-			PreparseSource(*SourceFile);
+			for (TSharedRef<FUnrealSourceFile>& SourceFile : PackageDef->GetAllSourceFiles())
+			{
+				LoadSource(*SourceFile, ModuleInfoPath);
+				PreparseSource(*SourceFile);
+			}
 		}
 	}
-#endif
+
 	FResults::WaitForErrorTasks();
 }
 
@@ -6650,100 +6663,101 @@ void ParseSourceFiles(TArray<FUnrealSourceFile*>& OrderedSourceFiles)
 	TGuardValue<bool> AutoRestoreVerifyObjectRefsFlag(GVerifyObjectReferencesOnly, true);
 
 	GSourcesToParse = (int)OrderedSourceFiles.Num();
-	GSourcesConcurrent = UHT_ENABLE_CONCURRENT_PARSING != 0;
-#if UHT_ENABLE_CONCURRENT_PARSING
-
-	/**
-	 * For every FUnrealSourceFile being processed, an instance of this class represents the data associated with generating the new output.
-	 */
-	struct FParseCPP
+	if (bGoWide)
 	{
-		FParseCPP(FUnrealPackageDefinitionInfo& InPackageDef, FUnrealSourceFile& InSourceFile)
-			: PackageDef(InPackageDef)
-			, SourceFile(InSourceFile)
-		{}
-
-		/**
-		 * The package definition being exported
-		 */
-		FUnrealPackageDefinitionInfo& PackageDef;
-
-		/**
-		 * The source file being exported
-		 */
-		FUnrealSourceFile& SourceFile;
-
 		/**
-		 * This task represents the task that parses the source
+		 * For every FUnrealSourceFile being processed, an instance of this class represents the data associated with generating the new output.
 		 */
-		FGraphEventRef ParseTaskRef;
-	};
+		struct FParseCPP
+		{
+			FParseCPP(FUnrealPackageDefinitionInfo& InPackageDef, FUnrealSourceFile& InSourceFile)
+				: PackageDef(InPackageDef)
+				, SourceFile(InSourceFile)
+			{}
+
+			/**
+			 * The package definition being exported
+			 */
+			FUnrealPackageDefinitionInfo& PackageDef;
+
+			/**
+			 * The source file being exported
+			 */
+			FUnrealSourceFile& SourceFile;
+
+			/**
+			 * This task represents the task that parses the source
+			 */
+			FGraphEventRef ParseTaskRef;
+		};
 
-	TArray<FParseCPP> ParsedCPPs;
-	ParsedCPPs.Reserve(OrderedSourceFiles.Num());
-	for (FUnrealSourceFile* SourceFile : OrderedSourceFiles)
-	{
-		ParsedCPPs.Emplace(SourceFile->GetPackageDef(), *SourceFile);
-	}
+		TArray<FParseCPP> ParsedCPPs;
+		ParsedCPPs.Reserve(OrderedSourceFiles.Num());
+		for (FUnrealSourceFile* SourceFile : OrderedSourceFiles)
+		{
+			ParsedCPPs.Emplace(SourceFile->GetPackageDef(), *SourceFile);
+		}
 
-	TSet<FUnrealSourceFile*> Includes;
-	Includes.Reserve(ParsedCPPs.Num());
-	FGraphEventArray TempTasks;
-	TempTasks.Reserve(ParsedCPPs.Num());
-	FGraphEventArray ParsedSourceTasks;
-	ParsedSourceTasks.Reserve(ParsedCPPs.Num());
-	for (FParseCPP& ParsedCPP : ParsedCPPs)
-	{
-		const FManifestModule& Module = ParsedCPP.PackageDef.GetModule();
-		FUnrealSourceFile& SourceFile = ParsedCPP.SourceFile;
+		TSet<FUnrealSourceFile*> Includes;
+		Includes.Reserve(ParsedCPPs.Num());
+		FGraphEventArray TempTasks;
+		TempTasks.Reserve(ParsedCPPs.Num());
+		FGraphEventArray ParsedSourceTasks;
+		ParsedSourceTasks.Reserve(ParsedCPPs.Num());
+		for (FParseCPP& ParsedCPP : ParsedCPPs)
+		{
+			const FManifestModule& Module = ParsedCPP.PackageDef.GetModule();
+			FUnrealSourceFile& SourceFile = ParsedCPP.SourceFile;
 
-		FString ModuleRelativeFilename = SourceFile.GetFilename();
-		ConvertToBuildIncludePath(Module, ModuleRelativeFilename);
+			FString ModuleRelativeFilename = SourceFile.GetFilename();
+			ConvertToBuildIncludePath(Module, ModuleRelativeFilename);
 
-		auto ParseSource = [&ParsedCPP]()
-		{
-			++GSourcesParsing;
-			FResults::TryAlways([&ParsedCPP]()
+			auto ParseSource = [&ParsedCPP]()
 			{
-				FHeaderParser::Parse(ParsedCPP.PackageDef, ParsedCPP.SourceFile);
-			});
-			++GSourcesCompleted;
-		};
+				++GSourcesParsing;
+				FResults::TryAlways([&ParsedCPP]()
+					{
+						FHeaderParser::Parse(ParsedCPP.PackageDef, ParsedCPP.SourceFile);
+					});
+				++GSourcesCompleted;
+			};
 
-		Includes.Reset();
-		for (FHeaderProvider& Header : SourceFile.GetIncludes())
-		{
-			if (FUnrealSourceFile* Include = Header.Resolve(SourceFile))
+			Includes.Reset();
+			for (FHeaderProvider& Header : SourceFile.GetIncludes())
 			{
-				Includes.Add(Include);
+				if (FUnrealSourceFile* Include = Header.Resolve(SourceFile))
+				{
+					Includes.Add(Include);
+				}
 			}
-		}
 
-		// Our generation must wait on all of our includes generation to complete
-		TempTasks.Reset();
-		for (FUnrealSourceFile* Include : Includes)
-		{
-			FParseCPP& IncludeCPP = ParsedCPPs[Include->GetOrderedIndex()];
-			TempTasks.Add(IncludeCPP.ParseTaskRef);
+			// Our generation must wait on all of our includes generation to complete
+			TempTasks.Reset();
+			for (FUnrealSourceFile* Include : Includes)
+			{
+				FParseCPP& IncludeCPP = ParsedCPPs[Include->GetOrderedIndex()];
+				TempTasks.Add(IncludeCPP.ParseTaskRef);
 
+			}
+			ParsedCPP.ParseTaskRef = FFunctionGraphTask::CreateAndDispatchWhenReady(MoveTemp(ParseSource), TStatId(), &TempTasks);
+			ParsedSourceTasks.Add(ParsedCPP.ParseTaskRef);
 		}
-		ParsedCPP.ParseTaskRef = FFunctionGraphTask::CreateAndDispatchWhenReady(MoveTemp(ParseSource), TStatId(), &TempTasks);
-		ParsedSourceTasks.Add(ParsedCPP.ParseTaskRef);
-	}
 
-	// Wait for the results
-	FTaskGraphInterface::Get().WaitUntilTasksComplete(ParsedSourceTasks);
-	FResults::WaitForErrorTasks();
-#else
-	for (FUnrealSourceFile* SourceFile : OrderedSourceFiles)
+		// Wait for the results
+		FTaskGraphInterface::Get().WaitUntilTasksComplete(ParsedSourceTasks);
+		FResults::WaitForErrorTasks();
+	}
+	else
 	{
-		FUnrealPackageDefinitionInfo& PackageDef = SourceFile->GetPackageDef();
-		FScopedDurationTimer SourceTimer(SourceFile->GetTime(ESourceFileTime::Parse));
-		++GSourcesParsing;
-		FResults::TryAlways([&PackageDef, SourceFile]() { FHeaderParser::Parse(PackageDef, *SourceFile); });
-		++GSourcesCompleted;
+		for (FUnrealSourceFile* SourceFile : OrderedSourceFiles)
+		{
+			FUnrealPackageDefinitionInfo& PackageDef = SourceFile->GetPackageDef();
+			FScopedDurationTimer SourceTimer(SourceFile->GetTime(ESourceFileTime::Parse));
+			++GSourcesParsing;
+			FResults::TryAlways([&PackageDef, SourceFile]() { FHeaderParser::Parse(PackageDef, *SourceFile); });
+			++GSourcesCompleted;
+		}
 	}
-#endif
 }
 
 void PostParseFinalize(TArray<FUnrealPackageDefinitionInfo*>& PackageDefs)
@@ -6930,6 +6944,9 @@ ECompilationResult::Type UnrealHeaderTool_Main(const FString& ModuleInfoFilename
 	FDurationTimer MainTimer(MainTime);
 	MainTimer.Start();
 
+	bIncludeDebugOutput = FParse::Param(FCommandLine::Get(), TEXT("IncludeDebugOutput"));
+	bGoWide = !FParse::Param(FCommandLine::Get(), TEXT("NoGoWide"));
+
 	check(GIsUCCMakeStandaloneHeaderGenerator);
 
 	FString ModuleInfoPath = FPaths::GetPath(ModuleInfoFilename);
diff --git a/Engine/Source/Programs/UnrealHeaderTool/Private/HeaderParser.cpp b/Engine/Source/Programs/UnrealHeaderTool/Private/HeaderParser.cpp
--- a/Engine/Source/Programs/UnrealHeaderTool/Private/HeaderParser.cpp
+++ b/Engine/Source/Programs/UnrealHeaderTool/Private/HeaderParser.cpp
@@ -87,6 +87,7 @@ TArray<FString> FHeaderParser::PropertyCPPTypesRequiringUIRanges = { TEXT("float
 TArray<FString> FHeaderParser::ReservedTypeNames = { TEXT("none") };
 
 // Busy wait support for holes in the include graph issues
+extern bool bGoWide;
 extern std::atomic<bool> GSourcesConcurrent;
 extern std::atomic<int> GSourcesToParse;
 extern std::atomic<int> GSourcesParsing;
@@ -708,7 +709,7 @@ namespace
 		}
 
 		// If we aren't doing concurrent processing, then there isn't any reason to continue, we will never make any progress.
-		if (!GSourcesConcurrent)
+		if (!bGoWide)
 		{
 			return false;
 		}
diff --git a/Engine/Source/Programs/UnrealHeaderTool/Private/UnrealHeaderTool.h b/Engine/Source/Programs/UnrealHeaderTool/Private/UnrealHeaderTool.h
--- a/Engine/Source/Programs/UnrealHeaderTool/Private/UnrealHeaderTool.h
+++ b/Engine/Source/Programs/UnrealHeaderTool/Private/UnrealHeaderTool.h
@@ -13,18 +13,6 @@
 // we might never need this tracking since it is always by pointer.
 #define UHT_ENABLE_PTR_PROPERTY_TAG 0
 
-// Enable concurrent pre-parsing of header files
-#define UHT_ENABLE_CONCURRENT_PREPARSING 1
-
-// Enable concurrent parsing of header files
-#define UHT_ENABLE_CONCURRENT_PARSING 1
-
-// Enable concurrent code generation
-#define UHT_ENABLE_CONCURRENT_CODE_GENERATION 1
-
-// Enable the output of extra code gen that contributes to hashes
-#define UHT_ENABLE_EXTRA_HASH_OUTPUT 0
-
 struct FHeaderParserNames
 {
 	static const FName NAME_IsConversionRoot;
