Functions will now contain arguments and return values when used in conjunction with the script exporter

diff --git a/Engine/Source/Programs/UnrealHeaderTool/Private/UnrealTypeDefinitionInfo.cpp b/Engine/Source/Programs/UnrealHeaderTool/Private/UnrealTypeDefinitionInfo.cpp
--- a/Engine/Source/Programs/UnrealHeaderTool/Private/UnrealTypeDefinitionInfo.cpp
+++ b/Engine/Source/Programs/UnrealHeaderTool/Private/UnrealTypeDefinitionInfo.cpp
@@ -1763,6 +1763,11 @@ void FUnrealStructDefinitionInfo::CreateUObjectEngineTypesInternal(ECreateEngine
 			PropertyDef->CreateUObjectEngineTypes(Phase);
 		}
 
+		for (TSharedRef<FUnrealFunctionDefinitionInfo>& FunctionDef : Functions)
+		{
+			FunctionDef->CreateUObjectEngineTypes(Phase);
+		}
+
 		GetStruct()->Bind();
 
 		// Internals will assert of we are relinking an intrinsic 
