Deprecate CLASS_NoExport.  It was only being used by UHT.

#rnx

diff --git a/Engine/Source/Programs/Shared/EpicGames.Core/UnrealEngineTypes.cs b/Engine/Source/Programs/Shared/EpicGames.Core/UnrealEngineTypes.cs
--- a/Engine/Source/Programs/Shared/EpicGames.Core/UnrealEngineTypes.cs
+++ b/Engine/Source/Programs/Shared/EpicGames.Core/UnrealEngineTypes.cs
@@ -139,6 +139,7 @@ namespace EpicGames.Core
 		/// <summary>
 		/// Don't export to C++ header.
 		/// </summary>
+		[Obsolete("No longer used in the engine.")]
 		NoExport = 0x00000100u,
 
 		/// <summary>
@@ -174,6 +175,7 @@ namespace EpicGames.Core
 		/// <summary>
 		/// Do not export a constructor for this class, assuming it is in the cpptext
 		/// </summary>
+		[Obsolete("No longer used in the engine.")]
 		CustomConstructor = 0x00008000u,
 
 		/// <summary>
@@ -265,7 +267,7 @@ namespace EpicGames.Core
 		/// <summary>
 		/// These flags will be cleared by the compiler when the class is parsed during script compilation
 		/// </summary>
-		RecompilerClear = Inherit | Abstract | NoExport | Native | Intrinsic | TokenStreamAssembled,
+		RecompilerClear = Inherit | Abstract | Native | Intrinsic | TokenStreamAssembled,
 
 		/// <summary>
 		/// These flags will be cleared by the compiler when the class is parsed during script compilation
diff --git a/Engine/Source/Programs/Shared/EpicGames.UHT/Exporters/CodeGen/UhtCodeGenerator.cs b/Engine/Source/Programs/Shared/EpicGames.UHT/Exporters/CodeGen/UhtCodeGenerator.cs
--- a/Engine/Source/Programs/Shared/EpicGames.UHT/Exporters/CodeGen/UhtCodeGenerator.cs
+++ b/Engine/Source/Programs/Shared/EpicGames.UHT/Exporters/CodeGen/UhtCodeGenerator.cs
@@ -155,7 +155,8 @@ namespace EpicGames.UHT.Exporters.CodeGen
 							if (Type is UhtClass Class)
 							{
 								if (Class.ClassType != UhtClassType.NativeInterface && 
-									Class.ClassFlags.HasExactFlags(EClassFlags.Native | EClassFlags.NoExport | EClassFlags.Intrinsic, EClassFlags.Native))
+									Class.ClassFlags.HasExactFlags(EClassFlags.Native | EClassFlags.Intrinsic, EClassFlags.Native) &&
+									!Class.ClassExportFlags.HasAllFlags(UhtClassExportFlags.NoExport))
 								{
 									bWriteHeader = true;
 									break;
diff --git a/Engine/Source/Programs/Shared/EpicGames.UHT/Exporters/CodeGen/UhtHeaderCodeGeneratorCppFile.cs b/Engine/Source/Programs/Shared/EpicGames.UHT/Exporters/CodeGen/UhtHeaderCodeGeneratorCppFile.cs
--- a/Engine/Source/Programs/Shared/EpicGames.UHT/Exporters/CodeGen/UhtHeaderCodeGeneratorCppFile.cs
+++ b/Engine/Source/Programs/Shared/EpicGames.UHT/Exporters/CodeGen/UhtHeaderCodeGeneratorCppFile.cs
@@ -900,12 +900,12 @@ namespace EpicGames.UHT.Exporters.CodeGen
 			CallbackFunctions.Sort((x, y) => StringComparerUE.OrdinalIgnoreCase.Compare(x.EngineName, y.EngineName));
 
 			// VM -> C++ proxies (events and delegates).
-			if (!Class.ClassFlags.HasAnyFlags(EClassFlags.NoExport))
+			if (!Class.ClassExportFlags.HasAnyFlags(UhtClassExportFlags.NoExport))
 			{
 				AppendCallbackFunctions(Builder, Class, CallbackFunctions);
 			}
 
-			if (!Class.ClassFlags.HasAnyFlags(EClassFlags.NoExport))
+			if (!Class.ClassExportFlags.HasAnyFlags(UhtClassExportFlags.NoExport))
 			{
 				AppendNatives(Builder, Class);
 			}
@@ -1119,7 +1119,7 @@ namespace EpicGames.UHT.Exporters.CodeGen
 			{
 				if (!Function.FunctionFlags.HasAnyFlags(EFunctionFlags.Delegate))
 				{
-					AppendFunction(Builder, Function, Class.ClassFlags.HasAnyFlags(EClassFlags.NoExport));
+					AppendFunction(Builder, Function, Class.ClassExportFlags.HasAnyFlags(UhtClassExportFlags.NoExport));
 				}				
 			}
 
@@ -1345,7 +1345,7 @@ namespace EpicGames.UHT.Exporters.CodeGen
 					}
 				}
 
-				if (Class.ClassFlags.HasAnyFlags(EClassFlags.NoExport))
+				if (Class.ClassExportFlags.HasAnyFlags(UhtClassExportFlags.NoExport))
 				{
 					Builder.Append("\t/* friend declarations for pasting into noexport class ").Append(Class.SourceName).Append("\r\n");
 					Builder.Append("\tfriend struct ").Append(StaticsName).Append(";\r\n");
diff --git a/Engine/Source/Programs/Shared/EpicGames.UHT/Exporters/CodeGen/UhtHeaderCodeGeneratorHFile.cs b/Engine/Source/Programs/Shared/EpicGames.UHT/Exporters/CodeGen/UhtHeaderCodeGeneratorHFile.cs
--- a/Engine/Source/Programs/Shared/EpicGames.UHT/Exporters/CodeGen/UhtHeaderCodeGeneratorHFile.cs
+++ b/Engine/Source/Programs/Shared/EpicGames.UHT/Exporters/CodeGen/UhtHeaderCodeGeneratorHFile.cs
@@ -935,7 +935,7 @@ namespace EpicGames.UHT.Exporters.CodeGen
 			// Export the class's native function registration.
 			Builder.Append("private: \\\r\n");
 			Builder.Append("\tstatic void StaticRegisterNatives").Append(Class.SourceName).Append("(); \\\r\n");
-			if (!Class.ClassFlags.HasAnyFlags(EClassFlags.NoExport))
+			if (!Class.ClassExportFlags.HasAnyFlags(UhtClassExportFlags.NoExport))
 			{
 				Builder.Append("\tfriend struct ").Append(this.ObjectInfos[Class.ObjectTypeIndex].RegisteredSingletonName).Append("_Statics; \\\r\n");
 			}
diff --git a/Engine/Source/Programs/Shared/EpicGames.UHT/Parsers/UhtClassBaseParser.cs b/Engine/Source/Programs/Shared/EpicGames.UHT/Parsers/UhtClassBaseParser.cs
--- a/Engine/Source/Programs/Shared/EpicGames.UHT/Parsers/UhtClassBaseParser.cs
+++ b/Engine/Source/Programs/Shared/EpicGames.UHT/Parsers/UhtClassBaseParser.cs
@@ -44,7 +44,7 @@ namespace EpicGames.UHT.Parsers
 					BindAndResolveBases(this.BaseIdentifiers, UhtFindOptions.Class);
 
 					// Force the MatchedSerializers on for anything being exported
-					if (!this.ClassFlags.HasAnyFlags(EClassFlags.NoExport))
+					if (!this.ClassExportFlags.HasAnyFlags(UhtClassExportFlags.NoExport))
 					{
 						this.ClassFlags |= EClassFlags.MatchedSerializers;
 					}
diff --git a/Engine/Source/Programs/Shared/EpicGames.UHT/Parsers/UhtHeaderFileParser.cs b/Engine/Source/Programs/Shared/EpicGames.UHT/Parsers/UhtHeaderFileParser.cs
--- a/Engine/Source/Programs/Shared/EpicGames.UHT/Parsers/UhtHeaderFileParser.cs
+++ b/Engine/Source/Programs/Shared/EpicGames.UHT/Parsers/UhtHeaderFileParser.cs
@@ -248,7 +248,7 @@ namespace EpicGames.UHT.Parsers
 					{
 						if (Type is UhtClass Class)
 						{
-							if (Class.ClassType != UhtClassType.NativeInterface && !Class.ClassFlags.HasAnyFlags(EClassFlags.NoExport))
+							if (Class.ClassType != UhtClassType.NativeInterface && !Class.ClassExportFlags.HasAnyFlags(UhtClassExportFlags.NoExport))
 							{
 								bNoExportClassesOnly = false;
 								break;
diff --git a/Engine/Source/Programs/Shared/EpicGames.UHT/Specifiers/UhtClassSpecifiers.cs b/Engine/Source/Programs/Shared/EpicGames.UHT/Specifiers/UhtClassSpecifiers.cs
--- a/Engine/Source/Programs/Shared/EpicGames.UHT/Specifiers/UhtClassSpecifiers.cs
+++ b/Engine/Source/Programs/Shared/EpicGames.UHT/Specifiers/UhtClassSpecifiers.cs
@@ -18,7 +18,7 @@ namespace EpicGames.UHT.Parsers
 		private static void NoExportSpecifier(UhtSpecifierContext SpecifierContext)
 		{
 			UhtClassParser Class = (UhtClassParser)SpecifierContext.Scope.ScopeType;
-			Class.AddClassFlags(EClassFlags.NoExport);
+			Class.ClassExportFlags |= UhtClassExportFlags.NoExport;
 		}
 
 		[UhtSpecifier(Extends = UhtTableNames.Class, ValueType = UhtSpecifierValueType.Legacy)]
diff --git a/Engine/Source/Programs/Shared/EpicGames.UHT/Types/UhtClass.cs b/Engine/Source/Programs/Shared/EpicGames.UHT/Types/UhtClass.cs
--- a/Engine/Source/Programs/Shared/EpicGames.UHT/Types/UhtClass.cs
+++ b/Engine/Source/Programs/Shared/EpicGames.UHT/Types/UhtClass.cs
@@ -69,6 +69,11 @@ namespace EpicGames.UHT.Types
 		/// GetLifetimeReplicatedProps was found in the class
 		/// </summary>
 		HasGetLifetimeReplicatedProps = 1 << 7,
+
+		/// <summary>
+		/// Class should not be exported
+		/// </summary>
+		NoExport = 1 << 8,
 	}
 
 	/// <summary>
diff --git a/Engine/Source/Programs/Shared/EpicGames.UHT/Types/UhtProperty.cs b/Engine/Source/Programs/Shared/EpicGames.UHT/Types/UhtProperty.cs
--- a/Engine/Source/Programs/Shared/EpicGames.UHT/Types/UhtProperty.cs
+++ b/Engine/Source/Programs/Shared/EpicGames.UHT/Types/UhtProperty.cs
@@ -2132,7 +2132,7 @@ namespace EpicGames.UHT.Types
 			}
 			else if (Object is UhtClass Class)
 			{
-				if (Class.ClassFlags.HasAnyFlags(EClassFlags.NoExport))
+				if (Class.ClassExportFlags.HasAnyFlags(UhtClassExportFlags.NoExport))
 				{
 					return Builder;
 				}
diff --git a/Engine/Source/Programs/UnrealHeaderTool/Private/CodeGenerator.cpp b/Engine/Source/Programs/UnrealHeaderTool/Private/CodeGenerator.cpp
--- a/Engine/Source/Programs/UnrealHeaderTool/Private/CodeGenerator.cpp
+++ b/Engine/Source/Programs/UnrealHeaderTool/Private/CodeGenerator.cpp
@@ -1995,7 +1995,7 @@ void FNativeClassHeaderGenerator::ExportNativeGeneratedInitCode(FOutputDevice& O
 {
 	check(!OutFriendText.Len());
 
-	const bool   bIsNoExport  = ClassDef.HasAnyClassFlags(CLASS_NoExport);
+	const bool   bIsNoExport  = ClassDef.IsNoExport();
 	const FString ClassNameCPP = ClassDef.GetAlternateNameCPP();
 
 	const FString& ApiString = GetAPIString();
@@ -2750,7 +2750,7 @@ void FNativeClassHeaderGenerator::ExportClassFromSourceFileInner(
 
 		// VM -> C++ proxies (events and delegates).
 		FOutputDeviceNull NullOutput;
-		FOutputDevice& CallbackOut = ClassDef.HasAnyClassFlags(CLASS_NoExport) ? NullOutput : OutCpp;
+		FOutputDevice& CallbackOut = ClassDef.IsNoExport() ? NullOutput : OutCpp;
 		FString CallbackWrappersMacroName = SourceFile.GetGeneratedMacroName(ClassDef.GetGeneratedBodyLine(), TEXT("_CALLBACK_WRAPPERS"));
 		ExportCallbackFunctions(
 			OutGeneratedHeaderText,
@@ -2767,7 +2767,7 @@ void FNativeClassHeaderGenerator::ExportClassFromSourceFileInner(
 	}
 
 	// Class definition.
-	if (!ClassDef.HasAnyClassFlags(CLASS_NoExport))
+	if (!ClassDef.IsNoExport())
 	{
 		ExportNatives(OutCpp, ClassDef);
 	}
@@ -7118,8 +7118,6 @@ void ProcessParsedClass(FUnrealClassDefinitionInfo& ClassDef)
 		ClassDef.Throwf(TEXT("Class '%s' cannot inherit from itself"), *ClassName);
 	}
 
-	UClass* ResultClass = FEngineAPI::FindObject<UClass>(ANY_PACKAGE, *ClassNameStripped);
-
 	// if we aren't generating headers, then we shouldn't set misaligned object, since it won't get cleared
 
 	const static bool bVerboseOutput = FParse::Param(FCommandLine::Get(), TEXT("VERBOSE"));
diff --git a/Engine/Source/Programs/UnrealHeaderTool/Private/HeaderParser.cpp b/Engine/Source/Programs/UnrealHeaderTool/Private/HeaderParser.cpp
--- a/Engine/Source/Programs/UnrealHeaderTool/Private/HeaderParser.cpp
+++ b/Engine/Source/Programs/UnrealHeaderTool/Private/HeaderParser.cpp
@@ -7957,7 +7957,7 @@ void FHeaderParser::ParseHeader()
 		FUnrealClassDefinitionInfo& ClassDef = UHTCastChecked<FUnrealClassDefinitionInfo>(TypeDef);
 		PostParsingClassSetup(ClassDef);
 
-		bNoExportClassesOnly = bNoExportClassesOnly && ClassDef.HasAnyClassFlags(CLASS_NoExport);
+		bNoExportClassesOnly = bNoExportClassesOnly && ClassDef.IsNoExport();
 	}
 
 	if (!bSpottedAutogeneratedHeaderInclude && !bEmptyFile && !bNoExportClassesOnly)
@@ -8603,7 +8603,7 @@ void FHeaderParser::SimplifiedClassParse(FUnrealSourceFile& SourceFile, const TC
 					}
 
 					TSharedRef<FUnrealTypeDefinitionInfo> ClassDecl = Parser.ParseClassDeclaration(StartOfLine + (UInterfaceMacroDecl - Str), CurrentLine, true, TEXT("UINTERFACE"));
-					bGeneratedIncludeRequired |= !ClassDecl->AsClassChecked().HasAnyClassFlags(CLASS_NoExport);
+					bGeneratedIncludeRequired |= !ClassDecl->AsClassChecked().IsNoExport();
 					SourceFile.AddDefinedClass(MoveTemp(ClassDecl));
 				}
 			}
@@ -8618,7 +8618,7 @@ void FHeaderParser::SimplifiedClassParse(FUnrealSourceFile& SourceFile, const TC
 					}
 
 					TSharedRef<FUnrealTypeDefinitionInfo> ClassDecl = Parser.ParseClassDeclaration(StartOfLine + (UClassMacroDecl - Str), CurrentLine, false, TEXT("UCLASS"));
-					bGeneratedIncludeRequired |= !ClassDecl->AsClassChecked().HasAnyClassFlags(CLASS_NoExport);
+					bGeneratedIncludeRequired |= !ClassDecl->AsClassChecked().IsNoExport();
 					SourceFile.AddDefinedClass(MoveTemp(ClassDecl));
 				}
 			}
diff --git a/Engine/Source/Programs/UnrealHeaderTool/Private/UnrealTypeDefinitionInfo.cpp b/Engine/Source/Programs/UnrealHeaderTool/Private/UnrealTypeDefinitionInfo.cpp
--- a/Engine/Source/Programs/UnrealHeaderTool/Private/UnrealTypeDefinitionInfo.cpp
+++ b/Engine/Source/Programs/UnrealHeaderTool/Private/UnrealTypeDefinitionInfo.cpp
@@ -2099,12 +2099,9 @@ FUnrealClassDefinitionInfo* FUnrealClassDefinitionInfo::FindScriptClass(const FS
 
 uint32 FUnrealClassDefinitionInfo::GetHash(FUnrealTypeDefinitionInfo& ReferencingDef, bool bIncludeNoExport) const
 {
-	if (!bIncludeNoExport)
+	if (!bIncludeNoExport && IsNoExport())
 	{
-		if (HasAnyClassFlags(CLASS_NoExport))
-		{
-			return 0;
-		}
+		return 0;
 	}
 	return FUnrealStructDefinitionInfo::GetHash(ReferencingDef, bIncludeNoExport);
 }
@@ -2212,7 +2209,7 @@ void FUnrealClassDefinitionInfo::CreateUObjectEngineTypesInternal(ECreateEngineT
 			{
 				//Class->UnMark(EObjectMark(OBJECTMARK_TagImp | OBJECTMARK_TagExp));
 			}
-			else if (!HasAnyClassFlags(CLASS_NoExport | CLASS_Intrinsic))
+			else if (!IsNoExport() && !HasAnyClassFlags(CLASS_Intrinsic))
 			{
 				//Class->UnMark(OBJECTMARK_TagImp);
 				//Class->Mark(OBJECTMARK_TagExp);
@@ -2297,7 +2294,7 @@ void FUnrealClassDefinitionInfo::PostParseFinalizeInternal(EPostParseFinalizePha
 		if (!HasAnyClassFlags(CLASS_Native))
 		{
 		}
-		else if (!HasAnyClassFlags(CLASS_NoExport | CLASS_Intrinsic))
+		else if (!IsNoExport() && !HasAnyClassFlags(CLASS_Intrinsic))
 		{
 			GetPackageDef().SetWriteClassesH(true);
 		}
@@ -2307,7 +2304,7 @@ void FUnrealClassDefinitionInfo::PostParseFinalizeInternal(EPostParseFinalizePha
 		{
 
 			// Check for differences in ClassFlags
-			const EClassFlags CheckFlags = CLASS_SaveInCompiledInClasses & ~CLASS_NoExport & ~CLASS_Intrinsic;
+			const EClassFlags CheckFlags = CLASS_SaveInCompiledInClasses & ~CLASS_Intrinsic;
 			if ((ClassFlags & CheckFlags) != (Class->ClassFlags & CheckFlags))
 			{
 				LogError(TEXT("Class flags in core class '%s', is '0x%08x', should be '0x%08x'"), *GetNameCPP(), ClassFlags & CheckFlags, Class->ClassFlags & CheckFlags);
@@ -2393,7 +2390,7 @@ void FUnrealClassDefinitionInfo::ParseClassProperties(TArray<FPropertySpecifier>
 		case EClassMetadataSpecifier::NoExport:
 
 			// Don't export to C++ header.
-			ParsedClassFlags |= CLASS_NoExport;
+			MarkNoExport();
 			break;
 
 		case EClassMetadataSpecifier::Intrinsic:
diff --git a/Engine/Source/Programs/UnrealHeaderTool/Private/UnrealTypeDefinitionInfo.h b/Engine/Source/Programs/UnrealHeaderTool/Private/UnrealTypeDefinitionInfo.h
--- a/Engine/Source/Programs/UnrealHeaderTool/Private/UnrealTypeDefinitionInfo.h
+++ b/Engine/Source/Programs/UnrealHeaderTool/Private/UnrealTypeDefinitionInfo.h
@@ -3336,6 +3336,22 @@ public:
 		return bIsCompiledIn;
 	}
 
+	/**
+	 * Return true if the class should not export
+	 */
+	bool IsNoExport() const
+	{
+		return bNoExport;
+	}
+
+	/**
+	 * Mark the class as not to be exported
+	 */
+	void MarkNoExport()
+	{
+		bNoExport = true;
+	}
+
 	/**
 	 * Get the generated body access specifier
 	 */
@@ -3455,6 +3471,9 @@ private:
 	/** True if the class has a custom constructor */
 	bool bCustomConstructor = false;
 
+	/** True if the class is not to be exported */
+	bool bNoExport = false;
+
 	/** True if the class is an interface */
 	bool bIsInterface = false;
 
diff --git a/Engine/Source/Runtime/CoreUObject/Public/UObject/Object.h b/Engine/Source/Runtime/CoreUObject/Public/UObject/Object.h
--- a/Engine/Source/Runtime/CoreUObject/Public/UObject/Object.h
+++ b/Engine/Source/Runtime/CoreUObject/Public/UObject/Object.h
@@ -52,7 +52,7 @@ namespace ECastCheckedType
 class COREUOBJECT_API UObject : public UObjectBaseUtility
 {
 	// Declarations, normally created by UnrealHeaderTool boilerplate code
-	DECLARE_CLASS(UObject,UObject,CLASS_Abstract|CLASS_NoExport|CLASS_Intrinsic|CLASS_MatchedSerializers,CASTCLASS_None,TEXT("/Script/CoreUObject"),NO_API)
+	DECLARE_CLASS(UObject,UObject,CLASS_Abstract|CLASS_Intrinsic|CLASS_MatchedSerializers,CASTCLASS_None,TEXT("/Script/CoreUObject"),NO_API)
 	DEFINE_DEFAULT_OBJECT_INITIALIZER_CONSTRUCTOR_CALL(UObject)
 	typedef UObject WithinClass;
 	static UObject* __VTableCtorCaller(FVTableHelper& Helper)
diff --git a/Engine/Source/Runtime/CoreUObject/Public/UObject/ObjectMacros.h b/Engine/Source/Runtime/CoreUObject/Public/UObject/ObjectMacros.h
--- a/Engine/Source/Runtime/CoreUObject/Public/UObject/ObjectMacros.h
+++ b/Engine/Source/Runtime/CoreUObject/Public/UObject/ObjectMacros.h
@@ -199,7 +199,7 @@ enum EClassFlags
 	/** Class is a native class - native interfaces will have CLASS_Native set, but not RF_MarkAsNative */
 	CLASS_Native			  = 0x00000080u,
 	/** Don't export to C++ header. */
-	CLASS_NoExport            = 0x00000100u,
+	CLASS_NoExport UE_DEPRECATED(5.1, "CLASS_NoExport should no longer be used. It is no longer being set by engine code.") = 0x00000100u,
 	/** Do not allow users to create in the editor. */
 	CLASS_NotPlaceable        = 0x00000200u,
 	/** Handle object configuration on a per-object basis, rather than per-class. */
@@ -215,7 +215,7 @@ enum EClassFlags
 	/** Class is an interface **/
 	CLASS_Interface           = 0x00004000u,
 	/**  Do not export a constructor for this class, assuming it is in the cpptext **/
-	CLASS_CustomConstructor   = 0x00008000u,
+	CLASS_CustomConstructor UE_DEPRECATED(5.1, "CLASS_CustomConstructor should no longer be used. It is no longer being set by engine code.") = 0x00008000u,
 	/** all properties and functions in this class are const and should be exported as const */
 	CLASS_Const			      = 0x00010000u,
 
@@ -264,7 +264,7 @@ ENUM_CLASS_FLAGS(EClassFlags);
 						| CLASS_Const | CLASS_HasInstancedReference | CLASS_Deprecated | CLASS_DefaultToInstanced | CLASS_GlobalUserConfig | CLASS_ProjectUserConfig | CLASS_NeedsDeferredDependencyLoading))
 
 /** These flags will be cleared by the compiler when the class is parsed during script compilation */
-#define CLASS_RecompilerClear ((EClassFlags)(CLASS_Inherit | CLASS_Abstract | CLASS_NoExport | CLASS_Native | CLASS_Intrinsic | CLASS_TokenStreamAssembled))
+#define CLASS_RecompilerClear ((EClassFlags)(CLASS_Inherit | CLASS_Abstract | CLASS_Native | CLASS_Intrinsic | CLASS_TokenStreamAssembled))
 
 /** These flags will be cleared by the compiler when the class is parsed during script compilation */
 #define CLASS_ShouldNeverBeLoaded ((EClassFlags)(CLASS_Native | CLASS_Optional | CLASS_Intrinsic | CLASS_TokenStreamAssembled))
