UHT: Add FieldNotify specifier to UHT. It auto generate the UE_FIELD_NOTIFICATION_DECLARE_CLASS_DESCRIPTOR and UE_FIELD_NOTIFICATION_IMPLEMENT_CLASS_DESCRIPTOR macro. FieldNotify is only available on  FProperty and UFunction member of a UCLass and if the UClass implement the INotifyFieldValueChanged interface.

diff --git a/Engine/Source/Programs/UnrealHeaderTool/Private/CodeGenerator.cpp b/Engine/Source/Programs/UnrealHeaderTool/Private/CodeGenerator.cpp
--- a/Engine/Source/Programs/UnrealHeaderTool/Private/CodeGenerator.cpp
+++ b/Engine/Source/Programs/UnrealHeaderTool/Private/CodeGenerator.cpp
@@ -2972,6 +2972,7 @@ void FNativeClassHeaderGenerator::ExportClassFromSourceFileInner(
 				ClassMacroCalls.Logf(TEXT("\t%s\r\n"), *MacroName);
 
 				ExportConstructorsMacros(OutGeneratedHeaderText, OutCpp, StandardUObjectConstructorsMacroCall, EnhancedUObjectConstructorsMacroCall, SourceFile.GetGeneratedMacroName(ClassDef.GetGeneratedBodyLine()), ClassDef, *APIArg);
+				ExportFieldNotify(OutGeneratedHeaderText, OutCpp, StandardUObjectConstructorsMacroCall, EnhancedUObjectConstructorsMacroCall, SourceFile.GetGeneratedMacroName(ClassDef.GetGeneratedBodyLine()), ClassDef);
 			}
 		}
 	}
@@ -5047,6 +5048,210 @@ void FNativeClassHeaderGenerator::ExportNativeFunctions(FOutputDevice& OutGenera
 	}
 }
 
+
+void FNativeClassHeaderGenerator::ExportFieldNotify(FOutputDevice& OutGeneratedHeaderText, FOutputDevice& OutGeneratedCPPText, FOutputDevice& StandardUObjectConstructorsMacroCall, FOutputDevice& EnhancedUObjectConstructorsMacroCall, const FString& ConstructorsMacroPrefix, FUnrealClassDefinitionInfo& ClassDef) const
+{
+	TArray<TSharedRef<FUnrealPropertyDefinitionInfo>> EditorProperties;
+	TArray<TSharedRef<FUnrealPropertyDefinitionInfo>> RuntimeProperties;
+	TArray<TSharedRef<FUnrealFunctionDefinitionInfo>> EditorFunctions;
+	TArray<TSharedRef<FUnrealFunctionDefinitionInfo>> RuntimeFunctions;
+	const FString ClassCPPName = ClassDef.GetAlternateNameCPP();
+
+	if (ClassDef.HasCustomFieldNotify())
+	{
+		return;
+	}
+
+	{
+		TArray<TSharedRef<FUnrealPropertyDefinitionInfo>>& Properties = ClassDef.GetProperties();
+		for (TSharedRef<FUnrealPropertyDefinitionInfo>& Prop : Properties)
+		{
+			if (Prop->GetPropertyBase().bFieldNotify)
+			{
+				const bool bEditorOnlyProperty = Prop->IsEditorOnlyProperty();
+				if (bEditorOnlyProperty)
+				{
+					EditorProperties.Add(Prop);
+				}
+				else
+				{
+					RuntimeProperties.Add(Prop);
+				}
+			}
+		}
+		for (TSharedRef<FUnrealFunctionDefinitionInfo> FunctionDef : ClassDef.GetFunctions())
+		{
+			if (FunctionDef->GetFunctionData().bFieldNotify)
+			{
+				const bool bEditorOnlyFunc = FunctionDef->HasAnyFunctionFlags(FUNC_EditorOnly);
+				if (bEditorOnlyFunc)
+				{
+					EditorFunctions.Add(FunctionDef);
+				}
+				else
+				{
+					RuntimeFunctions.Add(FunctionDef);
+				}
+			}
+		}
+	}
+
+	if (RuntimeProperties.Num() > 0 || EditorProperties.Num() > 0 || RuntimeFunctions.Num() > 0 || EditorFunctions.Num() > 0)
+	{
+		FUHTStringBuilder RuntimeHeaderStringBuilders;
+		FUHTStringBuilder EditorHeaderStringBuilders;
+
+		const bool bOnlyHasEditorFields = RuntimeProperties.Num() == 0 && RuntimeFunctions.Num() == 0;
+		const bool bOnlyHasRuntimeFields = EditorProperties.Num() == 0 && EditorFunctions.Num() == 0;
+		const bool bHasEditorFields = EditorProperties.Num() > 0 || EditorFunctions.Num() > 0;
+
+		if (bOnlyHasEditorFields)
+		{
+			OutGeneratedCPPText.Logf(TEXT("#if WITH_EDITORONLY_DATA\r\n"));
+		}
+
+		RuntimeHeaderStringBuilders.Logf(TEXT("\tUE_FIELD_NOTIFICATION_DECLARE_CLASS_DESCRIPTOR_BEGIN()\r\n"));
+		EditorHeaderStringBuilders.Logf(TEXT("\tUE_FIELD_NOTIFICATION_DECLARE_CLASS_DESCRIPTOR_BEGIN()\r\n"));
+
+		//UE_FIELD_NOTIFICATION_DECLARE_FIELD
+		{
+			for (TSharedRef<FUnrealPropertyDefinitionInfo>& Prop : RuntimeProperties)
+			{
+				RuntimeHeaderStringBuilders.Logf(TEXT("\tUE_FIELD_NOTIFICATION_DECLARE_FIELD(%s)\r\n"), *Prop->GetName());
+				EditorHeaderStringBuilders.Logf(TEXT("\tUE_FIELD_NOTIFICATION_DECLARE_FIELD(%s)\r\n"), *Prop->GetName());
+				OutGeneratedCPPText.Logf(TEXT("\tUE_FIELD_NOTIFICATION_IMPLEMENT_FIELD(%s, %s)\r\n"), *ClassCPPName ,*Prop->GetName());
+			}
+			for (TSharedRef<FUnrealFunctionDefinitionInfo>& FunctionDef : RuntimeFunctions)
+			{
+				RuntimeHeaderStringBuilders.Logf(TEXT("\tUE_FIELD_NOTIFICATION_DECLARE_FIELD(%s)\r\n"), *FunctionDef->GetFunctionData().CppImplName);
+				EditorHeaderStringBuilders.Logf(TEXT("\tUE_FIELD_NOTIFICATION_DECLARE_FIELD(%s)\r\n"), *FunctionDef->GetFunctionData().CppImplName);
+				OutGeneratedCPPText.Logf(TEXT("\tUE_FIELD_NOTIFICATION_IMPLEMENT_FIELD(%s, %s)\r\n"), *ClassCPPName, *FunctionDef->GetFunctionData().CppImplName);
+			}
+
+			if (bHasEditorFields)
+			{
+				if (!bOnlyHasEditorFields)
+				{
+					OutGeneratedCPPText.Logf(TEXT("#if WITH_EDITORONLY_DATA\r\n"));
+				}
+
+				for (TSharedRef<FUnrealPropertyDefinitionInfo>& Prop : EditorProperties)
+				{
+					EditorHeaderStringBuilders.Logf(TEXT("\tUE_FIELD_NOTIFICATION_DECLARE_FIELD(%s)\r\n"), *Prop->GetName());
+					OutGeneratedCPPText.Logf(TEXT("\tUE_FIELD_NOTIFICATION_IMPLEMENT_FIELD(%s, %s)\r\n"), *ClassCPPName, *Prop->GetName());
+				}
+				for (TSharedRef<FUnrealFunctionDefinitionInfo>& FunctionDef : EditorFunctions)
+				{
+					EditorHeaderStringBuilders.Logf(TEXT("\tUE_FIELD_NOTIFICATION_DECLARE_FIELD(%s)\r\n"), *FunctionDef->GetFunctionData().CppImplName);
+					OutGeneratedCPPText.Logf(TEXT("\tUE_FIELD_NOTIFICATION_IMPLEMENT_FIELD(%s, %s)\r\n"), *ClassCPPName, *FunctionDef->GetFunctionData().CppImplName);
+				}
+
+				if (!bOnlyHasEditorFields)
+				{
+					OutGeneratedCPPText.Logf(TEXT("#endif // WITH_EDITORONLY_DATA\r\n"));
+				}
+			}
+		}
+
+		OutGeneratedCPPText.Logf(TEXT("\tUE_FIELD_NOTIFICATION_IMPLEMENTATION_BEGIN(%s)\r\n"), *ClassCPPName);
+
+		//UE_FIELD_NOTIFICATION_DECLARE_ENUM_FIELD
+		{
+			bool bIsRuntimeFirst = true;
+			bool bIsEditorFirst = true;
+			auto AddFieldNotificationEnumField = [&bIsRuntimeFirst, &bIsEditorFirst, &RuntimeHeaderStringBuilders, &EditorHeaderStringBuilders, &OutGeneratedCPPText, &ClassCPPName](bool bRuntime, const FString& FieldName)
+			{
+				if (bRuntime)
+				{
+					if (bIsRuntimeFirst)
+					{
+						bIsRuntimeFirst = false;
+						bIsEditorFirst = false;
+						RuntimeHeaderStringBuilders.Logf(TEXT("\tUE_FIELD_NOTIFICATION_DECLARE_ENUM_FIELD_BEGIN(%s)\r\n"), *FieldName);
+						EditorHeaderStringBuilders.Logf(TEXT("\tUE_FIELD_NOTIFICATION_DECLARE_ENUM_FIELD_BEGIN(%s)\r\n"), *FieldName);
+					}
+					else
+					{
+						RuntimeHeaderStringBuilders.Logf(TEXT("\tUE_FIELD_NOTIFICATION_DECLARE_ENUM_FIELD(%s)\r\n"), *FieldName);
+						EditorHeaderStringBuilders.Logf(TEXT("\tUE_FIELD_NOTIFICATION_DECLARE_ENUM_FIELD(%s)\r\n"), *FieldName);
+					}
+				}
+				else
+				{
+					if (bIsEditorFirst)
+					{
+						bIsEditorFirst = false;
+						EditorHeaderStringBuilders.Logf(TEXT("\tUE_FIELD_NOTIFICATION_DECLARE_ENUM_FIELD_BEGIN(%s)\r\n"), *FieldName);
+					}
+					EditorHeaderStringBuilders.Logf(TEXT("\tUE_FIELD_NOTIFICATION_DECLARE_ENUM_FIELD(%s)\r\n"), *FieldName);
+				}
+				OutGeneratedCPPText.Logf(TEXT("\tUE_FIELD_NOTIFICATION_IMPLEMENT_ENUM_FIELD(%s, %s)\r\n"), *ClassCPPName, *FieldName);
+			};
+
+			bool bIsFirstRuntime = true;
+			for (TSharedRef<FUnrealPropertyDefinitionInfo>& Prop : RuntimeProperties)
+			{
+				AddFieldNotificationEnumField(true, *Prop->GetName());
+			}
+			for (TSharedRef<FUnrealFunctionDefinitionInfo>& FunctionDef : RuntimeFunctions)
+			{
+				AddFieldNotificationEnumField(true, *FunctionDef->GetFunctionData().CppImplName);
+			}
+
+			if (bHasEditorFields)
+			{
+				if (!bOnlyHasEditorFields)
+				{
+					OutGeneratedCPPText.Logf(TEXT("#if WITH_EDITORONLY_DATA\r\n"));
+				}
+
+				for (TSharedRef<FUnrealPropertyDefinitionInfo>& Prop : EditorProperties)
+				{
+					AddFieldNotificationEnumField(false, *Prop->GetName());
+				}
+				for (TSharedRef<FUnrealFunctionDefinitionInfo>& FunctionDef : EditorFunctions)
+				{
+					AddFieldNotificationEnumField(false, *FunctionDef->GetFunctionData().CppImplName);
+				}
+
+				if (!bOnlyHasEditorFields)
+				{
+					OutGeneratedCPPText.Logf(TEXT("#endif // WITH_EDITORONLY_DATA\r\n"));
+				}
+			}
+		}
+
+		RuntimeHeaderStringBuilders.Logf(TEXT("\tUE_FIELD_NOTIFICATION_DECLARE_ENUM_FIELD_END()\r\n"));
+		EditorHeaderStringBuilders.Logf(TEXT("\tUE_FIELD_NOTIFICATION_DECLARE_ENUM_FIELD_END()\r\n"));
+		RuntimeHeaderStringBuilders.Logf(TEXT("\tUE_FIELD_NOTIFICATION_DECLARE_CLASS_DESCRIPTOR_END();\r\n"));
+		EditorHeaderStringBuilders.Logf(TEXT("\tUE_FIELD_NOTIFICATION_DECLARE_CLASS_DESCRIPTOR_END();\r\n"));
+		OutGeneratedCPPText.Logf(TEXT("\tUE_FIELD_NOTIFICATION_IMPLEMENTATION_END(%s);\r\n"), *ClassCPPName);
+
+		if (bOnlyHasEditorFields)
+		{
+			OutGeneratedCPPText.Logf(TEXT("#endif // WITH_EDITORONLY_DATA\r\n"));
+		}
+
+		FString HeaderMacroName = ConstructorsMacroPrefix + TEXT("_FIELDNOTIFY");
+
+		if (bHasEditorFields)
+		{
+			OutGeneratedHeaderText.Logf(TEXT("#if WITH_EDITORONLY_DATA\r\n"));
+			WriteMacro(OutGeneratedHeaderText, HeaderMacroName, EditorHeaderStringBuilders);
+			OutGeneratedHeaderText.Logf(TEXT("#else //WITH_EDITORONLY_DATA\r\n"));
+			WriteMacro(OutGeneratedHeaderText, HeaderMacroName, RuntimeHeaderStringBuilders);
+			OutGeneratedHeaderText.Logf(TEXT("#endif // WITH_EDITORONLY_DATA\r\n"));
+		}
+		else 
+		{
+			WriteMacro(OutGeneratedHeaderText, HeaderMacroName, RuntimeHeaderStringBuilders);
+		}
+
+		StandardUObjectConstructorsMacroCall.Logf(TEXT("\t%s\r\n"), *HeaderMacroName);
+		EnhancedUObjectConstructorsMacroCall.Logf(TEXT("\t%s\r\n"), *HeaderMacroName);
+	}
+}
+
 /**
  * Exports the methods which trigger UnrealScript events and delegates.
  *
diff --git a/Engine/Source/Programs/UnrealHeaderTool/Private/HeaderParser.cpp b/Engine/Source/Programs/UnrealHeaderTool/Private/HeaderParser.cpp
--- a/Engine/Source/Programs/UnrealHeaderTool/Private/HeaderParser.cpp
+++ b/Engine/Source/Programs/UnrealHeaderTool/Private/HeaderParser.cpp
@@ -534,6 +534,11 @@ namespace
 					}
 				}
 				break;
+				case EFunctionSpecifier::FieldNotify:
+				{
+					FuncInfo.bFieldNotify = true;
+				}
+				break;
 			}
 		}
 
@@ -2608,6 +2613,73 @@ void FHeaderParser::VerifyRepNotifyCallback(FUnrealPropertyDefinitionInfo& Prope
 		LogError(TEXT("Replication notification function %s not found"), *PropertyDef.GetRepNotifyFunc().ToString() );
 	}
 }
+
+void FHeaderParser::VerifyNotifyValueChangedProperties(FUnrealPropertyDefinitionInfo& PropertyDef)
+{
+	FUnrealTypeDefinitionInfo* Info = PropertyDef.GetOuter();
+	if (Info == nullptr)
+	{
+		LogError(TEXT("FieldNofity property %s does not have a outer."), *PropertyDef.GetName());
+		return;
+	}
+	
+	FUnrealClassDefinitionInfo* ClassInfo = Info->AsClass();
+	if (ClassInfo == nullptr)
+	{
+		LogError(TEXT("FieldNofity property are only valid as UClass member variable."));
+		return;
+	}
+	if (ClassInfo->IsInterface())
+	{
+		LogError(TEXT("FieldNofity are not valid on UInterface."));
+		return;
+	}
+
+	ClassInfo->MarkHasFieldNotify();
+}
+
+void FHeaderParser::VerifyNotifyValueChangedFunction(const FUnrealFunctionDefinitionInfo& TargetFuncDef)
+{
+	FUnrealTypeDefinitionInfo* Info = TargetFuncDef.GetOuter();
+	if (Info == nullptr)
+	{
+		LogError(TEXT("FieldNofity function %s does not have a outer."), *TargetFuncDef.GetName());
+		return;
+	}
+
+	FUnrealClassDefinitionInfo* ClassInfo = Info->AsClass();
+	if (ClassInfo == nullptr)
+	{
+		LogError(TEXT("FieldNofity function %s are only valid as UClass member function."), *TargetFuncDef.GetName());
+		return;
+	}
+
+	const TArray<TSharedRef<FUnrealPropertyDefinitionInfo>>& Properties = TargetFuncDef.GetProperties();
+	FUnrealPropertyDefinitionInfo* ReturnPropDef = TargetFuncDef.GetReturn();
+	if (Properties.Num() > 1 || (Properties.Num() == 1 && ReturnPropDef == nullptr))
+	{
+		LogError(TEXT("FieldNotify function %s must not have parameters."), *TargetFuncDef.GetName());
+		return;
+	}
+	if (ReturnPropDef == nullptr)
+	{
+		LogError(TEXT("FieldNotify function %s must return a value."), *TargetFuncDef.GetName());
+		return;
+	}
+	if (TargetFuncDef.HasAnyFunctionFlags(FUNC_Event))
+	{
+		LogError(TEXT("FieldNotify function %s cannot be a blueprint event."), *TargetFuncDef.GetName());
+		return;
+	}
+	if (!TargetFuncDef.HasAnyFunctionFlags(FUNC_BlueprintPure))
+	{
+		LogError(TEXT("FieldNotify function %s must be pure."), *TargetFuncDef.GetName());
+		return;
+	}
+
+	ClassInfo->MarkHasFieldNotify();
+}
+
 void FHeaderParser::VerifyPropertyMarkups(FUnrealClassDefinitionInfo& TargetClassDef)
 {
 	// Iterate over all properties, looking for those flagged as CPF_RepNotify
@@ -2688,6 +2760,22 @@ void FHeaderParser::VerifyPropertyMarkups(FUnrealClassDefinitionInfo& TargetClas
 				LogError(TEXT("Property %s getter function %s not found"), *PropertyDef->GetName(), *PropertyToken.GetterName);
 			}
 		}
+
+		if (PropertyDef->GetPropertyBase().bFieldNotify)
+		{
+			VerifyNotifyValueChangedProperties(*PropertyDef);
+		}
+	}
+}
+
+void FHeaderParser::VerifyFunctionsMarkups(FUnrealClassDefinitionInfo& TargetClassDef)
+{
+	for (TSharedRef<FUnrealFunctionDefinitionInfo> FunctionDef : TargetClassDef.GetFunctions())
+	{
+		if (FunctionDef->GetFunctionData().bFieldNotify)
+		{
+			VerifyNotifyValueChangedFunction(*FunctionDef);
+		}
 	}
 }
 
@@ -2859,6 +2947,7 @@ void FHeaderParser::GetVarType(
 	FString GetterName;
 	bool bHasSetterTag = false;
 	bool bHasGetterTag = false;
+	bool bFieldNotify = false;
 
 	// Get flags.
 	EPropertyFlags Flags        = CPF_None;
@@ -3311,6 +3400,12 @@ void FHeaderParser::GetVarType(
 					}
 				}
 				break;
+				
+				case EVariableSpecifier::FieldNotify:
+				{
+					bFieldNotify = true;
+				}
+				break;
 
 				default:
 				{
@@ -4406,6 +4501,8 @@ void FHeaderParser::GetVarType(
 	VarProperty.bGetterTagFound = bHasGetterTag;
 	VarProperty.GetterName = GetterName;
 
+	VarProperty.bFieldNotify = bFieldNotify;
+
 	// Perform some more specific validation on the property flags
 	if (VarProperty.PropertyFlags & CPF_PersistentInstance)
 	{
@@ -9179,6 +9276,8 @@ void FHeaderParser::PostPopNestClass(FUnrealClassDefinitionInfo& CurrentClassDef
 	// Validate all the rep notify events here, to make sure they're implemented
 	VerifyPropertyMarkups(CurrentClassDef);
 
+	VerifyFunctionsMarkups(CurrentClassDef);
+
 	// Iterate over all the interfaces we claim to implement
 	for (const FUnrealStructDefinitionInfo::FBaseStructInfo& BaseClassInfo : CurrentClassDef.GetBaseStructInfos())
 	{
diff --git a/Engine/Source/Programs/UnrealHeaderTool/Private/HeaderParser.h b/Engine/Source/Programs/UnrealHeaderTool/Private/HeaderParser.h
--- a/Engine/Source/Programs/UnrealHeaderTool/Private/HeaderParser.h
+++ b/Engine/Source/Programs/UnrealHeaderTool/Private/HeaderParser.h
@@ -593,6 +593,15 @@ private:
 	// Verifies the target function meets the criteria for a replication notify callback
 	void VerifyRepNotifyCallback(FUnrealPropertyDefinitionInfo& PropertyDef, FUnrealFunctionDefinitionInfo* TargetFuncDef);
 
+	// Vefifies the target property meets the criteria for FieldNotify
+	void VerifyNotifyValueChangedProperties(FUnrealPropertyDefinitionInfo& PropertyDef);
+
+	// Verifies that all specified class's UFunction
+	void VerifyFunctionsMarkups(FUnrealClassDefinitionInfo& TargetClassDef);
+
+	// Verifies the target function meets the criteria for FieldNotify
+	void VerifyNotifyValueChangedFunction(const FUnrealFunctionDefinitionInfo& TargetFuncDef);
+
 	// Constructs the policy from a string
 	static FDocumentationPolicy GetDocumentationPolicyFromName(const FUHTMessageProvider& Context, const FString& PolicyName);
 
diff --git a/Engine/Source/Programs/UnrealHeaderTool/Private/NativeClassExporter.h b/Engine/Source/Programs/UnrealHeaderTool/Private/NativeClassExporter.h
--- a/Engine/Source/Programs/UnrealHeaderTool/Private/NativeClassExporter.h
+++ b/Engine/Source/Programs/UnrealHeaderTool/Private/NativeClassExporter.h
@@ -532,6 +532,18 @@ private:
 	 */
 	void ExportFunctionThunk(FUHTStringBuilder& RPCWrappers, FReferenceGatherers& OutReferenceGatherers, FUnrealFunctionDefinitionInfo& FunctionDef, const TArray<FUnrealPropertyDefinitionInfo*>& ParameterDefs, FUnrealPropertyDefinitionInfo* ReturnDef) const;
 
+	/** 
+	 * Export the declaration for FFieldNotificationClassDescriptor.
+	 * 
+	 * @param OutGeneratedHeaderText Output device for writing in the generated header file.
+	 * @param OutGeneratedCPPText Output device for writing in the generated cpp file.
+	 * @param StandardUObjectConstructorsMacroCall The destination to write standard constructor macros to.
+	 * @param EnhancedUObjectConstructorsMacroCall The destination to write enhanced constructor macros to.
+	 * @param ConstructorsMacroPrefix Prefix for constructors macro.
+	 * @param Class Class for which to export macros.
+	 */
+	void ExportFieldNotify(FOutputDevice& OutGeneratedHeaderText, FOutputDevice& OutGeneratedCPPText, FOutputDevice& StandardUObjectConstructorsMacroCall, FOutputDevice& EnhancedUObjectConstructorsMacroCall, const FString& ConstructorsMacroPrefix, FUnrealClassDefinitionInfo& ClassDef) const;
+
 	/** Exports the native function registration code for the given class. */
 	static void ExportNatives(FOutputDevice& Out, FUnrealClassDefinitionInfo& ClassDef);
 
diff --git a/Engine/Source/Programs/UnrealHeaderTool/Private/ParserHelper.h b/Engine/Source/Programs/UnrealHeaderTool/Private/ParserHelper.h
--- a/Engine/Source/Programs/UnrealHeaderTool/Private/ParserHelper.h
+++ b/Engine/Source/Programs/UnrealHeaderTool/Private/ParserHelper.h
@@ -214,12 +214,14 @@ public:
 	FString GetterName;
 	/** True if the property was marked as having a setter */
 	bool bSetterTagFound = false;
-	/** True if the property was marked as having a gatter */
+	/** True if the property was marked as having a getter */
 	bool bGetterTagFound = false;
 	/** True if setter function declaration was found while parsing class header */
 	bool bSetterFunctionFound = false;
 	/** True if getter function declaration was found while parsing class header */
 	bool bGetterFunctionFound = false;
+	/** True if the property is declared as FieldNotify. */
+	bool bFieldNotify = false;
 
 public:
 	/** @name Constructors */
@@ -419,6 +421,8 @@ struct FFuncInfo
 	bool		bSealedEvent = false;
 	/** TRUE if the function is being forced to be considered as impure by the user */
 	bool		bForceBlueprintImpure = false;
+	/** TRUE generate the entry for the FieldNotificationClassDescriptor. */
+	bool		bFieldNotify = false;
 
 	//@}
 
diff --git a/Engine/Source/Programs/UnrealHeaderTool/Private/Specifiers/ClassMetadataSpecifiers.def b/Engine/Source/Programs/UnrealHeaderTool/Private/Specifiers/ClassMetadataSpecifiers.def
--- a/Engine/Source/Programs/UnrealHeaderTool/Private/Specifiers/ClassMetadataSpecifiers.def
+++ b/Engine/Source/Programs/UnrealHeaderTool/Private/Specifiers/ClassMetadataSpecifiers.def
@@ -16,6 +16,7 @@ CLASS_METADATA_SPECIFIER(ConfigDoNotCheckDefaults)
 CLASS_METADATA_SPECIFIER(Const)
 CLASS_METADATA_SPECIFIER(ConversionRoot)
 CLASS_METADATA_SPECIFIER(CustomConstructor)
+CLASS_METADATA_SPECIFIER(CustomFieldNotify)
 CLASS_METADATA_SPECIFIER(DefaultConfig)
 CLASS_METADATA_SPECIFIER(DefaultToInstanced)
 CLASS_METADATA_SPECIFIER(DependsOn)
diff --git a/Engine/Source/Programs/UnrealHeaderTool/Private/Specifiers/FunctionSpecifiers.def b/Engine/Source/Programs/UnrealHeaderTool/Private/Specifiers/FunctionSpecifiers.def
--- a/Engine/Source/Programs/UnrealHeaderTool/Private/Specifiers/FunctionSpecifiers.def
+++ b/Engine/Source/Programs/UnrealHeaderTool/Private/Specifiers/FunctionSpecifiers.def
@@ -15,6 +15,7 @@ FUNCTION_SPECIFIER(BlueprintSetter)
 FUNCTION_SPECIFIER(Client)
 FUNCTION_SPECIFIER(CustomThunk)
 FUNCTION_SPECIFIER(Exec)
+FUNCTION_SPECIFIER(FieldNotify)
 FUNCTION_SPECIFIER(NetMulticast)
 FUNCTION_SPECIFIER(Reliable)
 FUNCTION_SPECIFIER(SealedEvent)
diff --git a/Engine/Source/Programs/UnrealHeaderTool/Private/Specifiers/VariableSpecifiers.def b/Engine/Source/Programs/UnrealHeaderTool/Private/Specifiers/VariableSpecifiers.def
--- a/Engine/Source/Programs/UnrealHeaderTool/Private/Specifiers/VariableSpecifiers.def
+++ b/Engine/Source/Programs/UnrealHeaderTool/Private/Specifiers/VariableSpecifiers.def
@@ -22,6 +22,7 @@ VARIABLE_SPECIFIER(EditFixedSize)
 VARIABLE_SPECIFIER(EditInline)
 VARIABLE_SPECIFIER(EditInstanceOnly)
 VARIABLE_SPECIFIER(Export)
+VARIABLE_SPECIFIER(FieldNotify)
 VARIABLE_SPECIFIER(Getter)
 VARIABLE_SPECIFIER(GlobalConfig)
 VARIABLE_SPECIFIER(Instanced)
diff --git a/Engine/Source/Programs/UnrealHeaderTool/Private/UnrealTypeDefinitionInfo.cpp b/Engine/Source/Programs/UnrealHeaderTool/Private/UnrealTypeDefinitionInfo.cpp
--- a/Engine/Source/Programs/UnrealHeaderTool/Private/UnrealTypeDefinitionInfo.cpp
+++ b/Engine/Source/Programs/UnrealHeaderTool/Private/UnrealTypeDefinitionInfo.cpp
@@ -2349,6 +2349,23 @@ void FUnrealClassDefinitionInfo::PostParseFinalizeInternal(EPostParseFinalizePha
 	}
 }
 
+void FUnrealClassDefinitionInfo::MarkHasFieldNotify()
+{
+	if (!bHasFieldNotify)
+	{
+		bHasFieldNotify = true;
+		const FUnrealClassDefinitionInfo* NotifyFieldInterface = GTypeDefinitionInfoMap.FindByName<FUnrealClassDefinitionInfo>(TEXT("NotifyFieldValueChanged"));
+		if (NotifyFieldInterface == nullptr)
+		{
+			LogError(TEXT("INotifyFieldValueChanged could not be found. FieldNotify are not valid for UClass %s"), *GetName());
+		}
+		if (!ImplementsInterface(*NotifyFieldInterface))
+		{
+			LogError(TEXT("UClass %s need to implement the interface INotifyFieldValueChanged to support FieldNotify."), *GetName());
+		}
+	}
+}
+
 bool FUnrealClassDefinitionInfo::ImplementsInterface(const FUnrealClassDefinitionInfo& SomeInterface) const
 {
 	if (SomeInterface.HasAnyClassFlags(CLASS_Interface) && &SomeInterface != GUInterfaceDef)
@@ -2356,7 +2373,7 @@ bool FUnrealClassDefinitionInfo::ImplementsInterface(const FUnrealClassDefinitio
 		for (const FUnrealClassDefinitionInfo* CurrentClassDef = this; CurrentClassDef; CurrentClassDef = CurrentClassDef->GetSuperClass())
 		{
 			// SomeInterface might be a base interface of our implemented interface
-			for (const FBaseStructInfo& BaseStructInfo : GetBaseStructInfos())
+			for (const FBaseStructInfo& BaseStructInfo : CurrentClassDef->GetBaseStructInfos())
 			{
 				if (const FUnrealClassDefinitionInfo* InterfaceDef = UHTCast<FUnrealClassDefinitionInfo>(BaseStructInfo.Struct))
 				{
@@ -2514,6 +2531,12 @@ void FUnrealClassDefinitionInfo::ParseClassProperties(TArray<FPropertySpecifier>
 			MarkCustomConstructor();
 			break;
 
+		case EClassMetadataSpecifier::CustomFieldNotify:
+
+			// we will not export a UE_FIELD_NOTIFICATION for this class
+			MarkCustomFieldNotify();
+			break;
+
 		case EClassMetadataSpecifier::Config:
 
 			// Class containing config properties - parse the name of the config file to use
diff --git a/Engine/Source/Programs/UnrealHeaderTool/Private/UnrealTypeDefinitionInfo.h b/Engine/Source/Programs/UnrealHeaderTool/Private/UnrealTypeDefinitionInfo.h
--- a/Engine/Source/Programs/UnrealHeaderTool/Private/UnrealTypeDefinitionInfo.h
+++ b/Engine/Source/Programs/UnrealHeaderTool/Private/UnrealTypeDefinitionInfo.h
@@ -3328,6 +3328,35 @@ public:
 		bCustomConstructor = true;
 	}
 
+	/**
+	 * Return true if the class has a custom FieldNotify declaration and implementation
+	 */
+	bool HasCustomFieldNotify() const
+	{
+		return bCustomFieldNotify;
+	}
+
+	/**
+	 * Mark the class as having a custom FieldNotify declaration and implementation
+	 */
+	void MarkCustomFieldNotify()
+	{
+		bCustomFieldNotify = true;
+	}
+
+	/**
+	 * Return true if the class has at least one FieldNotify field
+	 */
+	bool HasFieldNotify() const
+	{
+		return bHasFieldNotify;
+	}
+
+	/**
+	 * Mark the class as having at least one FieldNotify field
+	 */
+	void MarkHasFieldNotify();
+
 	/**
 	 * Return true if the class is compiled in
 	 */
@@ -3471,6 +3500,12 @@ private:
 	/** True if the class has a custom constructor */
 	bool bCustomConstructor = false;
 
+	/** True if the class has a custom FieldNotify declaration and implementation */
+	bool bCustomFieldNotify = false;
+
+	/** True if the class has at least one FProperty or UFunction that has the FieldNotify specifier */
+	bool bHasFieldNotify = false;
+
 	/** True if the class is not to be exported */
 	bool bNoExport = false;
 
diff --git a/Engine/Source/Runtime/CoreUObject/Public/UObject/ObjectMacros.h b/Engine/Source/Runtime/CoreUObject/Public/UObject/ObjectMacros.h
--- a/Engine/Source/Runtime/CoreUObject/Public/UObject/ObjectMacros.h
+++ b/Engine/Source/Runtime/CoreUObject/Public/UObject/ObjectMacros.h
@@ -734,6 +734,9 @@ namespace UC
 		/// Prevents automatic generation of the constructor declaration.
 		customConstructor,
 
+		/// Prevents automatic generation of the FieldNotify declaration.
+		CustomFieldNotify,
+
 		/// Class was declared directly in C++ and has no boilerplate generated by UnrealHeaderTool.
 		/// DO NOT USE THIS FLAG ON NEW CLASSES.
 		Intrinsic,
@@ -925,6 +928,9 @@ namespace UF
 		/// Usage: Category=CategoryName or Category="MajorCategory,SubCategory"
 		Category,
 
+		/// Generate a field entry for the NotifyFieldValueChanged interface.
+		FieldNotify,
+
 		/// This function must supply a _Validate implementation
 		WithValidation,
 
@@ -1073,6 +1079,9 @@ namespace UP
 
 		/// If true, the self pin should not be shown or connectable regardless of purity, const, etc. similar to InternalUseParam
 		HideSelfPin, 
+
+		/// Generate a field entry for the NotifyFieldValueChanged interface.
+		FieldNotify,
 	};
 }
 
