Add support for deprecation macros for enum values in UHT

diff --git a/Engine/Source/Programs/UnrealHeaderTool/Private/HeaderParser.cpp b/Engine/Source/Programs/UnrealHeaderTool/Private/HeaderParser.cpp
--- a/Engine/Source/Programs/UnrealHeaderTool/Private/HeaderParser.cpp
+++ b/Engine/Source/Programs/UnrealHeaderTool/Private/HeaderParser.cpp
@@ -1188,6 +1188,8 @@ FUnrealEnumDefinitionInfo& FHeaderParser::CompileEnum()
 	int64 CurrentEnumValue = 0;
 	while (GetIdentifier(TagToken))
 	{
+		SkipDeprecatedMacroIfNecessary(*this);
+
 		FTokenValue TagIdentifier = TagToken.GetTokenValue();
 		AddFormattedPrevCommentAsTooltipMetaData(TagMetaData);
 
