UHT: Autogenerate the Getter/Setter accessor. The setter and getter can now also be UFunction. It allow some deprecation path.

diff --git a/Engine/Source/Programs/UnrealHeaderTool/Private/BaseParser.h b/Engine/Source/Programs/UnrealHeaderTool/Private/BaseParser.h
--- a/Engine/Source/Programs/UnrealHeaderTool/Private/BaseParser.h
+++ b/Engine/Source/Programs/UnrealHeaderTool/Private/BaseParser.h
@@ -59,9 +59,6 @@ struct FUHTConfig
 	EPointerMemberBehavior NonEngineNativePointerMemberBehavior = EPointerMemberBehavior::AllowSilently;
 	EPointerMemberBehavior NonEngineObjectPtrMemberBehavior = EPointerMemberBehavior::AllowSilently;
 
-	/** If true setters and getters will be automatically (without specifying their function names on a property) parsed and generated if a function with matching signature is found */
-	bool bAllowAutomaticSettersAndGetters = false;
-
 private:
 	FUHTConfig();
 };
diff --git a/Engine/Source/Programs/UnrealHeaderTool/Private/HeaderParser.cpp b/Engine/Source/Programs/UnrealHeaderTool/Private/HeaderParser.cpp
--- a/Engine/Source/Programs/UnrealHeaderTool/Private/HeaderParser.cpp
+++ b/Engine/Source/Programs/UnrealHeaderTool/Private/HeaderParser.cpp
@@ -789,7 +789,6 @@ FUHTConfig::FUHTConfig()
 	const FName EngineObjectPtrMemberBehaviorKey(TEXT("EngineObjectPtrMemberBehavior"));
 	const FName NonEngineNativePointerMemberBehaviorKey(TEXT("NonEngineNativePointerMemberBehavior"));
 	const FName NonEngineObjectPtrMemberBehaviorKey(TEXT("NonEngineObjectPtrMemberBehavior"));
-	const FName AutomaticSettersAndGettersKey(TEXT("AutomaticSettersAndGetters"));
 
 	FConfigSection* ConfigSection = GConfig->GetSectionPrivate(TEXT("UnrealHeaderTool"), false, true, GEngineIni);
 	if (ConfigSection)
@@ -838,10 +837,6 @@ FUHTConfig::FUHTConfig()
 			{
 				NonEngineObjectPtrMemberBehavior = ToPointerMemberBehavior(It.Value().GetValue());
 			}
-			else if (It.Key() == AutomaticSettersAndGettersKey)
-			{
-				bAllowAutomaticSettersAndGetters = It.Value().GetValue().ToBool();
-			}
 		}
 	}
 }
@@ -2616,6 +2611,128 @@ void FHeaderParser::VerifyRepNotifyCallback(FUnrealPropertyDefinitionInfo& Prope
 	}
 }
 
+void FHeaderParser::VerifyGetterSetterAccessorProperties(const FUnrealClassDefinitionInfo& TargetClassDef, FUnrealPropertyDefinitionInfo& PropertyDef)
+{
+	FPropertyBase& PropertyToken = PropertyDef.GetPropertyBase();
+	if (!PropertyToken.SetterName.IsEmpty())
+	{
+		if (PropertyToken.SetterName != TEXT("None") && !PropertyToken.bSetterFunctionFound)
+		{
+			// The function can be a UFunction (not parsed the same way as the other accessors)
+			const TSharedRef<FUnrealFunctionDefinitionInfo>* FoundFunction = TargetClassDef.GetFunctions().FindByPredicate([&PropertyToken](const TSharedRef<FUnrealFunctionDefinitionInfo>& Other) { return Other->GetName() == PropertyToken.SetterName; });
+			if (FoundFunction)
+			{
+				const FUnrealFunctionDefinitionInfo& Function = FoundFunction->Get();
+				if (!Function.HasAllFunctionFlags(FUNC_Native))
+				{
+					LogError(TEXT("Property %s setter function %s has to be native."), *PropertyDef.GetName(), *PropertyToken.SetterName);
+				}
+				if (Function.HasAnyFunctionFlags(FUNC_Net | FUNC_Event))
+				{
+					LogError(TEXT("Property %s setter function %s cannot be a net or an event"), *PropertyDef.GetName(), *PropertyToken.SetterName);
+				}
+				else if (Function.HasAllFunctionFlags(FUNC_EditorOnly) != PropertyDef.HasAllPropertyFlags(CPF_EditorOnly))
+				{
+					LogError(TEXT("Property %s setter function %s must both be EditorOnly or both not EditorOnly"), *PropertyDef.GetName(), *PropertyToken.SetterName);
+				}
+				else if (Function.GetProperties().Num() != 1)
+				{
+					LogError(TEXT("Property %s setter function %s must have a single argument"), *PropertyDef.GetName(), *PropertyToken.SetterName);
+				}
+				else
+				{
+					TSharedPtr<FUnrealPropertyDefinitionInfo> FoundArgument;
+					for (const TSharedRef<FUnrealPropertyDefinitionInfo>& Arguments : Function.GetProperties())
+					{
+						if (Arguments->HasAllPropertyFlags(CPF_Parm) && !Arguments->HasAnyPropertyFlags(CPF_ReturnParm))
+						{
+							FoundArgument = Arguments;
+							break;
+						}
+					}
+					if (!FoundArgument.IsValid())
+					{
+						LogError(TEXT("Property %s setter function %s does not have a valid argument."), *PropertyDef.GetName(), *PropertyToken.SetterName);
+					}
+					else if (!FoundArgument->SameType(PropertyDef))
+					{
+						LogError(TEXT("Property %s setter function %s argument is not the same type."), *PropertyDef.GetName(), *PropertyToken.SetterName);
+					}
+					else
+					{
+						PropertyToken.bSetterFunctionFound = true;
+					}
+				}
+			}
+			if (!PropertyToken.bSetterFunctionFound)
+			{
+				LogError(TEXT("Property %s setter function %s not found"), *PropertyDef.GetName(), *PropertyToken.SetterName);
+			}
+		}
+	}
+	if (!PropertyToken.GetterName.IsEmpty())
+	{
+		if (PropertyToken.GetterName != TEXT("None") && !PropertyToken.bGetterFunctionFound)
+		{
+			// The function can be a UFunction (not parsed the same way as the other accessors)
+			const TSharedRef<FUnrealFunctionDefinitionInfo>* FoundFunction = TargetClassDef.GetFunctions().FindByPredicate([&PropertyToken](const TSharedRef<FUnrealFunctionDefinitionInfo>& Other) { return Other->GetName() == PropertyToken.GetterName; });
+			if (FoundFunction)
+			{
+				const FUnrealFunctionDefinitionInfo& Function = FoundFunction->Get();
+				if (!Function.HasAllFunctionFlags(FUNC_Native))
+				{
+					LogError(TEXT("Property %s getter function %s has to be native."), *PropertyDef.GetName(), *PropertyToken.GetterName);
+				}
+				if (!Function.HasAllFunctionFlags(FUNC_Const))
+				{
+					LogError(TEXT("Property %s getter function %s has to be const."), *PropertyDef.GetName(), *PropertyToken.GetterName);
+				}
+				if (Function.HasAnyFunctionFlags(FUNC_Net | FUNC_Event))
+				{
+					LogError(TEXT("Property %s getter function %s cannot be a net or an event"), *PropertyDef.GetName(), *PropertyToken.GetterName);
+				}
+				else if (Function.HasAllFunctionFlags(FUNC_EditorOnly) != PropertyDef.HasAllPropertyFlags(CPF_EditorOnly))
+				{
+					LogError(TEXT("Property %s getter function %s must both be EditorOnly or both not EditorOnly"), *PropertyDef.GetName(), *PropertyToken.GetterName);
+				}
+				else if (Function.GetProperties().Num() != 1)
+				{
+					LogError(TEXT("Property %s getter function %s must have a single return value"), *PropertyDef.GetName(), *PropertyToken.GetterName);
+				}
+				else
+				{
+					TSharedPtr<FUnrealPropertyDefinitionInfo> FoundArgument;
+					for (const TSharedRef<FUnrealPropertyDefinitionInfo>& Arguments : Function.GetProperties())
+					{
+						if (Arguments->HasAllPropertyFlags(CPF_Parm) && Arguments->HasAnyPropertyFlags(CPF_ReturnParm))
+						{
+							FoundArgument = Arguments;
+							break;
+						}
+					}
+					if (!FoundArgument.IsValid())
+					{
+						LogError(TEXT("Property %s getter function %s does not have a return value."), *PropertyDef.GetName(), *PropertyToken.GetterName);
+					}
+					else if (!FoundArgument->SameType(PropertyDef))
+					{
+						LogError(TEXT("Property %s getter function %s return value is not the same type."), *PropertyDef.GetName(), *PropertyToken.GetterName);
+					}
+					else
+					{
+						PropertyToken.bGetterFunctionFound = true;
+					}
+				}
+			}
+
+			if (!PropertyToken.bGetterFunctionFound)
+			{
+				LogError(TEXT("Property %s getter function %s not found"), *PropertyDef.GetName(), *PropertyToken.GetterName);
+			}
+		}
+	}
+}
+
 void FHeaderParser::VerifyNotifyValueChangedProperties(FUnrealPropertyDefinitionInfo& PropertyDef)
 {
 	FUnrealTypeDefinitionInfo* Info = PropertyDef.GetOuter();
@@ -2747,21 +2864,7 @@ void FHeaderParser::VerifyPropertyMarkups(FUnrealClassDefinitionInfo& TargetClas
 		}
 
 		// Verify if native property setter and getter functions could actually be found while parsing class header
-		FPropertyBase& PropertyToken = PropertyDef->GetPropertyBase();
-		if (!PropertyToken.SetterName.IsEmpty())
-		{
-			if (PropertyToken.SetterName != TEXT("None") && !PropertyToken.bSetterFunctionFound)
-			{
-				LogError(TEXT("Property %s setter function %s not found"), *PropertyDef->GetName(), *PropertyToken.SetterName);
-			}
-		}
-		if (!PropertyToken.GetterName.IsEmpty())
-		{
-			if (PropertyToken.GetterName != TEXT("None") && !PropertyToken.bGetterFunctionFound)
-			{
-				LogError(TEXT("Property %s getter function %s not found"), *PropertyDef->GetName(), *PropertyToken.GetterName);
-			}
-		}
+		VerifyGetterSetterAccessorProperties(TargetClassDef, *PropertyDef);
 
 		if (PropertyDef->GetPropertyBase().bFieldNotify)
 		{
@@ -5492,30 +5595,6 @@ bool FHeaderParser::CheckForPropertySetterFunction(FUnrealStructDefinitionInfo&
 		}
 	}
 
-	// If not, try with autogenerated setter names
-	if (!PropertyWithSetter && SetterNameToken.ValueStartsWith(TEXT("Set"), ESearchCase::CaseSensitive))
-	{
-		for (TSharedRef<FUnrealPropertyDefinitionInfo>& Prop : Properties)
-		{
-			FPropertyBase& BaseProp = Prop->GetPropertyBase();
-			if (!BaseProp.bSetterFunctionFound && BaseProp.SetterName.IsEmpty() &&
-				(BaseProp.bSetterTagFound || FUHTConfig::Get().bAllowAutomaticSettersAndGetters))
-			{
-				SetterFunctionName.Reset();
-				SetterFunctionName += TEXT("Set");
-				SetterFunctionName += Prop->GetName();
-
-				if (SetterNameToken.Value.Compare(SetterFunctionName, ESearchCase::CaseSensitive) == 0)
-				{
-					// We found a matching function name so we can move onto verifying if the declaration matches
-					bExplicitSetter = BaseProp.bSetterTagFound;
-					PropertyWithSetter = &Prop.Get();
-					break;
-				}
-			}
-		}
-	}
-
 	if (!PropertyWithSetter)
 	{
 		return false;
@@ -5658,28 +5737,6 @@ bool FHeaderParser::CheckForPropertyGetterFunction(FUnrealStructDefinitionInfo&
 		}
 	}
 
-	// If not, try with autogenerated getter names
-	if (!PropertyWithGetter && GetterNameToken.ValueStartsWith(TEXT("Get"), ESearchCase::CaseSensitive))
-	{
-		for (TSharedRef<FUnrealPropertyDefinitionInfo>& Prop : Properties)
-		{
-			FPropertyBase& BaseProp = Prop->GetPropertyBase();			
-			if (!BaseProp.bGetterFunctionFound && BaseProp.GetterName != TEXT("None") && 
-				(BaseProp.bGetterTagFound || FUHTConfig::Get().bAllowAutomaticSettersAndGetters))
-			{
-				GetterFunctionName.Reset();
-				GetterFunctionName += TEXT("Get");
-				GetterFunctionName += Prop->GetName();
-				if (GetterNameToken.Value.Compare(GetterFunctionName, ESearchCase::CaseSensitive) == 0)
-				{
-					bExplicitGetter = BaseProp.bGetterTagFound;
-					PropertyWithGetter = &Prop.Get();
-					break;
-				}
-			}
-		}
-	}
-
 	if (!PropertyWithGetter)
 	{
 		return false;
@@ -7749,6 +7806,16 @@ void FHeaderParser::CompileVariableDeclaration(FUnrealStructDefinitionInfo& Stru
 			}
 		}
 
+		// If the accessor tag was found but need to be autogenerated.
+		if (NewPropDef.GetPropertyBase().bSetterTagFound && NewPropDef.GetPropertyBase().SetterName.IsEmpty())
+		{
+			NewPropDef.GetPropertyBase().SetterName = FString::Printf(TEXT("Set%s"), *NewPropDef.GetName());
+		}
+		if (NewPropDef.GetPropertyBase().bGetterTagFound && NewPropDef.GetPropertyBase().GetterName.IsEmpty())
+		{
+			NewPropDef.GetPropertyBase().GetterName = FString::Printf(TEXT("Get%s"), *NewPropDef.GetName());
+		}
+
 		NewProperties.Add(&NewPropDef);
 
 		// we'll need any metadata tags we parsed later on when we call ConvertEOLCommentToTooltip() so the tags aren't clobbered
diff --git a/Engine/Source/Programs/UnrealHeaderTool/Private/HeaderParser.h b/Engine/Source/Programs/UnrealHeaderTool/Private/HeaderParser.h
--- a/Engine/Source/Programs/UnrealHeaderTool/Private/HeaderParser.h
+++ b/Engine/Source/Programs/UnrealHeaderTool/Private/HeaderParser.h
@@ -593,6 +593,9 @@ private:
 	// Verifies the target function meets the criteria for a replication notify callback
 	void VerifyRepNotifyCallback(FUnrealPropertyDefinitionInfo& PropertyDef, FUnrealFunctionDefinitionInfo* TargetFuncDef);
 
+	// Verifies the target property meets the criteria for Getter/Setter accessor
+	void VerifyGetterSetterAccessorProperties(const FUnrealClassDefinitionInfo& TargetClassDef, FUnrealPropertyDefinitionInfo& PropertyDef);
+
 	// Vefifies the target property meets the criteria for FieldNotify
 	void VerifyNotifyValueChangedProperties(FUnrealPropertyDefinitionInfo& PropertyDef);
 
