UnrealHeaderTool: Pass public rigvm context as const

diff --git a/Engine/Source/Programs/UnrealHeaderTool/Private/CodeGenerator.cpp b/Engine/Source/Programs/UnrealHeaderTool/Private/CodeGenerator.cpp
--- a/Engine/Source/Programs/UnrealHeaderTool/Private/CodeGenerator.cpp
+++ b/Engine/Source/Programs/UnrealHeaderTool/Private/CodeGenerator.cpp
@@ -3382,7 +3382,7 @@ void FNativeClassHeaderGenerator::ExportGeneratedStructBodyMacros(FOutputDevice&
 	const FString& ChoppedSingletonName = ScriptStructDef.GetSingletonNameChopped(true);
 
 	const FString RigVMExecuteContextDeclaration = TEXT("FRigVMExtendedExecuteContext& RigVMExecuteContext");
-	const FString RigVMExecuteContextPublicDeclaration = TEXT("FRigVMExecuteContext& RigVMExecuteContext");
+	const FString RigVMExecuteContextPublicDeclaration = TEXT("const FRigVMExecuteContext& RigVMExecuteContext");
 	TArray<FString> RigVMVirtualFuncProlog, RigVMVirtualFuncEpilog, RigVMStubProlog;
 
 	// for RigVM methods we need to generated a macro used for implementing the static method
