Fixing static array properties with setters and getters not generating correct wrapper code

diff --git a/Engine/Source/Programs/UnrealHeaderTool/Private/CodeGenerator.cpp b/Engine/Source/Programs/UnrealHeaderTool/Private/CodeGenerator.cpp
--- a/Engine/Source/Programs/UnrealHeaderTool/Private/CodeGenerator.cpp
+++ b/Engine/Source/Programs/UnrealHeaderTool/Private/CodeGenerator.cpp
@@ -5018,8 +5018,27 @@ void FNativeClassHeaderGenerator::ExportNativeFunctions(FOutputDevice& OutGenera
 				FString ExtendedType;
 				FString PropertyType = Prop->GetCPPType(&ExtendedType, CPPF_Implementation | CPPF_ArgumentOrReturnValue);
 				PropertyType += ExtendedType;
-				RuntimeStringBuilders.AccessorWrapperImplementations.Logf(TEXT("\t\t%s& Result = *(%s*)OutValue;\r\n"), *PropertyType, *PropertyType);
-				RuntimeStringBuilders.AccessorWrapperImplementations.Logf(TEXT("\t\tResult = (%s)Obj->%s();\r\n"), *PropertyType, *Prop->GetPropertyBase().GetterName);
+				if (Prop->GetArrayDimensions())
+				{
+					RuntimeStringBuilders.AccessorWrapperImplementations.Logf(TEXT("\t\t%s* Source = (%s*)Obj->%s();\r\n"), *PropertyType, *PropertyType, *Prop->GetPropertyBase().GetterName);
+					RuntimeStringBuilders.AccessorWrapperImplementations.Logf(TEXT("\t\t%s* Result = (%s*)OutValue;\r\n"), *PropertyType, *PropertyType);
+					RuntimeStringBuilders.AccessorWrapperImplementations.Logf(TEXT("\t\tif (TIsZeroConstructType<%s>::Value)\r\n"), *PropertyType);
+					RuntimeStringBuilders.AccessorWrapperImplementations.Logf(TEXT("\t\t{\r\n"));
+					RuntimeStringBuilders.AccessorWrapperImplementations.Logf(TEXT("\t\t\tFMemory::Memcpy(Result, Source, sizeof(%s) * %s);\r\n"), *PropertyType, Prop->GetArrayDimensions());
+					RuntimeStringBuilders.AccessorWrapperImplementations.Logf(TEXT("\t\t}\r\n"));
+					RuntimeStringBuilders.AccessorWrapperImplementations.Logf(TEXT("\t\telse\r\n"));
+					RuntimeStringBuilders.AccessorWrapperImplementations.Logf(TEXT("\t\t{\r\n"));
+					RuntimeStringBuilders.AccessorWrapperImplementations.Logf(TEXT("\t\t\tfor (int32 Index = 0; Index < %s; ++Index)\r\n"), Prop->GetArrayDimensions());
+					RuntimeStringBuilders.AccessorWrapperImplementations.Logf(TEXT("\t\t\t{\r\n"));
+					RuntimeStringBuilders.AccessorWrapperImplementations.Logf(TEXT("\t\t\t\tResult[Index] = Source[Index];\r\n"));
+					RuntimeStringBuilders.AccessorWrapperImplementations.Logf(TEXT("\t\t\t}\r\n"));
+					RuntimeStringBuilders.AccessorWrapperImplementations.Logf(TEXT("\t\t}\r\n"));
+				}
+				else
+				{
+					RuntimeStringBuilders.AccessorWrapperImplementations.Logf(TEXT("\t\t%s& Result = *(%s*)OutValue;\r\n"), *PropertyType, *PropertyType);
+					RuntimeStringBuilders.AccessorWrapperImplementations.Logf(TEXT("\t\tResult = (%s)Obj->%s();\r\n"), *PropertyType, *Prop->GetPropertyBase().GetterName);
+				}
 				if (bEditorOnlyProperty)
 				{
 					RuntimeStringBuilders.AccessorWrapperImplementations += TEXT("\t#endif // WITH_EDITORONLY_DATA\r\n");
@@ -5040,7 +5059,14 @@ void FNativeClassHeaderGenerator::ExportNativeFunctions(FOutputDevice& OutGenera
 				FString ExtendedType;
 				FString PropertyType = Prop->GetCPPType(&ExtendedType, CPPF_Implementation | CPPF_ArgumentOrReturnValue);
 				PropertyType += ExtendedType;
-				RuntimeStringBuilders.AccessorWrapperImplementations.Logf(TEXT("\t\t%s& Value = *(%s*)InValue;\r\n"), *PropertyType, *PropertyType);
+				if (Prop->GetArrayDimensions())
+				{
+					RuntimeStringBuilders.AccessorWrapperImplementations.Logf(TEXT("\t\t%s* Value = (%s*)InValue;\r\n"), *PropertyType, *PropertyType);
+				}
+				else
+				{
+					RuntimeStringBuilders.AccessorWrapperImplementations.Logf(TEXT("\t\t%s& Value = *(%s*)InValue;\r\n"), *PropertyType, *PropertyType);
+				}
 				RuntimeStringBuilders.AccessorWrapperImplementations.Logf(TEXT("\t\tObj->%s(Value);\r\n"), *Prop->GetPropertyBase().SetterName);
 				if (bEditorOnlyProperty)
 				{
diff --git a/Engine/Source/Programs/UnrealHeaderTool/Private/HeaderParser.cpp b/Engine/Source/Programs/UnrealHeaderTool/Private/HeaderParser.cpp
--- a/Engine/Source/Programs/UnrealHeaderTool/Private/HeaderParser.cpp
+++ b/Engine/Source/Programs/UnrealHeaderTool/Private/HeaderParser.cpp
@@ -5518,20 +5518,29 @@ static bool ParseAccessorType(FString& OutAccessorType, FTokenReplay& Tokens)
 		}
 		Tokens.GetToken(Token);
 	}
-	// Skip '&' and add any '*'
-	for (; !Token.IsIdentifier(); Tokens.GetToken(Token))
+	// Skip '&', 'const' and retain '*'
+	do
 	{
 		if (Token.TokenType == ETokenType::None)
 		{
 			return false;
 		}
 
+		if (Token.IsIdentifier())
+		{
+			if (!Token.IsIdentifier(TEXT("const"), ESearchCase::CaseSensitive))
+			{
+				break;
+			}
+			// else skip const
+		}
 		// Support for passing values by ref for setters
-		if (!Token.IsSymbol(TEXT("&"), ESearchCase::CaseSensitive))
+		else if (!Token.IsSymbol(TEXT("&"), ESearchCase::CaseSensitive))
 		{
 			OutAccessorType += Token.Value;
 		}
-	}
+	} while (Tokens.GetToken(Token));
+
 	// Return the last token since it's either ')' in case of setters or the getter name which will be parsed by the caller
 	Tokens.UngetToken(Token);
 
@@ -5647,6 +5656,10 @@ bool FHeaderParser::CheckForPropertySetterFunction(FUnrealStructDefinitionInfo&
 	FString PropertyType = PropertyWithSetter->GetCPPType(&ExtendedPropertyType, CPPF_Implementation | CPPF_ArgumentOrReturnValue | CPPF_NoRef);
 	ExtendedPropertyType.RemoveSpacesInline();
 	PropertyType += ExtendedPropertyType;
+	if (PropertyWithSetter->GetArrayDimensions())
+	{
+		PropertyType += TEXT("*");
+	}
 	if (ParameterType.Compare(PropertyType, ESearchCase::CaseSensitive) != 0)
 	{
 		if (bExplicitSetter)
@@ -5745,6 +5758,10 @@ bool FHeaderParser::CheckForPropertyGetterFunction(FUnrealStructDefinitionInfo&
 	FString ExtendedPropertyType;
 	FString PropertyType = PropertyWithGetter->GetCPPType(&ExtendedPropertyType, CPPF_Implementation | CPPF_ArgumentOrReturnValue);
 	PropertyType += ExtendedPropertyType;
+	if (PropertyWithGetter->GetArrayDimensions())
+	{
+		PropertyType += TEXT("*");
+	}
 	if (GetterType.Compare(PropertyType, ESearchCase::CaseSensitive) != 0)
 	{
 		if (bExplicitGetter)
