Much cleaner way of copying arrays for static array getters

diff --git a/Engine/Source/Programs/UnrealHeaderTool/Private/CodeGenerator.cpp b/Engine/Source/Programs/UnrealHeaderTool/Private/CodeGenerator.cpp
--- a/Engine/Source/Programs/UnrealHeaderTool/Private/CodeGenerator.cpp
+++ b/Engine/Source/Programs/UnrealHeaderTool/Private/CodeGenerator.cpp
@@ -5022,17 +5022,7 @@ void FNativeClassHeaderGenerator::ExportNativeFunctions(FOutputDevice& OutGenera
 				{
 					RuntimeStringBuilders.AccessorWrapperImplementations.Logf(TEXT("\t\t%s* Source = (%s*)Obj->%s();\r\n"), *PropertyType, *PropertyType, *Prop->GetPropertyBase().GetterName);
 					RuntimeStringBuilders.AccessorWrapperImplementations.Logf(TEXT("\t\t%s* Result = (%s*)OutValue;\r\n"), *PropertyType, *PropertyType);
-					RuntimeStringBuilders.AccessorWrapperImplementations.Logf(TEXT("\t\tif (TIsZeroConstructType<%s>::Value)\r\n"), *PropertyType);
-					RuntimeStringBuilders.AccessorWrapperImplementations.Logf(TEXT("\t\t{\r\n"));
-					RuntimeStringBuilders.AccessorWrapperImplementations.Logf(TEXT("\t\t\tFMemory::Memcpy(Result, Source, sizeof(%s) * %s);\r\n"), *PropertyType, Prop->GetArrayDimensions());
-					RuntimeStringBuilders.AccessorWrapperImplementations.Logf(TEXT("\t\t}\r\n"));
-					RuntimeStringBuilders.AccessorWrapperImplementations.Logf(TEXT("\t\telse\r\n"));
-					RuntimeStringBuilders.AccessorWrapperImplementations.Logf(TEXT("\t\t{\r\n"));
-					RuntimeStringBuilders.AccessorWrapperImplementations.Logf(TEXT("\t\t\tfor (int32 Index = 0; Index < %s; ++Index)\r\n"), Prop->GetArrayDimensions());
-					RuntimeStringBuilders.AccessorWrapperImplementations.Logf(TEXT("\t\t\t{\r\n"));
-					RuntimeStringBuilders.AccessorWrapperImplementations.Logf(TEXT("\t\t\t\tResult[Index] = Source[Index];\r\n"));
-					RuntimeStringBuilders.AccessorWrapperImplementations.Logf(TEXT("\t\t\t}\r\n"));
-					RuntimeStringBuilders.AccessorWrapperImplementations.Logf(TEXT("\t\t}\r\n"));
+					RuntimeStringBuilders.AccessorWrapperImplementations.Logf(TEXT("\t\tCopyAssignItems(Result, Source, %s);\r\n"), Prop->GetArrayDimensions());
 				}
 				else
 				{
