UHT: Support enum namespace as accessor type

diff --git a/Engine/Source/Programs/UnrealHeaderTool/Private/HeaderParser.cpp b/Engine/Source/Programs/UnrealHeaderTool/Private/HeaderParser.cpp
--- a/Engine/Source/Programs/UnrealHeaderTool/Private/HeaderParser.cpp
+++ b/Engine/Source/Programs/UnrealHeaderTool/Private/HeaderParser.cpp
@@ -5492,7 +5492,20 @@ static bool ParseAccessorType(FString& OutAccessorType, FTokenReplay& Tokens)
 
 	OutAccessorType = Token.Value;
 
+	// parse enum type declared as a namespace
 	Tokens.GetToken(Token);
+	while (Token.IsSymbol(TEXT("::"), ESearchCase::CaseSensitive))
+	{
+		OutAccessorType += Token.Value;
+		Tokens.GetToken(Token);
+		if (!Token.IsIdentifier())
+		{
+			return false;
+		}
+		OutAccessorType += Token.Value;
+		Tokens.GetToken(Token);
+	}
+
 	if (Token.IsSymbol(TEXT("<"), ESearchCase::CaseSensitive))
 	{
 		// template type - add everything until the matching '>' into the type string
