RigVM: Automatically mark up rig structs as blueprint type in the HeaderParser

diff --git a/Engine/Plugins/Animation/ControlRig/Source/ControlRig/Private/Units/Deprecated/RigUnit_AimConstraint.h b/Engine/Plugins/Animation/ControlRig/Source/ControlRig/Private/Units/Deprecated/RigUnit_AimConstraint.h
--- a/Engine/Plugins/Animation/ControlRig/Source/ControlRig/Private/Units/Deprecated/RigUnit_AimConstraint.h
+++ b/Engine/Plugins/Animation/ControlRig/Source/ControlRig/Private/Units/Deprecated/RigUnit_AimConstraint.h
@@ -27,7 +27,7 @@ enum class EAimMode : uint8
 	MAX,
 };
 
-USTRUCT()
+USTRUCT(BlueprintType)
 struct FAimTarget
 {
 	GENERATED_BODY()
diff --git a/Engine/Plugins/Animation/ControlRig/Source/ControlRig/Private/Units/Deprecated/RigUnit_BlendTransform.h b/Engine/Plugins/Animation/ControlRig/Source/ControlRig/Private/Units/Deprecated/RigUnit_BlendTransform.h
--- a/Engine/Plugins/Animation/ControlRig/Source/ControlRig/Private/Units/Deprecated/RigUnit_BlendTransform.h
+++ b/Engine/Plugins/Animation/ControlRig/Source/ControlRig/Private/Units/Deprecated/RigUnit_BlendTransform.h
@@ -5,7 +5,7 @@
 #include "Units/RigUnit.h"
 #include "RigUnit_BlendTransform.generated.h"
 
-USTRUCT()
+USTRUCT(BlueprintType)
 struct FBlendTarget
 {
 	GENERATED_BODY()
diff --git a/Engine/Plugins/Animation/ControlRig/Source/ControlRig/Private/Units/Highlevel/Hierarchy/RigUnit_AimBone.h b/Engine/Plugins/Animation/ControlRig/Source/ControlRig/Private/Units/Highlevel/Hierarchy/RigUnit_AimBone.h
--- a/Engine/Plugins/Animation/ControlRig/Source/ControlRig/Private/Units/Highlevel/Hierarchy/RigUnit_AimBone.h
+++ b/Engine/Plugins/Animation/ControlRig/Source/ControlRig/Private/Units/Highlevel/Hierarchy/RigUnit_AimBone.h
@@ -50,7 +50,7 @@ struct CONTROLRIG_API FRigUnit_AimBone_Target
 	FName Space;
 };
 
-USTRUCT()
+USTRUCT(BlueprintType)
 struct CONTROLRIG_API FRigUnit_AimItem_Target
 {
 	GENERATED_BODY()
@@ -95,7 +95,7 @@ struct CONTROLRIG_API FRigUnit_AimItem_Target
 	FRigElementKey Space;
 };
 
-USTRUCT()
+USTRUCT(BlueprintType)
 struct CONTROLRIG_API FRigUnit_AimBone_DebugSettings
 {
 	GENERATED_BODY()
diff --git a/Engine/Plugins/Animation/ControlRig/Source/ControlRig/Public/Math/ControlRigMathLibrary.h b/Engine/Plugins/Animation/ControlRig/Source/ControlRig/Public/Math/ControlRigMathLibrary.h
--- a/Engine/Plugins/Animation/ControlRig/Source/ControlRig/Public/Math/ControlRigMathLibrary.h
+++ b/Engine/Plugins/Animation/ControlRig/Source/ControlRig/Public/Math/ControlRigMathLibrary.h
@@ -42,7 +42,7 @@ enum class EControlRigAnimEasingType : uint8
 	BounceEaseInOut
 };
 
-USTRUCT()
+USTRUCT(BlueprintType)
 struct FCRFourPointBezier
 {
 	GENERATED_BODY()
@@ -52,16 +52,16 @@ struct FCRFourPointBezier
 		A = B = C = D = FVector::ZeroVector;
 	}
 
-	UPROPERTY()
+	UPROPERTY(EditAnywhere, Category=Bezier)
 	FVector A;
 
-	UPROPERTY()
+	UPROPERTY(EditAnywhere, Category=Bezier)
 	FVector B;
 
-	UPROPERTY()
+	UPROPERTY(EditAnywhere, Category=Bezier)
 	FVector C;
 
-	UPROPERTY()
+	UPROPERTY(EditAnywhere, Category=Bezier)
 	FVector D;
 };
 
diff --git a/Engine/Plugins/Animation/ControlRig/Source/ControlRig/Public/Math/Simulation/CRSimLinearSpring.h b/Engine/Plugins/Animation/ControlRig/Source/ControlRig/Public/Math/Simulation/CRSimLinearSpring.h
--- a/Engine/Plugins/Animation/ControlRig/Source/ControlRig/Public/Math/Simulation/CRSimLinearSpring.h
+++ b/Engine/Plugins/Animation/ControlRig/Source/ControlRig/Public/Math/Simulation/CRSimLinearSpring.h
@@ -6,7 +6,7 @@
 #include "CRSimPoint.h"
 #include "CRSimLinearSpring.generated.h"
 
-USTRUCT()
+USTRUCT(BlueprintType)
 struct FCRSimLinearSpring
 {
 	GENERATED_BODY()
@@ -21,19 +21,19 @@ struct FCRSimLinearSpring
 	/**
 	 * The first point affected by this spring
 	 */
-	UPROPERTY()
+	UPROPERTY(EditAnywhere, Category=Simulation)
 	int32 SubjectA;
 
 	/**
 	 * The second point affected by this spring
 	 */
-	UPROPERTY()
+	UPROPERTY(EditAnywhere, Category=Simulation)
 	int32 SubjectB;
 
 	/**
 	 * The power of this spring
 	 */
-	UPROPERTY()
+	UPROPERTY(EditAnywhere, Category=Simulation)
 	float Coefficient;
 
 	/**
@@ -41,7 +41,7 @@ struct FCRSimLinearSpring
 	 * A value of lower than zero indicates that the equilibrium
 	 * should be based on the current distance of the two subjects.
 	 */
-	UPROPERTY()
+	UPROPERTY(EditAnywhere, Category=Simulation)
 	float Equilibrium;
 
 	void CalculateForPoints(const FCRSimPoint& InPointA, const FCRSimPoint& InPointB, FVector& ForceA, FVector& ForceB) const;
diff --git a/Engine/Plugins/Animation/ControlRig/Source/ControlRig/Public/Math/Simulation/CRSimPoint.h b/Engine/Plugins/Animation/ControlRig/Source/ControlRig/Public/Math/Simulation/CRSimPoint.h
--- a/Engine/Plugins/Animation/ControlRig/Source/ControlRig/Public/Math/Simulation/CRSimPoint.h
+++ b/Engine/Plugins/Animation/ControlRig/Source/ControlRig/Public/Math/Simulation/CRSimPoint.h
@@ -12,7 +12,7 @@ enum class ECRSimPointIntegrateType : uint8
 	SemiExplicitEuler
 };
 
-USTRUCT()
+USTRUCT(BlueprintType)
 struct FCRSimPoint
 {
 	GENERATED_BODY()
@@ -29,19 +29,19 @@ struct FCRSimPoint
 	/**
 	 * The mass of the point
 	 */
-	UPROPERTY()
+	UPROPERTY(EditAnywhere, Category=Simulation)
 	float Mass;
 
 	/**
 	 * Size of the point - only used for collision
 	 */
-	UPROPERTY()
+	UPROPERTY(EditAnywhere, Category=Simulation)
 	float Size;
 
 	/**
 	 * The linear damping of the point
 	 */
-	UPROPERTY()
+	UPROPERTY(EditAnywhere, Category=Simulation)
 	float LinearDamping;
 
 	/**
@@ -49,19 +49,19 @@ struct FCRSimPoint
 	 * This does not have an effect on passive (mass == 0.0) points.
 	 * Values can be higher than 1 due to timestep - but they are clamped internally.
 	 */
-	UPROPERTY()
+	UPROPERTY(EditAnywhere, Category=Simulation)
 	float InheritMotion;
 
 	/**
 	 * The position of the point
 	 */
-	UPROPERTY()
+	UPROPERTY(EditAnywhere, Category=Simulation)
 	FVector Position;
 
 	/**
 	 * The velocity of the point per second
 	 */
-	UPROPERTY()
+	UPROPERTY(EditAnywhere, Category=Simulation)
 	FVector LinearVelocity;
 
 	FCRSimPoint IntegrateVerlet(const FVector& InForce, float InBlend, float InDeltaTime) const;
diff --git a/Engine/Plugins/Animation/ControlRig/Source/ControlRig/Public/Math/Simulation/CRSimPointForce.h b/Engine/Plugins/Animation/ControlRig/Source/ControlRig/Public/Math/Simulation/CRSimPointForce.h
--- a/Engine/Plugins/Animation/ControlRig/Source/ControlRig/Public/Math/Simulation/CRSimPointForce.h
+++ b/Engine/Plugins/Animation/ControlRig/Source/ControlRig/Public/Math/Simulation/CRSimPointForce.h
@@ -12,7 +12,7 @@ enum class ECRSimPointForceType : uint8
 	Direction
 };
 
-USTRUCT()
+USTRUCT(BlueprintType)
 struct FCRSimPointForce
 {
 	GENERATED_BODY()
@@ -28,7 +28,7 @@ struct FCRSimPointForce
 	/**
 	 * The type of force.
 	 */
-	UPROPERTY()
+	UPROPERTY(EditAnywhere, Category=Simulation)
 	ECRSimPointForceType ForceType;
 
 	/**
@@ -36,19 +36,19 @@ struct FCRSimPointForce
 	 * This is a direction for direction based forces,
 	 * while this is a position for attractor / repel based forces.
 	 */
-	UPROPERTY()
+	UPROPERTY(EditAnywhere, Category=Simulation)
 	FVector Vector;
 
 	/**
 	 * The strength of the force (a multiplier for direction based forces)
 	 */
-	UPROPERTY()
+	UPROPERTY(EditAnywhere, Category=Simulation)
 	float Coefficient;
 
 	/**
 	 * If set to true the input vector will be normalized.
 	 */
-	UPROPERTY()
+	UPROPERTY(EditAnywhere, Category=Simulation)
 	bool bNormalize;
 
 	FVector Calculate(const FCRSimPoint& InPoint, float InDeltaTime) const;
diff --git a/Engine/Plugins/Animation/ControlRig/Source/ControlRig/Public/Math/Simulation/CRSimSoftCollision.h b/Engine/Plugins/Animation/ControlRig/Source/ControlRig/Public/Math/Simulation/CRSimSoftCollision.h
--- a/Engine/Plugins/Animation/ControlRig/Source/ControlRig/Public/Math/Simulation/CRSimSoftCollision.h
+++ b/Engine/Plugins/Animation/ControlRig/Source/ControlRig/Public/Math/Simulation/CRSimSoftCollision.h
@@ -15,7 +15,7 @@ enum class ECRSimSoftCollisionType : uint8
 	Cone
 };
 
-USTRUCT()
+USTRUCT(BlueprintType)
 struct FCRSimSoftCollision
 {
 	GENERATED_BODY()
@@ -34,13 +34,13 @@ struct FCRSimSoftCollision
 	/**
 	 * The world / global transform of the collisoin
 	 */
-	UPROPERTY()
+	UPROPERTY(EditAnywhere, Category=Simulation)
 	FTransform Transform;
 
 	/**
 	 * The type of collision shape
 	 */
-	UPROPERTY()
+	UPROPERTY(EditAnywhere, Category=Simulation)
 	ECRSimSoftCollisionType ShapeType;
 
 	/**
@@ -48,7 +48,7 @@ struct FCRSimSoftCollision
 	 * If this is equal or higher than the maximum there's no falloff.
 	 * For a cone shape this represents the minimum angle in degrees.
 	 */
-	UPROPERTY()
+	UPROPERTY(EditAnywhere, Category=Simulation)
 	float MinimumDistance;
 
 	/**
@@ -56,25 +56,25 @@ struct FCRSimSoftCollision
 	 * If this is equal or lower than the minimum there's no falloff.
 	 * For a cone shape this represents the maximum angle in degrees.
 	 */
-	UPROPERTY()
+	UPROPERTY(EditAnywhere, Category=Simulation)
 	float MaximumDistance;
 
 	/**
 	 * The type of falloff to use
 	 */
-	UPROPERTY()
+	UPROPERTY(EditAnywhere, Category=Simulation)
 	EControlRigAnimEasingType FalloffType;
 
 	/**
 	 * The strength of the collision force
 	 */
-	UPROPERTY()
+	UPROPERTY(EditAnywhere, Category=Simulation)
 	float Coefficient;
 
 	/**
 	 * If set to true the collision volume will be inverted
 	 */
-	UPROPERTY()
+	UPROPERTY(EditAnywhere, Category=Simulation)
 	bool bInverted;
 
 	static float CalculateFalloff(const FCRSimSoftCollision& InCollision, const FVector& InPosition, float InSize, FVector& OutDirection);
diff --git a/Engine/Plugins/Experimental/FullBodyIK/Source/FullBodyIK/Public/FBIKConstraintOption.h b/Engine/Plugins/Experimental/FullBodyIK/Source/FullBodyIK/Public/FBIKConstraintOption.h
--- a/Engine/Plugins/Experimental/FullBodyIK/Source/FullBodyIK/Public/FBIKConstraintOption.h
+++ b/Engine/Plugins/Experimental/FullBodyIK/Source/FullBodyIK/Public/FBIKConstraintOption.h
@@ -120,7 +120,7 @@ struct FFBIKConstraintOption
 	}
 };
 
-USTRUCT()
+USTRUCT(BlueprintType)
 struct FMotionProcessInput
 {
 	GENERATED_BODY()
diff --git a/Engine/Plugins/Experimental/FullBodyIK/Source/FullBodyIK/Public/FBIKDebugOption.h b/Engine/Plugins/Experimental/FullBodyIK/Source/FullBodyIK/Public/FBIKDebugOption.h
--- a/Engine/Plugins/Experimental/FullBodyIK/Source/FullBodyIK/Public/FBIKDebugOption.h
+++ b/Engine/Plugins/Experimental/FullBodyIK/Source/FullBodyIK/Public/FBIKDebugOption.h
@@ -5,7 +5,7 @@
 #include "CoreMinimal.h"
 #include "FBIKDebugOption.generated.h"
 
-USTRUCT()
+USTRUCT(BlueprintType)
 struct FFBIKDebugOption
 {
 	GENERATED_BODY()
diff --git a/Engine/Plugins/Experimental/FullBodyIK/Source/FullBodyIK/Public/FBIKShared.h b/Engine/Plugins/Experimental/FullBodyIK/Source/FullBodyIK/Public/FBIKShared.h
--- a/Engine/Plugins/Experimental/FullBodyIK/Source/FullBodyIK/Public/FBIKShared.h
+++ b/Engine/Plugins/Experimental/FullBodyIK/Source/FullBodyIK/Public/FBIKShared.h
@@ -17,7 +17,7 @@ private:
 	virtual void PreSolve(TArray<FFBIKLinkData>& InOutLinkData, const TMap<int32, FFBIKEffectorTarget>& InEndEffectors) const override;
 };
 
-USTRUCT()
+USTRUCT(BlueprintType)
 struct FSolverInput
 {
 	GENERATED_BODY()
diff --git a/Engine/Plugins/Experimental/FullBodyIK/Source/PBIK/Public/Core/PBIKSolver.h b/Engine/Plugins/Experimental/FullBodyIK/Source/PBIK/Public/Core/PBIKSolver.h
--- a/Engine/Plugins/Experimental/FullBodyIK/Source/PBIK/Public/Core/PBIKSolver.h
+++ b/Engine/Plugins/Experimental/FullBodyIK/Source/PBIK/Public/Core/PBIKSolver.h
@@ -99,7 +99,7 @@ enum class EPBIKRootBehavior : uint8
 	Free,
 };
 
-USTRUCT()
+USTRUCT(BlueprintType)
 struct PBIK_API FPBIKSolverSettings
 {
 	GENERATED_BODY()
diff --git a/Engine/Plugins/Experimental/FullBodyIK/Source/PBIK/Public/RigUnit_PBIK.h b/Engine/Plugins/Experimental/FullBodyIK/Source/PBIK/Public/RigUnit_PBIK.h
--- a/Engine/Plugins/Experimental/FullBodyIK/Source/PBIK/Public/RigUnit_PBIK.h
+++ b/Engine/Plugins/Experimental/FullBodyIK/Source/PBIK/Public/RigUnit_PBIK.h
@@ -14,7 +14,7 @@
 
 using PBIK::FDebugLine;
 
-USTRUCT()
+USTRUCT(BlueprintType)
 struct FPBIKDebug
 {
 	GENERATED_BODY()
diff --git a/Engine/Source/Programs/UnrealHeaderTool/Private/HeaderParser.cpp b/Engine/Source/Programs/UnrealHeaderTool/Private/HeaderParser.cpp
--- a/Engine/Source/Programs/UnrealHeaderTool/Private/HeaderParser.cpp
+++ b/Engine/Source/Programs/UnrealHeaderTool/Private/HeaderParser.cpp
@@ -4685,6 +4685,49 @@ void FHeaderParser::GetVarType(
 	{
 		ParsedVarIndexRange->Count = InputPos - ParsedVarIndexRange->StartIndex;
 	}
+
+	// Setup additional property as well as script struct def flags
+	// for structs / properties being used for the RigVM.
+	// The Input / Output / Constant metadata tokens can be used to mark
+	// up an input / output pin of a RigVMNode. To allow authoring of those
+	// nodes we'll mark up the property as accessible in Blueprint / Python
+	// as well as make the struct a blueprint type.
+	if(OwnerScriptStructDef)
+	{
+		const EPropertyFlags OriginalFlags = VarProperty.PropertyFlags; 
+		if(VarProperty.MetaData.Contains(NAME_ConstantText))
+		{
+			VarProperty.PropertyFlags |= CPF_Edit | CPF_EditConst | CPF_BlueprintVisible;
+		}
+		if(VarProperty.MetaData.Contains(NAME_InputText) ||
+			VarProperty.MetaData.Contains(NAME_VisibleText))
+		{
+			VarProperty.PropertyFlags |= CPF_Edit | CPF_BlueprintVisible;
+		}
+		if(VarProperty.MetaData.Contains(NAME_OutputText))
+		{
+			if((VarProperty.PropertyFlags & CPF_BlueprintVisible) == 0)
+			{
+				VarProperty.PropertyFlags |= CPF_BlueprintVisible | CPF_BlueprintReadOnly;
+			}
+		}
+		
+		if(OriginalFlags != VarProperty.PropertyFlags &&
+			(((VarProperty.PropertyFlags & CPF_BlueprintVisible) != 0) ||
+			((VarProperty.PropertyFlags & CPF_BlueprintReadOnly) != 0)))
+		{
+			if(!OwnerScriptStructDef->GetBoolMetaDataHierarchical(FHeaderParserNames::NAME_BlueprintType))
+			{
+				OwnerScriptStructDef->SetMetaData(FHeaderParserNames::NAME_BlueprintType, TEXT("true"));
+			}
+			
+			if(!VarProperty.MetaData.Contains(NAME_Category))
+			{
+				static constexpr TCHAR PinsCategory[] = TEXT("Pins");
+				VarProperty.MetaData.Add(NAME_Category, PinsCategory);
+			}
+		}
+	}
 }
 
 FUnrealPropertyDefinitionInfo& FHeaderParser::GetVarNameAndDim
diff --git a/Engine/Source/Programs/UnrealHeaderTool/Resources/UHTDebugging/RigVMStructHeader.h b/Engine/Source/Programs/UnrealHeaderTool/Resources/UHTDebugging/RigVMStructHeader.h
--- a/Engine/Source/Programs/UnrealHeaderTool/Resources/UHTDebugging/RigVMStructHeader.h
+++ b/Engine/Source/Programs/UnrealHeaderTool/Resources/UHTDebugging/RigVMStructHeader.h
@@ -104,3 +104,18 @@ struct FRigVMMethodStruct : public FRigVMStructBase
 	RIGVM_METHOD()
 	virtual FRigVMStructUpgradeInfo GetUpgradeInfo() const override;
 };
+
+USTRUCT(meta=(DisplayName="Control", Category="Controls", ShowVariableNameInTitle, Deprecated = "4.24.0"))
+struct CONTROLRIG_API FRigUnit_Control
+{
+	GENERATED_BODY()
+
+	FRigUnit_Control()
+		: Factor(0)
+	{
+	}
+
+	/** The transform of this control */
+	UPROPERTY(EditAnywhere, Category="Control")
+	float Factor;
+};
\ No newline at end of file
diff --git a/Engine/Source/Runtime/RigVM/Public/RigVMCore/RigVMExecuteContext.h b/Engine/Source/Runtime/RigVM/Public/RigVMCore/RigVMExecuteContext.h
--- a/Engine/Source/Runtime/RigVM/Public/RigVMCore/RigVMExecuteContext.h
+++ b/Engine/Source/Runtime/RigVM/Public/RigVMCore/RigVMExecuteContext.h
@@ -192,7 +192,7 @@ struct RIGVM_API FRigVMRuntimeSettings
  * The execute context is used for mutable nodes to
  * indicate execution order.
  */
-USTRUCT()
+USTRUCT(BlueprintType)
 struct RIGVM_API FRigVMExecuteContext
 {
 	GENERATED_BODY()
