Fix UHT static analysis warnings
#rnx

diff --git a/Engine/Source/Programs/UnrealHeaderTool/Private/CodeGenerator.cpp b/Engine/Source/Programs/UnrealHeaderTool/Private/CodeGenerator.cpp
--- a/Engine/Source/Programs/UnrealHeaderTool/Private/CodeGenerator.cpp
+++ b/Engine/Source/Programs/UnrealHeaderTool/Private/CodeGenerator.cpp
@@ -2906,7 +2906,7 @@ void FNativeClassHeaderGenerator::ExportClassFromSourceFileInner(
 
 			// the name of the pure interface class
 			FString InterfaceCPPName = ClassDef.GetAlternateNameCPP(true);
-			FString SuperInterfaceCPPName = SuperClassDef ? SuperClassDef->GetAlternateNameCPP(true) : FString();
+			FString SuperInterfaceCPPName = SuperClassDef->GetAlternateNameCPP(true);
 
 			// Thunk functions
 			FUHTStringBuilder InterfaceBoilerplate;
diff --git a/Engine/Source/Programs/UnrealHeaderTool/Private/UnrealTypeDefinitionInfo.cpp b/Engine/Source/Programs/UnrealHeaderTool/Private/UnrealTypeDefinitionInfo.cpp
--- a/Engine/Source/Programs/UnrealHeaderTool/Private/UnrealTypeDefinitionInfo.cpp
+++ b/Engine/Source/Programs/UnrealHeaderTool/Private/UnrealTypeDefinitionInfo.cpp
@@ -1404,6 +1404,7 @@ void FUnrealFieldDefinitionInfo::PostParseFinalizeReferencedProperties()
 		{
 			OuterPropDef = NextPropDef;
 		}
+		checkSlow(OuterPropDef);
 		FPropertyTraits::PostParseFinalize(*OuterPropDef);
 	}
 }
@@ -2359,7 +2360,7 @@ void FUnrealClassDefinitionInfo::MarkHasFieldNotify()
 		{
 			LogError(TEXT("INotifyFieldValueChanged could not be found. FieldNotify are not valid for UClass %s"), *GetName());
 		}
-		if (!ImplementsInterface(*NotifyFieldInterface))
+		else if (!ImplementsInterface(*NotifyFieldInterface))
 		{
 			LogError(TEXT("UClass %s need to implement the interface INotifyFieldValueChanged to support FieldNotify."), *GetName());
 		}
