Fix issue where BP callable functions on native interfaces were incorrectly detected as invalid
This is fallout from fixing UE-99628, this check was missed when adding NotBlueprintable as a proper keyword

diff --git a/Engine/Source/Programs/UnrealHeaderTool/Private/HeaderParser.cpp b/Engine/Source/Programs/UnrealHeaderTool/Private/HeaderParser.cpp
--- a/Engine/Source/Programs/UnrealHeaderTool/Private/HeaderParser.cpp
+++ b/Engine/Source/Programs/UnrealHeaderTool/Private/HeaderParser.cpp
@@ -9464,6 +9464,10 @@ void FHeaderParser::PostPopNestClass(FUnrealClassDefinitionInfo& CurrentClassDef
 				continue;
 			}
 
+			const bool bCanImplementInBlueprints = InterfaceDef->GetBoolMetaData(NAME_IsBlueprintBase);
+			const bool bCannotImplementInBlueprints = (!bCanImplementInBlueprints && InterfaceDef->HasMetaData(NAME_IsBlueprintBase))
+				|| InterfaceDef->HasMetaData(NAME_CannotImplementInterfaceInBlueprint);
+
 			// So iterate over all functions this interface declares
 			for (FUnrealFunctionDefinitionInfo* InterfaceFunctionDef : TUHTFieldRange<FUnrealFunctionDefinitionInfo>(*InterfaceDef, EFieldIteratorFlags::ExcludeSuper))
 			{
@@ -9520,10 +9524,11 @@ void FHeaderParser::PostPopNestClass(FUnrealClassDefinitionInfo& CurrentClassDef
 				}
 
 				// Verify that if this has blueprint-callable functions that are not implementable events, we've implemented them as a UFunction in the target class
+				// This is only relevant for bp-implementable interfaces, for native interfaces the interface-defined function is sufficient
 				if (!bImplemented
 					&& InterfaceFunctionDef->HasAnyFunctionFlags(FUNC_BlueprintCallable)
 					&& !InterfaceFunctionDef->HasAnyFunctionFlags(FUNC_BlueprintEvent)
-					&& !InterfaceDef->HasMetaData(NAME_CannotImplementInterfaceInBlueprint))  // FBlueprintMetadata::MD_CannotImplementInterfaceInBlueprint
+					&& !bCannotImplementInBlueprints)
 				{
 					Throwf(TEXT("Missing UFunction implementation of function '%s' from interface '%s'.  This function needs a UFUNCTION() declaration."), *InterfaceFunctionDef->GetName(), *InterfaceDef->GetName());
 				}
