Removing debug code

diff --git a/Engine/Source/Programs/UnrealHeaderTool/Private/HeaderParser.cpp b/Engine/Source/Programs/UnrealHeaderTool/Private/HeaderParser.cpp
--- a/Engine/Source/Programs/UnrealHeaderTool/Private/HeaderParser.cpp
+++ b/Engine/Source/Programs/UnrealHeaderTool/Private/HeaderParser.cpp
@@ -5633,12 +5633,6 @@ bool FHeaderParser::CheckForPropertySetterFunction(FUnrealStructDefinitionInfo&
 		return false;
 	}
 
-	if (SetterNameToken.Value.Compare(TEXT("SetHardwareCursors"), ESearchCase::CaseSensitive) == 0)
-	{
-		static volatile int32 xx = 0;
-		xx++;
-	}
-
 	FUnrealPropertyDefinitionInfo* PropertyWithSetter = nullptr;
 	FString SetterFunctionName;
 	TArray<TSharedRef<FUnrealPropertyDefinitionInfo>>& Properties = ClassDef.GetProperties();
