Expand the check for optional class not to allow editor only properties but also prevent struct properties from containing editor only properties

diff --git a/Engine/Source/Programs/UnrealHeaderTool/Private/Exceptions.cpp b/Engine/Source/Programs/UnrealHeaderTool/Private/Exceptions.cpp
--- a/Engine/Source/Programs/UnrealHeaderTool/Private/Exceptions.cpp
+++ b/Engine/Source/Programs/UnrealHeaderTool/Private/Exceptions.cpp
@@ -97,6 +97,25 @@
 
 		FResults::MarkWarning();
 	}
+
+	void LogInfoInternal(const FString& Filename, int32 Line, const FString& Message)
+	{
+		TGuardValue<ELogTimes::Type> DisableLogTimes(GPrintLogTimes, ELogTimes::None);
+
+		FString FormattedErrorMessage;
+		if (Filename.IsEmpty())
+		{
+			FormattedErrorMessage = FString::Printf(TEXT("Info: %s\r\n"), *Message);
+		}
+		else
+		{
+			FormattedErrorMessage = FString::Printf(TEXT("%s(%d): Info: %s\r\n"), *Filename, Line, *Message);
+		}
+
+		UE_LOG(LogCompile, Log, TEXT("%s"), *FormattedErrorMessage);
+		GWarn->Log(ELogVerbosity::Display, FormattedErrorMessage);
+	}
+
 }
 
 void FResults::LogError(const FUHTMessageProvider& Context, const TCHAR* ErrorMsg, ECompilationResult::Type InResult)
@@ -142,6 +161,31 @@
 	}
 }
 
+void FResults::LogInfo(const FUHTMessageProvider& Context, const TCHAR* InfoMsg)
+{
+	using namespace UE::UnrealHeaderTool::Exceptions::Private;
+
+	FString Filename = IFileManager::Get().ConvertToAbsolutePathForExternalAppForRead(*Context.GetFilename());
+	int32 Line = Context.GetLineNumber();
+
+	if (IsInGameThread())
+	{
+		LogInfoInternal(Filename, Line, InfoMsg);
+	}
+	else
+	{
+		auto LogInfoTask = [Filename = MoveTemp(Filename), Line, Message = FString(InfoMsg)]()
+		{
+			LogInfoInternal(Filename, Line, Message);
+		};
+
+		FGraphEventRef EventRef = FFunctionGraphTask::CreateAndDispatchWhenReady(MoveTemp(LogInfoTask), TStatId(), nullptr, ENamedThreads::GameThread);
+
+		FScopeLock Lock(&ErrorTasksCS);
+		ErrorTasks.Add(EventRef);
+	}
+}
+
 void FResults::WaitForErrorTasks()
 {
 	using namespace UE::UnrealHeaderTool::Exceptions::Private;
diff --git a/Engine/Source/Programs/UnrealHeaderTool/Private/Exceptions.h b/Engine/Source/Programs/UnrealHeaderTool/Private/Exceptions.h
--- a/Engine/Source/Programs/UnrealHeaderTool/Private/Exceptions.h
+++ b/Engine/Source/Programs/UnrealHeaderTool/Private/Exceptions.h
@@ -132,6 +132,13 @@
 	static void LogWarning(const FUHTMessageProvider& Context, const TCHAR* ErrorMsg);
 
 	/**
+	 * Log an info for the given source file where the type is defined
+	 * @param Context The context of the warning
+	 * @param InfoMsg The text of the information
+	 */
+	static void LogInfo(const FUHTMessageProvider& Context, const TCHAR* InfoMsg);
+
+	/**
 	 * Invoke the given lambda in a try block catching all supported exception types.
 	 * @param InLambda The code to be executed in the try block
 	 */
@@ -316,6 +323,16 @@
 		FString ResultString = FString::Printf(InFmt, InArgs ...);
 		FResults::LogWarning(*this, *ResultString);
 	}
+
+	/**
+	 * Log an info
+	 * @param InText The text of the info
+	 */
+	void LogInfo(const FString& InText) const
+	{
+		FResults::LogInfo(*this, *InText);
+	}
+
 };
 
 template <typename T>
diff --git a/Engine/Source/Programs/UnrealHeaderTool/Private/HeaderParser.cpp b/Engine/Source/Programs/UnrealHeaderTool/Private/HeaderParser.cpp
--- a/Engine/Source/Programs/UnrealHeaderTool/Private/HeaderParser.cpp
+++ b/Engine/Source/Programs/UnrealHeaderTool/Private/HeaderParser.cpp
@@ -3122,15 +3122,6 @@
 		}
 	}
 
-	// Validate if we are using editor only data in a class or struct definition
-	if ((Flags & CPF_EditorOnly) != 0)
-	{
-		if (OwnerClassDef && OwnerClassDef->HasAnyClassFlags(CLASS_Optional))
-		{
-			LogError(TEXT("Cannot specify an editor only property inside an optional class."));
-		}
-	}
-
 	// Store the start and end positions of the parsed type
 	if (ParsedVarIndexRange)
 	{
@@ -4624,6 +4615,20 @@
 		}
 	}
 
+	// Validate if we are using editor only data in a class or struct definition
+	if (OwnerClassDef && OwnerClassDef->HasAnyClassFlags(CLASS_Optional))
+	{
+		if (VarProperty.IsEditorOnlyProperty())
+		{
+			LogError(TEXT("Cannot specify editor only property inside an optional class."));
+		}
+		else if (VarProperty.ContainsEditorOnlyProperties())
+		{
+			// TODO: this should technically be an error, but some code already relies on this at this time and should hence 
+			LogInfo(TEXT("Do not specify struct property containing editor only properties inside an optional class."));
+		}
+	}
+
 	if ( VarProperty.IsObjectOrInterface() && VarProperty.Type != CPT_SoftObjectReference && VarProperty.MetaClassDef == nullptr && (VarProperty.PropertyFlags&CPF_Config) != 0 )
 	{
 		Throwf(TEXT("Not allowed to use 'config' with object variables"));
diff --git a/Engine/Source/Programs/UnrealHeaderTool/Private/ParserHelper.cpp b/Engine/Source/Programs/UnrealHeaderTool/Private/ParserHelper.cpp
--- a/Engine/Source/Programs/UnrealHeaderTool/Private/ParserHelper.cpp
+++ b/Engine/Source/Programs/UnrealHeaderTool/Private/ParserHelper.cpp
@@ -111,6 +111,22 @@
 	return Enum->GetCppForm() != UEnum::ECppForm::EnumClass;
 }
 
+bool FPropertyBase::ContainsEditorOnlyProperties() const
+{
+	if (Type == CPT_Struct)
+	{
+		check(ScriptStructDef);
+		for (const TSharedRef<FUnrealPropertyDefinitionInfo>& PropDef : ScriptStructDef->GetProperties())
+		{
+			if (PropDef->IsEditorOnlyProperty() || PropDef->GetPropertyBase().ContainsEditorOnlyProperties())
+			{
+				return true;
+			}
+		}
+	}
+	return false;
+}
+
 bool FPropertyBase::MatchesType(const FPropertyBase& Other, bool bDisallowGeneralization, bool bIgnoreImplementedInterfaces/* = false*/, bool bEmulateSameType/* = false*/) const
 {
 	check(Type != CPT_None || !bDisallowGeneralization);
diff --git a/Engine/Source/Programs/UnrealHeaderTool/Private/ParserHelper.h b/Engine/Source/Programs/UnrealHeaderTool/Private/ParserHelper.h
--- a/Engine/Source/Programs/UnrealHeaderTool/Private/ParserHelper.h
+++ b/Engine/Source/Programs/UnrealHeaderTool/Private/ParserHelper.h
@@ -329,6 +329,13 @@
 		return Type == CPT_Delegate && IsPrimitiveOrPrimitiveStaticArray();
 	}
 
+	bool IsEditorOnlyProperty() const
+	{
+		return (PropertyFlags & CPF_EditorOnly) != 0;
+	}
+
+	bool ContainsEditorOnlyProperties() const;
+
 	bool PassCPPArgsByRef() const
 	{
 		if (ArrayType == EArrayType::Dynamic || ArrayType == EArrayType::Set || MapKeyProp.IsValid())
