Fixed threading issue with editor only check on optional classes.  Other objects can not be inspected until parsing has been completed.

#rnx

diff --git a/Engine/Source/Programs/UnrealHeaderTool/Private/HeaderParser.cpp b/Engine/Source/Programs/UnrealHeaderTool/Private/HeaderParser.cpp
--- a/Engine/Source/Programs/UnrealHeaderTool/Private/HeaderParser.cpp
+++ b/Engine/Source/Programs/UnrealHeaderTool/Private/HeaderParser.cpp
@@ -4615,20 +4615,6 @@
 		}
 	}
 
-	// Validate if we are using editor only data in a class or struct definition
-	if (OwnerClassDef && OwnerClassDef->HasAnyClassFlags(CLASS_Optional))
-	{
-		if (VarProperty.IsEditorOnlyProperty())
-		{
-			LogError(TEXT("Cannot specify editor only property inside an optional class."));
-		}
-		else if (VarProperty.ContainsEditorOnlyProperties())
-		{
-			// TODO: this should technically be an error, but some code already relies on this at this time and should hence 
-			LogInfo(TEXT("Do not specify struct property containing editor only properties inside an optional class."));
-		}
-	}
-
 	if ( VarProperty.IsObjectOrInterface() && VarProperty.Type != CPT_SoftObjectReference && VarProperty.MetaClassDef == nullptr && (VarProperty.PropertyFlags&CPF_Config) != 0 )
 	{
 		Throwf(TEXT("Not allowed to use 'config' with object variables"));
diff --git a/Engine/Source/Programs/UnrealHeaderTool/Private/UnrealTypeDefinitionInfo.cpp b/Engine/Source/Programs/UnrealHeaderTool/Private/UnrealTypeDefinitionInfo.cpp
--- a/Engine/Source/Programs/UnrealHeaderTool/Private/UnrealTypeDefinitionInfo.cpp
+++ b/Engine/Source/Programs/UnrealHeaderTool/Private/UnrealTypeDefinitionInfo.cpp
@@ -2345,6 +2345,23 @@ void FUnrealClassDefinitionInfo::PostParseFinalizeInternal(EPostParseFinalizePha
 				LogError(TEXT("Class cast flags in core class '%s', is '0x%016x', should be '0x%016x'"), *GetNameCPP(), ClassCastFlags, Class->ClassCastFlags);
 			}
 		}
+
+		// Validate if we are using editor only data in a class or struct definition
+		if (HasAnyClassFlags(CLASS_Optional))
+		{
+			for (TSharedRef<FUnrealPropertyDefinitionInfo> PropertyDef : GetProperties())
+			{
+				if (PropertyDef->GetPropertyBase().IsEditorOnlyProperty())
+				{
+					PropertyDef->LogError(TEXT("Cannot specify editor only property inside an optional class."));
+				}
+				else if (PropertyDef->GetPropertyBase().ContainsEditorOnlyProperties())
+				{
+					// TODO: this should technically be an error, but some code already relies on this at this time and should hence 
+					PropertyDef->LogInfo(TEXT("Do not specify struct property containing editor only properties inside an optional class."));
+				}
+			}
+		}
 		break;
 	}
 	}
