Fixed issue with UHT script exporters where the signature function wasn't being set on multicast delegate properties.

diff --git a/Engine/Source/Programs/UnrealHeaderTool/Private/PropertyTypes.cpp b/Engine/Source/Programs/UnrealHeaderTool/Private/PropertyTypes.cpp
--- a/Engine/Source/Programs/UnrealHeaderTool/Private/PropertyTypes.cpp
+++ b/Engine/Source/Programs/UnrealHeaderTool/Private/PropertyTypes.cpp
@@ -2377,11 +2377,15 @@ struct FPropertyTypeTraitsMulticastDelegate : public FPropertyTypeTraitsBase
 
 		if (VarProperty.FunctionDef->GetFunctionType() == EFunctionType::SparseDelegate)
 		{
-			return new FMulticastSparseDelegateProperty(Scope, Name, ObjectFlags);
+			FMulticastSparseDelegateProperty* Result = new FMulticastSparseDelegateProperty(Scope, Name, ObjectFlags);
+			Result->SignatureFunction = VarProperty.FunctionDef->GetFunction();
+			return Result;
 		}
 		else
 		{
-			return new FMulticastInlineDelegateProperty(Scope, Name, ObjectFlags);
+			FMulticastInlineDelegateProperty* Result = new FMulticastInlineDelegateProperty(Scope, Name, ObjectFlags);
+			Result->SignatureFunction = VarProperty.FunctionDef->GetFunction();
+			return Result;
 		}
 	}
 
