Put optional class check back to proper errors now that all misbehaving classes that leaked through are fixed

diff --git a/Engine/Source/Programs/UnrealHeaderTool/Private/UnrealTypeDefinitionInfo.cpp b/Engine/Source/Programs/UnrealHeaderTool/Private/UnrealTypeDefinitionInfo.cpp
--- a/Engine/Source/Programs/UnrealHeaderTool/Private/UnrealTypeDefinitionInfo.cpp
+++ b/Engine/Source/Programs/UnrealHeaderTool/Private/UnrealTypeDefinitionInfo.cpp
@@ -2362,8 +2362,7 @@ void FUnrealClassDefinitionInfo::PostParseFinalizeInternal(EPostParseFinalizePha
 				}
 				else if (PropertyDef->GetPropertyBase().ContainsEditorOnlyProperties())
 				{
-					// TODO: this should technically be an error, but some code already relies on this at this time and should hence 
-					PropertyDef->LogInfo(TEXT("Do not specify struct property containing editor only properties inside an optional class."));
+					PropertyDef->LogError(TEXT("Do not specify struct property containing editor only properties inside an optional class."));
 				}
 			}
 		}
