Fix for UHT failing to add enum meta data to a fixed array when a space appears in the array size.

#9337

diff --git a/Engine/Source/Programs/Shared/EpicGames.UHT/Types/UhtProperty.cs b/Engine/Source/Programs/Shared/EpicGames.UHT/Types/UhtProperty.cs
--- a/Engine/Source/Programs/Shared/EpicGames.UHT/Types/UhtProperty.cs
+++ b/Engine/Source/Programs/Shared/EpicGames.UHT/Types/UhtProperty.cs
@@ -1435,8 +1435,11 @@ namespace EpicGames.UHT.Types
 						{
 							again = false;
 
+							// Remove any irrelevant whitespace
+							dim = dim.Trim();
+
 							// Remove any outer brackets
-							if (dim[0] == '(')
+							if (dim.Length > 0 && dim[0] == '(')
 							{
 								for (int index = 1, depth = 1; index < dim.Length; ++index)
 								{
@@ -1506,7 +1509,7 @@ namespace EpicGames.UHT.Types
 							UhtEnum? enumObj = this.Session.FindRegularEnumValue(dim.ToString());
 							if (enumObj == null)
 							{
-								enumObj = this.Session.FindType(this.Outer, UhtFindOptions.Enum | UhtFindOptions.SourceName, this.ArrayDimensions) as UhtEnum;
+								enumObj = this.Session.FindType(this.Outer, UhtFindOptions.Enum | UhtFindOptions.SourceName, dim.ToString()) as UhtEnum;
 							}
 							if (enumObj != null)
 							{
diff --git a/Engine/Source/Programs/UnrealHeaderTool/Private/PropertyTypes.cpp b/Engine/Source/Programs/UnrealHeaderTool/Private/PropertyTypes.cpp
--- a/Engine/Source/Programs/UnrealHeaderTool/Private/PropertyTypes.cpp
+++ b/Engine/Source/Programs/UnrealHeaderTool/Private/PropertyTypes.cpp
@@ -2556,9 +2556,12 @@ struct FPropertyTypeTraitsStaticArray : public FPropertyTypeTraitsBase
 					TEXT("(unsigned int)"),
 					TEXT("(signed int)")
 				};
+				
+				// Remove any irrelevant whitespace
+				Temp.TrimStartAndEndInline();
 
 				// Remove any brackets
-				if (Temp[0] == TEXT('('))
+				if (Temp.Len() > 0 && Temp[0] == TEXT('('))
 				{
 					int32 TempLen = Temp.Len();
 					int32 ClosingParen = FindMatchingClosingParenthesis(Temp);
@@ -2586,7 +2589,7 @@ struct FPropertyTypeTraitsStaticArray : public FPropertyTypeTraitsBase
 				if (!EnumDef)
 				{
 					// If the enum wasn't declared in this scope, then try to find it anywhere we can
-					EnumDef = GTypeDefinitionInfoMap.FindByName<FUnrealEnumDefinitionInfo>(PropDef.GetArrayDimensions());
+					EnumDef = GTypeDefinitionInfoMap.FindByName<FUnrealEnumDefinitionInfo>(*Temp);
 				}
 
 				if (EnumDef)
