Updated UHT to output all the circular dependencies before exiting.

diff --git a/Engine/Source/Programs/UnrealHeaderTool/Private/CodeGenerator.cpp b/Engine/Source/Programs/UnrealHeaderTool/Private/CodeGenerator.cpp
--- a/Engine/Source/Programs/UnrealHeaderTool/Private/CodeGenerator.cpp
+++ b/Engine/Source/Programs/UnrealHeaderTool/Private/CodeGenerator.cpp
@@ -6919,6 +6919,7 @@ void TopologicalSort(TArray<FUnrealSourceFile*>& OrderedSourceFiles)
 		SourceFile->SetTopologicalState(ETopologicalState::Unmarked);
 	}
 
+	bool bCircularDependencyDetected = false;
 	for (FUnrealSourceFile* SourceFile : UnorderedSourceFiles)
 	{
 		if (SourceFile->GetTopologicalState() == ETopologicalState::Unmarked)
@@ -6928,11 +6929,16 @@ void TopologicalSort(TArray<FUnrealSourceFile*>& OrderedSourceFiles)
 				UE_LOG(LogCompile, Error, TEXT("Circular dependency detected:"));
 				TopologicalRecursion(*Recusion, *Recusion);
 				FResults::SetResult(ECompilationResult::OtherCompilationError);
-				return;
+				bCircularDependencyDetected = true;
 			}
 		}
 	}
 
+	if (bCircularDependencyDetected)
+	{
+		return;
+	}
+
 	for (int32 Index = 0, EIndex = OrderedSourceFiles.Num(); Index != EIndex; ++Index)
 	{
 		OrderedSourceFiles[Index]->SetOrderedIndex(Index);
