Fixed comment to read TEnumAsByte instead of TByteAsEnum

#rnx

diff --git a/Engine/Source/Programs/Shared/EpicGames.UHT/Exporters/CodeGen/UhtHeaderCodeGeneratorCppFile.cs b/Engine/Source/Programs/Shared/EpicGames.UHT/Exporters/CodeGen/UhtHeaderCodeGeneratorCppFile.cs
--- a/Engine/Source/Programs/Shared/EpicGames.UHT/Exporters/CodeGen/UhtHeaderCodeGeneratorCppFile.cs
+++ b/Engine/Source/Programs/Shared/EpicGames.UHT/Exporters/CodeGen/UhtHeaderCodeGeneratorCppFile.cs
@@ -1092,8 +1092,8 @@ namespace EpicGames.UHT.Exporters.CodeGen
 						}
 						else if (property is UhtByteProperty byteProperty && byteProperty.Enum != null)
 						{
-							// If someone passed in a TByteAsEnum instead of the actual enum value, the cast in the else clause would cause an issue.
-							// Since this is known to be a TByteAsEnum, we just fetch the first byte.  *HOWEVER* on MSB machines where 
+							// If someone passed in a TEnumAsByte instead of the actual enum value, the cast in the else clause would cause an issue.
+							// Since this is known to be a TEnumAsByte, we just fetch the first byte.  *HOWEVER* on MSB machines where 
 							// the actual enum value is passed in this will fail and return zero if the native size of the enum > 1 byte.
 							builder
 								.Append("\t\t")
diff --git a/Engine/Source/Programs/UnrealHeaderTool/Private/CodeGenerator.cpp b/Engine/Source/Programs/UnrealHeaderTool/Private/CodeGenerator.cpp
--- a/Engine/Source/Programs/UnrealHeaderTool/Private/CodeGenerator.cpp
+++ b/Engine/Source/Programs/UnrealHeaderTool/Private/CodeGenerator.cpp
@@ -5099,8 +5099,8 @@ void FNativeClassHeaderGenerator::ExportNativeFunctions(FOutputDevice& OutGenera
 				}
 				else if (Prop->IsByteEnumOrByteEnumStaticArray() && Prop->GetPropertyBase().ArrayType == EArrayType::None)
 				{
-					// If someone passed in a TByteAsEnum instead of the actual enum value, the cast in the else clause would cause an issue.
-					// Since this is known to be a TByteAsEnum, we just fetch the first byte.  *HOWEVER* on MSB machines where 
+					// If someone passed in a TEnumAsByte instead of the actual enum value, the cast in the else clause would cause an issue.
+					// Since this is known to be a TEnumAsByte, we just fetch the first byte.  *HOWEVER* on MSB machines where 
 					// the actual enum value is passed in this will fail and return zero if the native size of the enum > 1 byte.
 					RuntimeStringBuilders.AccessorWrapperImplementations.Logf(TEXT("\t\t%s Value = (%s)*(uint8*)InValue;\r\n"), *PropertyType, *PropertyType);
 				}
