Fixed an issue where the generated cpp files weren't including all the needed headers. It relied on the original header that generated the cpp to include all the headers which is a problem because the header could forward declare the types.

diff --git a/Engine/Source/Programs/Shared/EpicGames.UHT/Exporters/CodeGen/UhtHeaderCodeGeneratorCppFile.cs b/Engine/Source/Programs/Shared/EpicGames.UHT/Exporters/CodeGen/UhtHeaderCodeGeneratorCppFile.cs
--- a/Engine/Source/Programs/Shared/EpicGames.UHT/Exporters/CodeGen/UhtHeaderCodeGeneratorCppFile.cs
+++ b/Engine/Source/Programs/Shared/EpicGames.UHT/Exporters/CodeGen/UhtHeaderCodeGeneratorCppFile.cs
@@ -38,24 +38,45 @@ namespace EpicGames.UHT.Exporters.CodeGen
 				builder.Append(HeaderCopyright);
 				builder.Append(RequiredCPPIncludes);
 				builder.Append("#include \"").Append(headerInfo.IncludePath).Append("\"\r\n");
-				if (headerInfo.NeedsFastArrayHeaders)
-				{
-					builder.Append("#include \"Net/Serialization/FastArraySerializerImplementation.h\"\r\n");
-				}
 
 				bool addedStructuredArchiveFromArchiveHeader = false;
 				bool addedArchiveUObjectFromStructuredArchiveHeader = false;
 				HashSet<UhtHeaderFile> addedIncludes = new();
-				addedIncludes.Add(HeaderFile);
-				foreach (UhtType type in HeaderFile.Children)
+				List<string> includesToAdd = new();
+				addedIncludes.Add(this.HeaderFile);
+
+				if (headerInfo.NeedsFastArrayHeaders)
 				{
+					includesToAdd.Add("Net/Serialization/FastArraySerializerImplementation.h");
+				}
+
+				foreach (UhtType type in this.HeaderFile.Children)
+				{
+					if (type is UhtStruct structObj)
+					{
+						// Functions
+						foreach (UhtFunction function in structObj.Functions)
+						{
+							foreach (UhtProperty property in function.Properties)
+							{
+								AddIncludeForProperty(property, addedIncludes, includesToAdd);
+							}
+						}
+
+						// Properties
+						foreach (UhtProperty property in structObj.Properties)
+						{
+							AddIncludeForProperty(property, addedIncludes, includesToAdd);
+						}
+					}
+
 					if (type is UhtClass classObj)
 					{
 						if (classObj.ClassWithin != Session.UObject && !classObj.ClassWithin.HeaderFile.IsNoExportTypes)
 						{
 							if (addedIncludes.Add(classObj.ClassWithin.HeaderFile))
 							{
-								builder.Append("#include \"").Append(HeaderInfos[classObj.ClassWithin.HeaderFile.HeaderFileTypeIndex].IncludePath).Append("\"\r\n");
+								includesToAdd.Add(HeaderInfos[classObj.ClassWithin.HeaderFile.HeaderFileTypeIndex].IncludePath);
 							}
 						}
 
@@ -67,7 +88,7 @@ namespace EpicGames.UHT.Exporters.CodeGen
 							case UhtSerializerArchiveType.Archive:
 								if (!addedArchiveUObjectFromStructuredArchiveHeader)
 								{
-									builder.Append("#include \"Serialization/ArchiveUObjectFromStructuredArchive.h\"\r\n");
+									includesToAdd.Add("Serialization/ArchiveUObjectFromStructuredArchive.h");
 									addedArchiveUObjectFromStructuredArchiveHeader = true;
 								}
 								break;
@@ -75,12 +96,25 @@ namespace EpicGames.UHT.Exporters.CodeGen
 							case UhtSerializerArchiveType.StructuredArchiveRecord:
 								if (!addedStructuredArchiveFromArchiveHeader)
 								{
-									builder.Append("#include \"Serialization/StructuredArchive.h\"\r\n");
+									includesToAdd.Add("Serialization/StructuredArchive.h");
 									addedStructuredArchiveFromArchiveHeader = true;
 								}
 								break;
 						}
 					}
+					else
+					{
+						if (!type.HeaderFile.IsNoExportTypes && addedIncludes.Add(type.HeaderFile))
+						{
+							includesToAdd.Add(this.HeaderInfos[type.HeaderFile.HeaderFileTypeIndex].IncludePath);
+						}
+					}
+				}
+
+				includesToAdd.Sort(StringComparerUE.OrdinalIgnoreCase);
+				foreach (string include in includesToAdd)
+				{
+					builder.Append("#include \"").Append(include).Append("\"\r\n");
 				}
 
 				builder.Append(DisableDeprecationWarnings).Append("\r\n");
@@ -307,6 +341,33 @@ namespace EpicGames.UHT.Exporters.CodeGen
 			}
 		}
 
+		private void AddIncludeForType(UhtProperty uhtProperty, HashSet<UhtHeaderFile> addedIncludes, IList<string> includesToAdd)
+		{
+			if (uhtProperty is UhtStructProperty structProperty)
+			{
+				UhtScriptStruct scriptStruct = structProperty.ScriptStruct;
+				if (!scriptStruct.HeaderFile.IsNoExportTypes && addedIncludes.Add(scriptStruct.HeaderFile))
+				{
+					includesToAdd.Add(HeaderInfos[scriptStruct.HeaderFile.HeaderFileTypeIndex].IncludePath);
+				}
+			}
+		}
+
+		private void AddIncludeForProperty(UhtProperty property, HashSet<UhtHeaderFile> addedIncludes, IList<string> includesToAdd)
+		{
+			AddIncludeForType(property, addedIncludes, includesToAdd);
+
+			if (property is UhtContainerBaseProperty containerProperty)
+			{
+				AddIncludeForType(containerProperty.ValueProperty, addedIncludes, includesToAdd);
+			}
+
+			if (property is UhtMapProperty mapProperty)
+			{
+				AddIncludeForType(mapProperty.KeyProperty, addedIncludes, includesToAdd);
+			}
+		}
+
 		private StringBuilder AppendEnum(StringBuilder builder, UhtEnum enumObj)
 		{
 			const string MetaDataParamsName = "Enum_MetaDataParams";
diff --git a/Engine/Source/Programs/UnrealHeaderTool/Private/CodeGenerator.cpp b/Engine/Source/Programs/UnrealHeaderTool/Private/CodeGenerator.cpp
--- a/Engine/Source/Programs/UnrealHeaderTool/Private/CodeGenerator.cpp
+++ b/Engine/Source/Programs/UnrealHeaderTool/Private/CodeGenerator.cpp
@@ -3199,7 +3199,7 @@ void ExportEnhancedConstructorsMacro(FOutputDevice& Out, FUnrealClassDefinitionI
  *
  * @returns Inclusion path.
  */
-FString GetBuildPath(FUnrealSourceFile& SourceFile)
+FString GetBuildPath(const FUnrealSourceFile& SourceFile)
 {
 	FString Out = SourceFile.GetFilename();
 
@@ -5466,16 +5466,34 @@ void FNativeClassHeaderGenerator::ApplyAlternatePropertyExportText(FUnrealProper
 	}
 }
 
+static void AddIncludeForType(const FUnrealTypeDefinitionInfo* PropertyTypeDef, TArray<FString>& RelativeIncludes)
+{
+	if (UHTCast<FUnrealScriptStructDefinitionInfo>(PropertyTypeDef) != nullptr)
+	{
+		if (PropertyTypeDef->HasSource() && !PropertyTypeDef->GetUnrealSourceFile().IsNoExportTypes())
+		{
+			FString Header = GetBuildPath(PropertyTypeDef->GetUnrealSourceFile());
+			RelativeIncludes.AddUnique(MoveTemp(Header));
+		}
+	}
+}
+
+static void AddIncludeForProperty(const TSharedRef<FUnrealPropertyDefinitionInfo>& PropertyDef, TArray<FString>& RelativeIncludes)
+{
+	const FPropertyBase& PropertyBase = PropertyDef->GetPropertyBase();
+	AddIncludeForType(PropertyBase.TypeDef, RelativeIncludes);
+	if (PropertyBase.MapKeyProp != nullptr)
+	{
+		AddIncludeForType(PropertyBase.MapKeyProp->TypeDef, RelativeIncludes);
+	}
+}
+
 bool FNativeClassHeaderGenerator::WriteSource(const FManifestModule& Module, FGeneratedFileInfo& FileInfo, const FString& InBodyText, FUnrealSourceFile* InSourceFile, const TSet<FString>& InCrossModuleReferences, const EExportClassOutFlags& ExportFlags)
 {
 	// Collect the includes if this is from a source file
 	TArray<FString> RelativeIncludes;
 	if (InSourceFile)
 	{
-		FString ModuleRelativeFilename = InSourceFile->GetFilename();
-		ConvertToBuildIncludePath(Module, ModuleRelativeFilename);
-		RelativeIncludes.Add(MoveTemp(ModuleRelativeFilename));
-
 		if (EnumHasAnyFlags(ExportFlags, EExportClassOutFlags::NeedsFastArrayHeaders))
 		{
 			RelativeIncludes.Add(FString(TEXT("Net/Serialization/FastArraySerializerImplementation.h")));
@@ -5483,28 +5501,60 @@ bool FNativeClassHeaderGenerator::WriteSource(const FManifestModule& Module, FGe
 
 		bool bAddedStructuredArchiveFromArchiveHeader = false;
 		bool bAddedArchiveUObjectFromStructuredArchiveHeader = false;
-		for (const TSharedRef<FUnrealTypeDefinitionInfo>& TypeDef : InSourceFile->GetDefinedClasses())
+
+		// We need to include the headers that declare any types that aren't pointers.
+		// We can't rely on the original header that generates this cpp to include the 
+		// headers because it could be forward declaring the types.
+		for (const TSharedRef<FUnrealTypeDefinitionInfo>& TypeDef : InSourceFile->GetDefinedTypes())
 		{
-			FUnrealClassDefinitionInfo& ClassDef = TypeDef->AsClassChecked();
-			FUnrealClassDefinitionInfo* ClassWithin = ClassDef.GetClassWithin();
-			if (ClassWithin && ClassWithin->HasSource() && !ClassWithin->GetUnrealSourceFile().IsNoExportTypes())
+			if (FUnrealStructDefinitionInfo* Struct = TypeDef->AsStruct())
 			{
-				FString Header = GetBuildPath(ClassWithin->GetUnrealSourceFile());
-				RelativeIncludes.AddUnique(MoveTemp(Header));
-			}
+				// Functions
+				for (const TSharedRef<FUnrealStructDefinitionInfo>& FunctionDef : Struct->GetFunctions())
+				{
+					for (const TSharedRef<FUnrealPropertyDefinitionInfo>& PropertyDef : FunctionDef->GetProperties())
+					{
+						AddIncludeForProperty(PropertyDef, RelativeIncludes);
+					}
+				}
 
-			if (!bAddedStructuredArchiveFromArchiveHeader && ClassDef.GetArchiveType() == ESerializerArchiveType::StructuredArchiveRecord)
-			{
-				RelativeIncludes.AddUnique(TEXT("Serialization/StructuredArchive.h"));
-				bAddedStructuredArchiveFromArchiveHeader = true;
+				// Properties
+				for (const TSharedRef<FUnrealPropertyDefinitionInfo>& PropertyDef : Struct->GetProperties())
+				{
+					AddIncludeForProperty(PropertyDef, RelativeIncludes);
+				}
 			}
 
-			if (!bAddedArchiveUObjectFromStructuredArchiveHeader && ClassDef.GetArchiveType() == ESerializerArchiveType::Archive)
+			if (FUnrealClassDefinitionInfo* ClassDef = TypeDef->AsClass())
 			{
-				RelativeIncludes.AddUnique(TEXT("Serialization/ArchiveUObjectFromStructuredArchive.h"));
-				bAddedArchiveUObjectFromStructuredArchiveHeader = true;
+				FUnrealClassDefinitionInfo* ClassWithin = ClassDef->GetClassWithin();
+				if (ClassWithin && ClassWithin->HasSource() && !ClassWithin->GetUnrealSourceFile().IsNoExportTypes())
+				{
+					FString Header = GetBuildPath(ClassWithin->GetUnrealSourceFile());
+					RelativeIncludes.AddUnique(MoveTemp(Header));
+				}
+
+				if (!bAddedStructuredArchiveFromArchiveHeader && ClassDef->GetArchiveType() == ESerializerArchiveType::StructuredArchiveRecord)
+				{
+					RelativeIncludes.AddUnique(TEXT("Serialization/StructuredArchive.h"));
+					bAddedStructuredArchiveFromArchiveHeader = true;
+				}
+
+				if (!bAddedArchiveUObjectFromStructuredArchiveHeader && ClassDef->GetArchiveType() == ESerializerArchiveType::Archive)
+				{
+					RelativeIncludes.AddUnique(TEXT("Serialization/ArchiveUObjectFromStructuredArchive.h"));
+					bAddedArchiveUObjectFromStructuredArchiveHeader = true;
+				}
 			}
 		}
+
+		RelativeIncludes.Sort();
+
+		// Add the source header file to the top of the list
+		FString ModuleRelativeFilename = InSourceFile->GetFilename();
+		ConvertToBuildIncludePath(Module, ModuleRelativeFilename);
+		RelativeIncludes.Remove(ModuleRelativeFilename);
+		RelativeIncludes.EmplaceAt(0, MoveTemp(ModuleRelativeFilename));
 	}
 
 	FUHTStringBuilder FileText;
