Fixed Linux compile issue.

diff --git a/Engine/Source/Programs/UnrealHeaderTool/Private/CodeGenerator.cpp b/Engine/Source/Programs/UnrealHeaderTool/Private/CodeGenerator.cpp
--- a/Engine/Source/Programs/UnrealHeaderTool/Private/CodeGenerator.cpp
+++ b/Engine/Source/Programs/UnrealHeaderTool/Private/CodeGenerator.cpp
@@ -5510,16 +5510,16 @@ bool FNativeClassHeaderGenerator::WriteSource(const FManifestModule& Module, FGe
 			if (FUnrealStructDefinitionInfo* Struct = TypeDef->AsStruct())
 			{
 				// Functions
-				for (const TSharedRef<FUnrealStructDefinitionInfo>& FunctionDef : Struct->GetFunctions())
+				for (TSharedRef<FUnrealStructDefinitionInfo> FunctionDef : Struct->GetFunctions())
 				{
-					for (const TSharedRef<FUnrealPropertyDefinitionInfo>& PropertyDef : FunctionDef->GetProperties())
+					for (TSharedRef<FUnrealPropertyDefinitionInfo> PropertyDef : FunctionDef->GetProperties())
 					{
 						AddIncludeForProperty(PropertyDef, RelativeIncludes);
 					}
 				}
 
 				// Properties
-				for (const TSharedRef<FUnrealPropertyDefinitionInfo>& PropertyDef : Struct->GetProperties())
+				for (TSharedRef<FUnrealPropertyDefinitionInfo> PropertyDef : Struct->GetProperties())
 				{
 					AddIncludeForProperty(PropertyDef, RelativeIncludes);
 				}
