Remove special case handling of engine plugins in UHT for ObjectPtr behavior.

diff --git a/Engine/Source/Programs/Shared/EpicGames.UHT/Types/Properties/UhtObjectPropertyBase.cs b/Engine/Source/Programs/Shared/EpicGames.UHT/Types/Properties/UhtObjectPropertyBase.cs
--- a/Engine/Source/Programs/Shared/EpicGames.UHT/Types/Properties/UhtObjectPropertyBase.cs
+++ b/Engine/Source/Programs/Shared/EpicGames.UHT/Types/Properties/UhtObjectPropertyBase.cs
@@ -241,7 +241,7 @@ namespace EpicGames.UHT.Types
 			}
 
 			UhtPackage package = propertySettings.Outer.Package;
-			UhtPointerMemberBehavior behavior = package.IsPartOfEngine && !package.IsPlugin ? engineBehavior : nonEngineBehavior;
+			UhtPointerMemberBehavior behavior = package.IsPartOfEngine ? engineBehavior : nonEngineBehavior;
 
 			if (behavior == UhtPointerMemberBehavior.AllowSilently)
 			{
diff --git a/Engine/Source/Programs/UnrealHeaderTool/Private/HeaderParser.cpp b/Engine/Source/Programs/UnrealHeaderTool/Private/HeaderParser.cpp
--- a/Engine/Source/Programs/UnrealHeaderTool/Private/HeaderParser.cpp
+++ b/Engine/Source/Programs/UnrealHeaderTool/Private/HeaderParser.cpp
@@ -4419,7 +4419,7 @@ void FHeaderParser::GetVarType(
 						if (VariableCategory == EVariableCategory::Member)
 						{
 							// TODO: Remove exclusion for plugins under engine when all plugins have had their raw pointers converted.
-							ConditionalLogPointerUsage(bIsCurrentModulePartOfEngine && !PackageDef.GetModule().BaseDirectory.Contains(TEXT("/Plugins/")) ? UHTConfig.EngineNativePointerMemberBehavior : UHTConfig.NonEngineNativePointerMemberBehavior,
+							ConditionalLogPointerUsage(bIsCurrentModulePartOfEngine ? UHTConfig.EngineNativePointerMemberBehavior : UHTConfig.NonEngineNativePointerMemberBehavior,
 								TEXT("Native pointer"), FString(InputPos - VarStartPos, Input + VarStartPos).TrimStartAndEnd().ReplaceCharWithEscapedChar(), TEXT("TObjectPtr"));
 
 							MatchIdentifier(TEXT("const"), ESearchCase::CaseSensitive);
@@ -4430,7 +4430,7 @@ void FHeaderParser::GetVarType(
 					else if ((PropertyType == CPT_ObjectPtrReference) && (VariableCategory == EVariableCategory::Member))
 					{
 						// TODO: Remove exclusion for plugins under engine when all plugins have had their raw pointers converted.
-						ConditionalLogPointerUsage(bIsCurrentModulePartOfEngine && !PackageDef.GetModule().BaseDirectory.Contains(TEXT("/Plugins/")) ? UHTConfig.EngineObjectPtrMemberBehavior : UHTConfig.NonEngineObjectPtrMemberBehavior,
+						ConditionalLogPointerUsage(bIsCurrentModulePartOfEngine ? UHTConfig.EngineObjectPtrMemberBehavior : UHTConfig.NonEngineObjectPtrMemberBehavior,
 							TEXT("ObjectPtr"), FString(InputPos - VarStartPos, Input + VarStartPos).TrimStartAndEnd().ReplaceCharWithEscapedChar(), nullptr);
 					}
 
