Fix header tool error when class name contains only a prefix.

#9509

diff --git a/Engine/Source/Programs/UnrealHeaderTool/Private/HeaderParser.cpp b/Engine/Source/Programs/UnrealHeaderTool/Private/HeaderParser.cpp
--- a/Engine/Source/Programs/UnrealHeaderTool/Private/HeaderParser.cpp
+++ b/Engine/Source/Programs/UnrealHeaderTool/Private/HeaderParser.cpp
@@ -8948,7 +8948,7 @@ TSharedRef<FUnrealTypeDefinitionInfo> FHeaderPreParser::ParseClassDeclaration(co
 
 	if (ClassNameWithoutPrefixStr.IsEmpty())
 	{
-		Throwf(TEXT("When compiling class definition for '%s', attempting to strip prefix results in an empty name. Did you leave off a prefix?"), *ClassNameWithoutPrefixStr);
+		Throwf(TEXT("When compiling class definition for '%s', attempting to strip prefix results in an empty name. Did you leave off a prefix?"), *ClassName);
 	}
 
 	// Skip optional final keyword
