Added support for final keyword on structure definitions

#8993

diff --git a/Engine/Source/Programs/Shared/EpicGames.UHT/Parsers/UhtScriptStructParser.cs b/Engine/Source/Programs/Shared/EpicGames.UHT/Parsers/UhtScriptStructParser.cs
--- a/Engine/Source/Programs/Shared/EpicGames.UHT/Parsers/UhtScriptStructParser.cs
+++ b/Engine/Source/Programs/Shared/EpicGames.UHT/Parsers/UhtScriptStructParser.cs
@@ -106,6 +106,9 @@ namespace EpicGames.UHT.Parsers
 						topScope.TokenReader.LogError($"When compiling struct definition for '{scriptStruct.SourceName}', attempting to strip prefix results in an empty name. Did you leave off a prefix?");
 					}
 
+					// Skip optional final keyword
+					topScope.TokenReader.Optional("final");
+
 					// Parse the inheritance
 					UhtParserHelpers.ParseInheritance(topScope.TokenReader, topScope.Session.Config!, out UhtToken superIdentifier, out List<UhtToken[]>? baseIdentifiers);
 					scriptStruct.SuperIdentifier = superIdentifier;
diff --git a/Engine/Source/Programs/UnrealHeaderTool/Private/HeaderParser.cpp b/Engine/Source/Programs/UnrealHeaderTool/Private/HeaderParser.cpp
--- a/Engine/Source/Programs/UnrealHeaderTool/Private/HeaderParser.cpp
+++ b/Engine/Source/Programs/UnrealHeaderTool/Private/HeaderParser.cpp
@@ -1798,6 +1798,9 @@ FUnrealScriptStructDefinitionInfo& FHeaderParser::CompileStructDeclaration()
 		}
 	}
 
+	// Skip optional final keyword
+	MatchIdentifier(TEXT("final"), ESearchCase::CaseSensitive);
+
 	// Parse the inheritance list
 	ParseInheritance(TEXT("struct"), [](const TCHAR* StructName, bool bIsSuperClass) {}); // Eat the results, already been parsed
 
@@ -9098,6 +9101,8 @@ TSharedRef<FUnrealTypeDefinitionInfo> FHeaderPreParser::ParseStructDeclaration(c
 		Throwf(TEXT("When compiling struct definition for '%s', attempting to strip prefix results in an empty name. Did you leave off a prefix?"), *StructNameInScript);
 	}
 
+	// Skip optional final keyword
+	MatchIdentifier(TEXT("final"), ESearchCase::CaseSensitive);
 
 	// Create the structure definition
 	TSharedRef<FUnrealScriptStructDefinitionInfo> StructDef = MakeShareable(new FUnrealScriptStructDefinitionInfo(SourceFile, InLineNumber, *StructNameInScript, FName(StructNameStripped, FNAME_Add)));
