* UnrealHeaderTool - Fixed so destructor is declared and defined for code gen classes
* UnrealTypeDefinitionInfo - Explicitly added destructor to make sure code gen is not adding.
* Niagra - Fixed bug where function was declared as non-inline virtual but implemented as inline.
* RecastNavmesh - Made sure destructor is outside #ifdef for UHT to more easily parse the header and generate proper code

diff --git a/Engine/Plugins/FX/Niagara/Source/Niagara/Private/NiagaraDataInterfaceVolumeCache.cpp b/Engine/Plugins/FX/Niagara/Source/Niagara/Private/NiagaraDataInterfaceVolumeCache.cpp
--- a/Engine/Plugins/FX/Niagara/Source/Niagara/Private/NiagaraDataInterfaceVolumeCache.cpp
+++ b/Engine/Plugins/FX/Niagara/Source/Niagara/Private/NiagaraDataInterfaceVolumeCache.cpp
@@ -104,7 +104,7 @@ bool UNiagaraDataInterfaceVolumeCache::Equals(const UNiagaraDataInterface* Other
 	return OtherTyped != nullptr && VolumeCache == OtherTyped->VolumeCache;
 }
 
-inline int32 UNiagaraDataInterfaceVolumeCache::PerInstanceDataSize() const 
+int32 UNiagaraDataInterfaceVolumeCache::PerInstanceDataSize() const 
 { 
 	return sizeof(FVolumeCacheInstanceData_GameThread); 
 }
diff --git a/Engine/Source/Editor/UnrealEd/Classes/Commandlets/DerivedDataCacheCommandlet.h b/Engine/Source/Editor/UnrealEd/Classes/Commandlets/DerivedDataCacheCommandlet.h
--- a/Engine/Source/Editor/UnrealEd/Classes/Commandlets/DerivedDataCacheCommandlet.h
+++ b/Engine/Source/Editor/UnrealEd/Classes/Commandlets/DerivedDataCacheCommandlet.h
@@ -16,6 +16,7 @@ class UDerivedDataCacheCommandlet : public UCommandlet
 {
 	GENERATED_UCLASS_BODY()
 	UDerivedDataCacheCommandlet(FVTableHelper& Helper); // Declare the FVTableHelper constructor manually so that we can forward-declare-only TUniquePtrs in the header without getting compile error in generated cpp
+	~UDerivedDataCacheCommandlet(); // Same reason as ctor above
 
 	//~ Begin UCommandlet Interface
 	virtual int32 Main(const FString& Params) override;
diff --git a/Engine/Source/Editor/UnrealEd/Private/Commandlets/DerivedDataCacheCommandlet.cpp b/Engine/Source/Editor/UnrealEd/Private/Commandlets/DerivedDataCacheCommandlet.cpp
--- a/Engine/Source/Editor/UnrealEd/Private/Commandlets/DerivedDataCacheCommandlet.cpp
+++ b/Engine/Source/Editor/UnrealEd/Private/Commandlets/DerivedDataCacheCommandlet.cpp
@@ -120,6 +120,10 @@ UDerivedDataCacheCommandlet::UDerivedDataCacheCommandlet(FVTableHelper& Helper)
 {
 }
 
+UDerivedDataCacheCommandlet::~UDerivedDataCacheCommandlet()
+{
+}
+
 UDerivedDataCacheCommandlet::UDerivedDataCacheCommandlet(const FObjectInitializer& ObjectInitializer)
 	: Super(ObjectInitializer)
 {
diff --git a/Engine/Source/Programs/UnrealHeaderTool/Private/CodeGenerator.cpp b/Engine/Source/Programs/UnrealHeaderTool/Private/CodeGenerator.cpp
--- a/Engine/Source/Programs/UnrealHeaderTool/Private/CodeGenerator.cpp
+++ b/Engine/Source/Programs/UnrealHeaderTool/Private/CodeGenerator.cpp
@@ -3047,6 +3047,14 @@ void ExportCopyConstructorDefinition(FOutputDevice& Out, const TCHAR* API, const
 	Out.Logf(TEXT("public:\r\n"));
 }
 
+void ExportDestructorDefinition(FOutputDevice& Out, FUnrealClassDefinitionInfo& ClassDef, const TCHAR* API, const TCHAR* ClassCPPName)
+{
+	if (!ClassDef.IsDestructorDeclared())
+	{
+		Out.Logf(TEXT("\t %s_API virtual ~%s();\r\n"), API, ClassCPPName);
+	}
+}
+
 /**
  * Generates vtable helper caller and eventual constructor body.
  *
@@ -3090,6 +3098,7 @@ void ExportStandardConstructorsMacro(FOutputDevice& Out, FUnrealClassDefinitionI
 
 	ExportVTableHelperCtorAndCaller(Out, ClassDef, API, ClassCPPName);
 	ExportCopyConstructorDefinition(Out, API, ClassCPPName);
+	ExportDestructorDefinition(Out, ClassDef, API, ClassCPPName);
 }
 
 /**
@@ -3190,6 +3199,7 @@ void ExportEnhancedConstructorsMacro(FOutputDevice& Out, FUnrealClassDefinitionI
 	ExportConstructorDefinition(Out, ClassDef, API, ClassCPPName);
 	ExportVTableHelperCtorAndCaller(Out, ClassDef, API, ClassCPPName);
 	ExportDefaultConstructorCallDefinition(Out, ClassDef, ClassCPPName);
+	ExportDestructorDefinition(Out, ClassDef, API, ClassCPPName);
 }
 
 /**
@@ -3225,6 +3235,11 @@ void FNativeClassHeaderGenerator::ExportConstructorsMacros(FOutputDevice& OutGen
 		Out.Logf(TEXT("\tDEFINE_VTABLE_PTR_HELPER_CTOR(%s);" LINE_TERMINATOR_ANSI), *ClassCPPName);
 	}
 
+	if (!ClassDef.IsDestructorDeclared())
+	{
+		Out.Logf(TEXT("\t%s::~%s() {}" LINE_TERMINATOR_ANSI), *ClassCPPName, *ClassCPPName);
+	}
+
 	OutGeneratedHeaderText.Log(Macroize(*StdMacroName, *StdMacro));
 	OutGeneratedHeaderText.Log(Macroize(*EnhMacroName, *EnhMacro));
 
diff --git a/Engine/Source/Programs/UnrealHeaderTool/Private/HeaderParser.cpp b/Engine/Source/Programs/UnrealHeaderTool/Private/HeaderParser.cpp
--- a/Engine/Source/Programs/UnrealHeaderTool/Private/HeaderParser.cpp
+++ b/Engine/Source/Programs/UnrealHeaderTool/Private/HeaderParser.cpp
@@ -5224,6 +5224,9 @@ bool FHeaderParser::CompileDeclaration(TArray<FUnrealFunctionDefinitionInfo*>& D
 		if (CheckForConstructor(StructDef, Declaration))
 		{
 		}
+		else if (CheckForDestructor(StructDef, Declaration))
+		{
+		}
 		else if (TopNest->NestType == ENestType::Class)
 		{
 			if (CheckForSerialize(StructDef, Declaration))
@@ -5377,6 +5380,45 @@ bool FHeaderParser::CheckForConstructor(FUnrealStructDefinitionInfo& StructDef,
 	return false;
 }
 
+bool FHeaderParser::CheckForDestructor(FUnrealStructDefinitionInfo& StructDef, const FDeclaration& Declaration)
+{
+	FUnrealClassDefinitionInfo& ClassDef = StructDef.AsClassChecked();
+
+	if (ClassDef.IsDestructorDeclared())
+	{
+		return false;
+	}
+
+	FTokenReplay Tokens(Declaration.Tokens);
+
+	FToken Token;
+	if (!Tokens.GetToken(Token))
+	{
+		return false;
+	}
+
+	while (Token.IsIdentifier(TEXT("virtual"), ESearchCase::CaseSensitive) || Token.Value.EndsWith(TEXT("_API"), ESearchCase::CaseSensitive))
+	{
+		Tokens.GetToken(Token);
+	}
+
+	if (!Token.IsSymbol('~'))
+	{
+		return false;
+	}
+
+	Tokens.GetToken(Token);
+
+	if (!Token.IsIdentifier(*ClassDef.GetAlternateNameCPP(), ESearchCase::IgnoreCase))
+	{
+		return false;
+	}
+
+	ClassDef.MarkDestructorDeclared();
+
+	return false;
+}
+
 bool FHeaderParser::CheckForSerialize(FUnrealStructDefinitionInfo& StructDef, const FDeclaration& Declaration)
 {
 	FTokenReplay Tokens(Declaration.Tokens);
diff --git a/Engine/Source/Programs/UnrealHeaderTool/Private/HeaderParser.h b/Engine/Source/Programs/UnrealHeaderTool/Private/HeaderParser.h
--- a/Engine/Source/Programs/UnrealHeaderTool/Private/HeaderParser.h
+++ b/Engine/Source/Programs/UnrealHeaderTool/Private/HeaderParser.h
@@ -635,6 +635,9 @@ private:
 	// Check to see if the declaration is a constructor
 	static bool CheckForConstructor(FUnrealStructDefinitionInfo& StructDef, const FDeclaration& Declaration);
 
+	// Check to see if the declaration is a destructor
+	static bool CheckForDestructor(FUnrealStructDefinitionInfo& StructDef, const FDeclaration& Declaration);
+
 	// Check to see if the declaration is a serialize
 	static bool CheckForSerialize(FUnrealStructDefinitionInfo& StructDef, const FDeclaration& Declaration);
 
diff --git a/Engine/Source/Programs/UnrealHeaderTool/Private/UnrealTypeDefinitionInfo.h b/Engine/Source/Programs/UnrealHeaderTool/Private/UnrealTypeDefinitionInfo.h
--- a/Engine/Source/Programs/UnrealHeaderTool/Private/UnrealTypeDefinitionInfo.h
+++ b/Engine/Source/Programs/UnrealHeaderTool/Private/UnrealTypeDefinitionInfo.h
@@ -3299,6 +3299,21 @@ public:
 		bObjectInitializerConstructorDeclared = true;
 	}
 
+	/**
+	 * Return true if the custom vtable helper destructor is declared
+	 */
+	bool IsDestructorDeclared() const
+	{
+		return bDestructorDeclared;
+	}
+	/**
+	 * Mark that the destructor has been declared
+	 */
+	void MarkDestructorDeclared()
+	{
+		bDestructorDeclared = true;
+	}
+
 	/**
 	 * Return true if the custom vtable helper constructor is declared
 	 */
@@ -3510,6 +3525,9 @@ private:
 	/** Is ObjectInitializer constructor (i.e. a constructor with only one parameter of type FObjectInitializer) declared? */
 	bool bObjectInitializerConstructorDeclared = false;
 
+	/** Is destructor declared? */
+	bool bDestructorDeclared = false;
+
 	/** Is custom VTable helper constructor declared? */
 	bool bCustomVTableHelperConstructorDeclared = false;
 
diff --git a/Engine/Source/Runtime/CoreUObject/Public/UObject/NoExportTypes.h b/Engine/Source/Runtime/CoreUObject/Public/UObject/NoExportTypes.h
--- a/Engine/Source/Runtime/CoreUObject/Public/UObject/NoExportTypes.h
+++ b/Engine/Source/Runtime/CoreUObject/Public/UObject/NoExportTypes.h
@@ -2592,6 +2592,7 @@ public:
 
 	UObject(const FObjectInitializer& ObjectInitializer = FObjectInitializer::Get());
 	UObject(FVTableHelper& Helper);
+	~UObject();
 	
 	/**
 	 * Executes some portion of the ubergraph.
diff --git a/Engine/Source/Runtime/NavigationSystem/Private/NavMesh/RecastNavMesh.cpp b/Engine/Source/Runtime/NavigationSystem/Private/NavMesh/RecastNavMesh.cpp
--- a/Engine/Source/Runtime/NavigationSystem/Private/NavMesh/RecastNavMesh.cpp
+++ b/Engine/Source/Runtime/NavigationSystem/Private/NavMesh/RecastNavMesh.cpp
@@ -259,6 +259,10 @@ void ARecastNavMesh::Serialize( FArchive& Ar )
 	}
 }
 
+ARecastNavMesh::~ARecastNavMesh()
+{
+}
+
 #else // WITH_RECAST
 
 #include "Detour/DetourNavMesh.h"
