Fixed unsafe type casts in UHT

#rnx

diff --git a/Engine/Source/Programs/UnrealHeaderTool/Private/BaseParser.h b/Engine/Source/Programs/UnrealHeaderTool/Private/BaseParser.h
--- a/Engine/Source/Programs/UnrealHeaderTool/Private/BaseParser.h
+++ b/Engine/Source/Programs/UnrealHeaderTool/Private/BaseParser.h
@@ -214,7 +214,7 @@ public:
 		case ETokenType::FloatConst:
 		{
 			float Float = GetFloatValue();
-			if (Float == FMath::TruncToInt(Float))
+			if (Float == (float)FMath::TruncToInt(Float))
 			{
 				I = (int32)Float;
 				return true;
@@ -240,7 +240,7 @@ public:
 		case ETokenType::FloatConst:
 		{
 			float Float = GetFloatValue();
-			if (Float == FMath::TruncToInt(Float))
+			if (Float == (float)FMath::TruncToInt(Float))
 			{
 				I = (int32)Float;
 				return true;
diff --git a/Engine/Source/Programs/UnrealHeaderTool/Private/CodeGenerator.cpp b/Engine/Source/Programs/UnrealHeaderTool/Private/CodeGenerator.cpp
--- a/Engine/Source/Programs/UnrealHeaderTool/Private/CodeGenerator.cpp
+++ b/Engine/Source/Programs/UnrealHeaderTool/Private/CodeGenerator.cpp
@@ -3129,7 +3129,7 @@ void ExportConstructorDefinition(FOutputDevice& Out, FUnrealClassDefinitionInfo&
 				// file would have been assigned out,  so we just have to wait
 				while (!SuperClassDef->IsConstructorDeclared())
 				{
-					FPlatformProcess::Sleep(0.01);
+					FPlatformProcess::Sleep(0.01f);
 				}
 
 				bSuperClassObjectInitializerConstructorDeclared = SuperClassDef->IsObjectInitializerConstructorDeclared();
