Set EngineIncludeOrderVersion.Latest for several additional targets

#rnx

diff --git a/Engine/Source/Programs/BenchmarkTool/BenchmarkTool.Target.cs b/Engine/Source/Programs/BenchmarkTool/BenchmarkTool.Target.cs
--- a/Engine/Source/Programs/BenchmarkTool/BenchmarkTool.Target.cs
+++ b/Engine/Source/Programs/BenchmarkTool/BenchmarkTool.Target.cs
@@ -10,6 +10,7 @@ public class BenchmarkToolTarget : TargetRules
 	{
 		Type = TargetType.Program;
 		LinkType = TargetLinkType.Monolithic;
+		IncludeOrderVersion = EngineIncludeOrderVersion.Latest;
 		LaunchModuleName = "BenchmarkTool";
 
 		bBuildDeveloperTools			= false;
diff --git a/Engine/Source/Programs/BlankProgram/BlankProgram.Target.cs b/Engine/Source/Programs/BlankProgram/BlankProgram.Target.cs
--- a/Engine/Source/Programs/BlankProgram/BlankProgram.Target.cs
+++ b/Engine/Source/Programs/BlankProgram/BlankProgram.Target.cs
@@ -9,6 +9,7 @@ public class BlankProgramTarget : TargetRules
 	public BlankProgramTarget(TargetInfo Target) : base(Target)
 	{
 		Type = TargetType.Program;
+		IncludeOrderVersion = EngineIncludeOrderVersion.Latest;
 		LinkType = TargetLinkType.Monolithic;
 		LaunchModuleName = "BlankProgram";
 
diff --git a/Engine/Source/Programs/BuildPatchTool/BuildPatchTool.Target.cs b/Engine/Source/Programs/BuildPatchTool/BuildPatchTool.Target.cs
--- a/Engine/Source/Programs/BuildPatchTool/BuildPatchTool.Target.cs
+++ b/Engine/Source/Programs/BuildPatchTool/BuildPatchTool.Target.cs
@@ -11,6 +11,7 @@ public class BuildPatchToolTarget : TargetRules
 	{
 		Type = TargetType.Program;
 		LinkType = TargetLinkType.Monolithic;
+		IncludeOrderVersion = EngineIncludeOrderVersion.Latest;
 		LaunchModuleName = "BuildPatchTool";
         bLegalToDistributeBinary = true;
 		UndecoratedConfiguration = UnrealTargetConfiguration.Shipping;
diff --git a/Engine/Source/Programs/CrashReportClient/CrashReportClient.Target.cs b/Engine/Source/Programs/CrashReportClient/CrashReportClient.Target.cs
--- a/Engine/Source/Programs/CrashReportClient/CrashReportClient.Target.cs
+++ b/Engine/Source/Programs/CrashReportClient/CrashReportClient.Target.cs
@@ -22,6 +22,7 @@ public class CrashReportClientTarget : TargetRules
 	{
 		Type = TargetType.Program;
 		LinkType = TargetLinkType.Monolithic;
+		IncludeOrderVersion = EngineIncludeOrderVersion.Latest;
 		UndecoratedConfiguration = UnrealTargetConfiguration.Shipping;
 
 		LaunchModuleName = "CrashReportClient";
diff --git a/Engine/Source/Programs/LiveCodingConsole/LiveCodingConsole.Target.cs b/Engine/Source/Programs/LiveCodingConsole/LiveCodingConsole.Target.cs
--- a/Engine/Source/Programs/LiveCodingConsole/LiveCodingConsole.Target.cs
+++ b/Engine/Source/Programs/LiveCodingConsole/LiveCodingConsole.Target.cs
@@ -9,6 +9,7 @@ public class LiveCodingConsoleTarget : TargetRules
 	{
 		Type = TargetType.Program;
 		LinkType = TargetLinkType.Monolithic;
+		IncludeOrderVersion = EngineIncludeOrderVersion.Latest;
 		LaunchModuleName = "LiveCodingConsole";
 
 		bBuildDeveloperTools = false;
diff --git a/Engine/Source/Programs/LowLevelTests/LowLevelTests.Target.cs b/Engine/Source/Programs/LowLevelTests/LowLevelTests.Target.cs
--- a/Engine/Source/Programs/LowLevelTests/LowLevelTests.Target.cs
+++ b/Engine/Source/Programs/LowLevelTests/LowLevelTests.Target.cs
@@ -8,5 +8,6 @@ public class LowLevelTestsTarget : TestTargetRules
 	public LowLevelTestsTarget(TargetInfo Target) : base(Target)
 	{
 		bWithLowLevelTestsOverride = true;
+		IncludeOrderVersion = EngineIncludeOrderVersion.Latest;
 	}
 }
diff --git a/Engine/Source/Programs/TestPAL/TestPAL.Target.cs b/Engine/Source/Programs/TestPAL/TestPAL.Target.cs
--- a/Engine/Source/Programs/TestPAL/TestPAL.Target.cs
+++ b/Engine/Source/Programs/TestPAL/TestPAL.Target.cs
@@ -1,4 +1,4 @@
-﻿// Copyright Epic Games, Inc. All Rights Reserved.
+// Copyright Epic Games, Inc. All Rights Reserved.
 
 using UnrealBuildTool;
 using System.Collections.Generic;
@@ -10,6 +10,7 @@ public class TestPALTarget : TargetRules
 	{
 		Type = TargetType.Program;
 		LinkType = TargetLinkType.Monolithic;
+		IncludeOrderVersion = EngineIncludeOrderVersion.Latest;
 		LaunchModuleName = "TestPAL";
 
 		// Lean and mean
diff --git a/Engine/Source/Programs/TextureBuildWorker/TextureBuildWorker.Build.cs b/Engine/Source/Programs/TextureBuildWorker/TextureBuildWorker.Build.cs
--- a/Engine/Source/Programs/TextureBuildWorker/TextureBuildWorker.Build.cs
+++ b/Engine/Source/Programs/TextureBuildWorker/TextureBuildWorker.Build.cs
@@ -18,6 +18,8 @@ public abstract class TextureBuildWorkerTarget : DerivedDataBuildWorkerTarget
 
 	public TextureBuildWorkerTarget(TargetInfo Target) : base(Target)
 	{
+		IncludeOrderVersion = EngineIncludeOrderVersion.Latest;
+
 		SolutionDirectory += "/Texture";
 
 		//always add the engine TextureFormatOodle (if enabled)
diff --git a/Engine/Source/Programs/UnrealHeaderTool/UnrealHeaderTool.Target.cs b/Engine/Source/Programs/UnrealHeaderTool/UnrealHeaderTool.Target.cs
--- a/Engine/Source/Programs/UnrealHeaderTool/UnrealHeaderTool.Target.cs
+++ b/Engine/Source/Programs/UnrealHeaderTool/UnrealHeaderTool.Target.cs
@@ -9,6 +9,7 @@ public class UnrealHeaderToolTarget : TargetRules
 	{
 		Type = TargetType.Program;
 		LinkType = TargetLinkType.Modular;
+		IncludeOrderVersion = EngineIncludeOrderVersion.Latest;
 		LaunchModuleName = "UnrealHeaderTool";
 
 		// Never use malloc profiling in Unreal Header Tool.  We set this because often UHT is compiled right before the engine
diff --git a/Engine/Source/Programs/UnrealInsights/UnrealInsights.Target.cs b/Engine/Source/Programs/UnrealInsights/UnrealInsights.Target.cs
--- a/Engine/Source/Programs/UnrealInsights/UnrealInsights.Target.cs
+++ b/Engine/Source/Programs/UnrealInsights/UnrealInsights.Target.cs
@@ -14,6 +14,7 @@ public class UnrealInsightsTarget : TargetRules
 	{
 		Type = TargetType.Program;
 		LinkType = bMonolithic ? TargetLinkType.Monolithic : TargetLinkType.Modular;
+		IncludeOrderVersion = EngineIncludeOrderVersion.Latest;
 
 		LaunchModuleName = "UnrealInsights";
 		if (bBuildEditor)
diff --git a/Engine/Source/Programs/UnrealObjectPtrTool/UnrealObjectPtrTool.Target.cs b/Engine/Source/Programs/UnrealObjectPtrTool/UnrealObjectPtrTool.Target.cs
--- a/Engine/Source/Programs/UnrealObjectPtrTool/UnrealObjectPtrTool.Target.cs
+++ b/Engine/Source/Programs/UnrealObjectPtrTool/UnrealObjectPtrTool.Target.cs
@@ -10,6 +10,7 @@ public class UnrealObjectPtrToolTarget : TargetRules
 	{
 		Type = TargetType.Program;
 		LinkType = TargetLinkType.Modular;
+		IncludeOrderVersion = EngineIncludeOrderVersion.Latest;
 		LaunchModuleName = "UnrealObjectPtrTool";
 
 		bBuildDeveloperTools			= false;
diff --git a/Engine/Source/Programs/UnrealPackageTool/UnrealPackageTool.Target.cs b/Engine/Source/Programs/UnrealPackageTool/UnrealPackageTool.Target.cs
--- a/Engine/Source/Programs/UnrealPackageTool/UnrealPackageTool.Target.cs
+++ b/Engine/Source/Programs/UnrealPackageTool/UnrealPackageTool.Target.cs
@@ -9,6 +9,7 @@ public class UnrealPackageToolTarget : TargetRules
 	{
 		Type = TargetType.Program;
 		LinkType = TargetLinkType.Monolithic;
+		IncludeOrderVersion = EngineIncludeOrderVersion.Latest;
 		LaunchModuleName = "UnrealPackageTool";
 
 		bBuildDeveloperTools = false;
diff --git a/Engine/Source/Programs/UnrealVersionSelector/UnrealVersionSelector.Target.cs b/Engine/Source/Programs/UnrealVersionSelector/UnrealVersionSelector.Target.cs
--- a/Engine/Source/Programs/UnrealVersionSelector/UnrealVersionSelector.Target.cs
+++ b/Engine/Source/Programs/UnrealVersionSelector/UnrealVersionSelector.Target.cs
@@ -9,6 +9,7 @@ public class UnrealVersionSelectorTarget : TargetRules
 	{
 		Type = TargetType.Program;
 		LinkType = TargetLinkType.Monolithic;
+		IncludeOrderVersion = EngineIncludeOrderVersion.Latest;
 		LaunchModuleName = "UnrealVersionSelector";
 		
 		bBuildDeveloperTools = false;
diff --git a/Engine/Source/Programs/ZenDashboard/ZenDashboard.Target.cs b/Engine/Source/Programs/ZenDashboard/ZenDashboard.Target.cs
--- a/Engine/Source/Programs/ZenDashboard/ZenDashboard.Target.cs
+++ b/Engine/Source/Programs/ZenDashboard/ZenDashboard.Target.cs
@@ -9,6 +9,7 @@ public class ZenDashboardTarget : TargetRules
 	{
 		Type = TargetType.Program;
 		LinkType = TargetLinkType.Monolithic;
+		IncludeOrderVersion = EngineIncludeOrderVersion.Latest;
 		LaunchModuleName = "ZenDashboard";
 
 		bBuildDeveloperTools = false;
