Allow an Epic-only configuration on UnrealHeaderTool that enforces TObjectPtr usage in Engine Plugins specifically.  Cannot be used externally due to the occurrance of third party plugins being placed in the engine plugins directory.  Tested against C++ UHT and C# UHT both with and without the Epic-only configuration file present.

diff --git a/Engine/Programs/UnrealHeaderTool/Config/DefaultEngine.ini b/Engine/Programs/UnrealHeaderTool/Config/DefaultEngine.ini
--- a/Engine/Programs/UnrealHeaderTool/Config/DefaultEngine.ini
+++ b/Engine/Programs/UnrealHeaderTool/Config/DefaultEngine.ini
@@ -162,6 +162,8 @@ gc.MaxObjectsInProgram=262144
 [UnrealHeaderTool]
 EngineNativePointerMemberBehavior=Disallow
 EngineObjectPtrMemberBehavior=AllowSilently
+EnginePluginNativePointerMemberBehavior=AllowSilently
+EnginePluginObjectPtrMemberBehavior=AllowSilently
 NonEngineNativePointerMemberBehavior=AllowSilently
 NonEngineObjectPtrMemberBehavior=AllowSilently
 AutomaticSettersAndGetters=False
diff --git a/Engine/Source/Programs/Shared/EpicGames.UHT/Parsers/UhtPropertyParser.cs b/Engine/Source/Programs/Shared/EpicGames.UHT/Parsers/UhtPropertyParser.cs
--- a/Engine/Source/Programs/Shared/EpicGames.UHT/Parsers/UhtPropertyParser.cs
+++ b/Engine/Source/Programs/Shared/EpicGames.UHT/Parsers/UhtPropertyParser.cs
@@ -1369,7 +1369,7 @@ namespace EpicGames.UHT.Parsers
 
 				// Optionally emit messages about native pointer members and swallow trailing 'const' after pointer properties
 				UhtObjectPropertyBase.ConditionalLogPointerUsage(propertySettings, session.Config!.EngineNativePointerMemberBehavior,
-					session.Config!.NonEngineNativePointerMemberBehavior, "Native pointer", tokenReader, typeStartPos, "TObjectPtr");
+					session.Config!.EnginePluginNativePointerMemberBehavior, session.Config!.NonEngineNativePointerMemberBehavior, "Native pointer", tokenReader, typeStartPos, "TObjectPtr");
 
 				if (propertySettings.PropertyCategory == UhtPropertyCategory.Member)
 				{
diff --git a/Engine/Source/Programs/Shared/EpicGames.UHT/Types/Properties/UhtObjectPropertyBase.cs b/Engine/Source/Programs/Shared/EpicGames.UHT/Types/Properties/UhtObjectPropertyBase.cs
--- a/Engine/Source/Programs/Shared/EpicGames.UHT/Types/Properties/UhtObjectPropertyBase.cs
+++ b/Engine/Source/Programs/Shared/EpicGames.UHT/Types/Properties/UhtObjectPropertyBase.cs
@@ -226,14 +226,15 @@ namespace EpicGames.UHT.Types
 		/// </summary>
 		/// <param name="propertySettings">Property settings</param>
 		/// <param name="engineBehavior">Expected behavior for engine types</param>
+		/// <param name="enginePluginBehavior">Expected behavior for engine plugin types</param>
 		/// <param name="nonEngineBehavior">Expected behavior for non-engine types</param>
 		/// <param name="pointerTypeDesc">Description of the pointer type</param>
 		/// <param name="tokenReader">Token reader for type being parsed</param>
 		/// <param name="typeStartPos">Starting character position of the type</param>
 		/// <param name="alternativeTypeDesc">Suggested alternate declaration</param>
 		/// <exception cref="UhtIceException">Thrown if the behavior type is unexpected</exception>
-		public static void ConditionalLogPointerUsage(UhtPropertySettings propertySettings, UhtPointerMemberBehavior engineBehavior, UhtPointerMemberBehavior nonEngineBehavior,
-			string pointerTypeDesc, IUhtTokenReader tokenReader, int typeStartPos, string? alternativeTypeDesc)
+		public static void ConditionalLogPointerUsage(UhtPropertySettings propertySettings, UhtPointerMemberBehavior engineBehavior, UhtPointerMemberBehavior enginePluginBehavior,
+			UhtPointerMemberBehavior nonEngineBehavior, string pointerTypeDesc, IUhtTokenReader tokenReader, int typeStartPos, string? alternativeTypeDesc)
 		{
 			if (propertySettings.PropertyCategory != UhtPropertyCategory.Member)
 			{
@@ -241,7 +242,18 @@ namespace EpicGames.UHT.Types
 			}
 
 			UhtPackage package = propertySettings.Outer.Package;
-			UhtPointerMemberBehavior behavior = package.IsPartOfEngine ? engineBehavior : nonEngineBehavior;
+			UhtPointerMemberBehavior behavior = nonEngineBehavior;
+			if (package.IsPartOfEngine)
+			{
+				if (package.Module.BaseDirectory.Contains("/Plugins/", StringComparison.InvariantCultureIgnoreCase) || package.Module.BaseDirectory.Contains("\\Plugins\\", StringComparison.InvariantCultureIgnoreCase))
+				{
+					behavior = enginePluginBehavior;
+				}
+				else
+				{
+					behavior = engineBehavior;
+				}
+			}
 
 			if (behavior == UhtPointerMemberBehavior.AllowSilently)
 			{
diff --git a/Engine/Source/Programs/Shared/EpicGames.UHT/Types/Properties/UhtObjectPtrProperty.cs b/Engine/Source/Programs/Shared/EpicGames.UHT/Types/Properties/UhtObjectPtrProperty.cs
--- a/Engine/Source/Programs/Shared/EpicGames.UHT/Types/Properties/UhtObjectPtrProperty.cs
+++ b/Engine/Source/Programs/Shared/EpicGames.UHT/Types/Properties/UhtObjectPtrProperty.cs
@@ -115,7 +115,7 @@ namespace EpicGames.UHT.Types
 				return null;
 			}
 
-			ConditionalLogPointerUsage(propertySettings, session.Config!.EngineObjectPtrMemberBehavior,
+			ConditionalLogPointerUsage(propertySettings, session.Config!.EngineObjectPtrMemberBehavior, session.Config!.EnginePluginObjectPtrMemberBehavior,
 				session.Config!.NonEngineObjectPtrMemberBehavior, "ObjectPtr", tokenReader, typeStartPos, null);
 
 			if (propertyClass.IsChildOf(propertyClass.Session.UClass))
diff --git a/Engine/Source/Programs/Shared/EpicGames.UHT/Utils/UhtConfig.cs b/Engine/Source/Programs/Shared/EpicGames.UHT/Utils/UhtConfig.cs
--- a/Engine/Source/Programs/Shared/EpicGames.UHT/Utils/UhtConfig.cs
+++ b/Engine/Source/Programs/Shared/EpicGames.UHT/Utils/UhtConfig.cs
@@ -49,6 +49,16 @@ namespace EpicGames.UHT.Utils
 		/// </summary>
 		public UhtPointerMemberBehavior EngineObjectPtrMemberBehavior { get; }
 
+		/// <summary>
+		/// Pointer warning for native pointers in engine plugins
+		/// </summary>
+		public UhtPointerMemberBehavior EnginePluginNativePointerMemberBehavior { get; }
+
+		/// <summary>
+		/// Pointer warning for object pointers in engine plugins
+		/// </summary>
+		public UhtPointerMemberBehavior EnginePluginObjectPtrMemberBehavior { get; }
+
 		/// <summary>
 		/// Pointer warning for native pointers outside the engine
 		/// </summary>
diff --git a/Engine/Source/Programs/UnrealBuildTool/Configuration/UEBuildTarget.cs b/Engine/Source/Programs/UnrealBuildTool/Configuration/UEBuildTarget.cs
--- a/Engine/Source/Programs/UnrealBuildTool/Configuration/UEBuildTarget.cs
+++ b/Engine/Source/Programs/UnrealBuildTool/Configuration/UEBuildTarget.cs
@@ -1,4 +1,4 @@
-// Copyright Epic Games, Inc. All Rights Reserved.
+﻿// Copyright Epic Games, Inc. All Rights Reserved.
 
 using System;
 using System.Collections.Generic;
@@ -1908,6 +1908,7 @@ namespace UnrealBuildTool
 			if (Rules.NativePointerMemberBehaviorOverride != null)
 			{
 				UHTAdditionalArguments.Add("-ini:Engine:[UnrealHeaderTool]:EngineNativePointerMemberBehavior=" + Rules.NativePointerMemberBehaviorOverride);
+				UHTAdditionalArguments.Add("-ini:Engine:[UnrealHeaderTool]:EnginePluginNativePointerMemberBehavior=" + Rules.NativePointerMemberBehaviorOverride);
 				UHTAdditionalArguments.Add("-ini:Engine:[UnrealHeaderTool]:NonEngineNativePointerMemberBehavior=" + Rules.NativePointerMemberBehaviorOverride);
 			}
 
diff --git a/Engine/Source/Programs/UnrealBuildTool/Modes/UnrealHeaderToolMode.cs b/Engine/Source/Programs/UnrealBuildTool/Modes/UnrealHeaderToolMode.cs
--- a/Engine/Source/Programs/UnrealBuildTool/Modes/UnrealHeaderToolMode.cs
+++ b/Engine/Source/Programs/UnrealBuildTool/Modes/UnrealHeaderToolMode.cs
@@ -71,6 +71,16 @@ namespace UnrealBuildTool.Modes
 		/// </summary>
 		private readonly UhtPointerMemberBehavior _engineObjectPtrMemberBehavior = UhtPointerMemberBehavior.AllowSilently;
 
+		/// <summary>
+		/// Internal version of pointer warning for native pointers in engine plugins
+		/// </summary>
+		private readonly UhtPointerMemberBehavior _enginePluginNativePointerMemberBehavior = UhtPointerMemberBehavior.AllowSilently;
+
+		/// <summary>
+		/// Internal version of pointer warning for object pointers in engine plugins
+		/// </summary>
+		private readonly UhtPointerMemberBehavior _enginePluginObjectPtrMemberBehavior = UhtPointerMemberBehavior.AllowSilently;
+
 		/// <summary>
 		/// Internal version of pointer warning for native pointers outside the engine
 		/// </summary>
@@ -106,6 +116,12 @@ namespace UnrealBuildTool.Modes
 		/// <inheritdoc/>
 		public UhtPointerMemberBehavior EngineObjectPtrMemberBehavior => this._engineObjectPtrMemberBehavior;
 
+		/// <inheritdoc/>
+		public UhtPointerMemberBehavior EnginePluginNativePointerMemberBehavior => this._enginePluginNativePointerMemberBehavior;
+
+		/// <inheritdoc/>
+		public UhtPointerMemberBehavior EnginePluginObjectPtrMemberBehavior => this._enginePluginObjectPtrMemberBehavior;
+
 		/// <inheritdoc/>
 		public UhtPointerMemberBehavior NonEngineNativePointerMemberBehavior => this._nonEngineNativePointerMemberBehavior;
 
@@ -213,6 +229,8 @@ namespace UnrealBuildTool.Modes
 			this._defaultGeneratedCodeVersion = GetGeneratedCodeVersion("UnrealHeaderTool", "DefaultGeneratedCodeVersion", EGeneratedCodeVersion.V1);
 			this._engineNativePointerMemberBehavior = GetPointerMemberBehavior("UnrealHeaderTool", "EngineNativePointerMemberBehavior", UhtPointerMemberBehavior.AllowSilently);
 			this._engineObjectPtrMemberBehavior = GetPointerMemberBehavior("UnrealHeaderTool", "EngineObjectPtrMemberBehavior", UhtPointerMemberBehavior.AllowSilently);
+			this._enginePluginNativePointerMemberBehavior = GetPointerMemberBehavior("UnrealHeaderTool", "EnginePluginNativePointerMemberBehavior", UhtPointerMemberBehavior.AllowSilently);
+			this._enginePluginObjectPtrMemberBehavior = GetPointerMemberBehavior("UnrealHeaderTool", "EnginePluginObjectPtrMemberBehavior", UhtPointerMemberBehavior.AllowSilently);
 			this._nonEngineNativePointerMemberBehavior = GetPointerMemberBehavior("UnrealHeaderTool", "NonEngineNativePointerMemberBehavior", UhtPointerMemberBehavior.AllowSilently);
 			this._nonEngineObjectPtrMemberBehavior = GetPointerMemberBehavior("UnrealHeaderTool", "NonEngineObjectPtrMemberBehavior", UhtPointerMemberBehavior.AllowSilently);
 			this._areRigVMUObjectProeprtiesEnabled = GetBoolean("UnrealHeaderTool", "AreRigVMUObjectProeprtiesEnabled", false);
diff --git a/Engine/Source/Programs/UnrealHeaderTool/Private/BaseParser.h b/Engine/Source/Programs/UnrealHeaderTool/Private/BaseParser.h
--- a/Engine/Source/Programs/UnrealHeaderTool/Private/BaseParser.h
+++ b/Engine/Source/Programs/UnrealHeaderTool/Private/BaseParser.h
@@ -56,6 +56,9 @@ struct FUHTConfig
 	EPointerMemberBehavior EngineNativePointerMemberBehavior = EPointerMemberBehavior::AllowSilently;
 	EPointerMemberBehavior EngineObjectPtrMemberBehavior = EPointerMemberBehavior::AllowSilently;
 
+	EPointerMemberBehavior EnginePluginNativePointerMemberBehavior = EPointerMemberBehavior::AllowSilently;
+	EPointerMemberBehavior EnginePluginObjectPtrMemberBehavior = EPointerMemberBehavior::AllowSilently;
+
 	EPointerMemberBehavior NonEngineNativePointerMemberBehavior = EPointerMemberBehavior::AllowSilently;
 	EPointerMemberBehavior NonEngineObjectPtrMemberBehavior = EPointerMemberBehavior::AllowSilently;
 
diff --git a/Engine/Source/Programs/UnrealHeaderTool/Private/HeaderParser.cpp b/Engine/Source/Programs/UnrealHeaderTool/Private/HeaderParser.cpp
--- a/Engine/Source/Programs/UnrealHeaderTool/Private/HeaderParser.cpp
+++ b/Engine/Source/Programs/UnrealHeaderTool/Private/HeaderParser.cpp
@@ -787,6 +787,8 @@ FUHTConfig::FUHTConfig()
 	const FName GeneratedCodeVersionKey(TEXT("GeneratedCodeVersion"));
 	const FName EngineNativePointerMemberBehaviorKey(TEXT("EngineNativePointerMemberBehavior"));
 	const FName EngineObjectPtrMemberBehaviorKey(TEXT("EngineObjectPtrMemberBehavior"));
+	const FName EnginePluginNativePointerMemberBehaviorKey(TEXT("EnginePluginNativePointerMemberBehavior"));
+	const FName EnginePluginObjectPtrMemberBehaviorKey(TEXT("EnginePluginObjectPtrMemberBehavior"));
 	const FName NonEngineNativePointerMemberBehaviorKey(TEXT("NonEngineNativePointerMemberBehavior"));
 	const FName NonEngineObjectPtrMemberBehaviorKey(TEXT("NonEngineObjectPtrMemberBehavior"));
 
@@ -829,6 +831,14 @@ FUHTConfig::FUHTConfig()
 			{
 				EngineObjectPtrMemberBehavior = ToPointerMemberBehavior(It.Value().GetValue());
 			}
+			else if (It.Key() == EnginePluginNativePointerMemberBehaviorKey)
+			{
+				EnginePluginNativePointerMemberBehavior = ToPointerMemberBehavior(It.Value().GetValue());
+			}
+			else if (It.Key() == EnginePluginObjectPtrMemberBehaviorKey)
+			{
+				EnginePluginObjectPtrMemberBehavior = ToPointerMemberBehavior(It.Value().GetValue());
+			}
 			else if (It.Key() == NonEngineNativePointerMemberBehaviorKey)
 			{
 				NonEngineNativePointerMemberBehavior = ToPointerMemberBehavior(It.Value().GetValue());
@@ -4451,8 +4461,20 @@ void FHeaderParser::GetVarType(
 						// Optionally emit messages about native pointer members and swallow trailing 'const' after pointer properties
 						if (VariableCategory == EVariableCategory::Member)
 						{
-							// TODO: Remove exclusion for plugins under engine when all plugins have had their raw pointers converted.
-							ConditionalLogPointerUsage(bIsCurrentModulePartOfEngine ? UHTConfig.EngineNativePointerMemberBehavior : UHTConfig.NonEngineNativePointerMemberBehavior,
+							EPointerMemberBehavior PointerMemberBehavior = UHTConfig.NonEngineNativePointerMemberBehavior;
+							if (bIsCurrentModulePartOfEngine)
+							{
+								if (PackageDef.GetModule().BaseDirectory.Contains(TEXT("/Plugins/")))
+								{
+									PointerMemberBehavior = UHTConfig.EnginePluginNativePointerMemberBehavior;
+								}
+								else
+								{
+									PointerMemberBehavior = UHTConfig.EngineNativePointerMemberBehavior;
+								}
+							}
+
+							ConditionalLogPointerUsage(PointerMemberBehavior,
 								TEXT("Native pointer"), FString(InputPos - VarStartPos, Input + VarStartPos).TrimStartAndEnd().ReplaceCharWithEscapedChar(), TEXT("TObjectPtr"));
 
 							MatchIdentifier(TEXT("const"), ESearchCase::CaseSensitive);
@@ -4462,8 +4484,20 @@ void FHeaderParser::GetVarType(
 					}
 					else if ((PropertyType == CPT_ObjectPtrReference) && (VariableCategory == EVariableCategory::Member))
 					{
-						// TODO: Remove exclusion for plugins under engine when all plugins have had their raw pointers converted.
-						ConditionalLogPointerUsage(bIsCurrentModulePartOfEngine ? UHTConfig.EngineObjectPtrMemberBehavior : UHTConfig.NonEngineObjectPtrMemberBehavior,
+						EPointerMemberBehavior PointerMemberBehavior = UHTConfig.NonEngineObjectPtrMemberBehavior;
+						if (bIsCurrentModulePartOfEngine)
+						{
+							if (PackageDef.GetModule().BaseDirectory.Contains(TEXT("/Plugins/")))
+							{
+								PointerMemberBehavior = UHTConfig.EnginePluginObjectPtrMemberBehavior;
+							}
+							else
+							{
+								PointerMemberBehavior = UHTConfig.EngineObjectPtrMemberBehavior;
+							}
+						}
+
+						ConditionalLogPointerUsage(PointerMemberBehavior,
 							TEXT("ObjectPtr"), FString(InputPos - VarStartPos, Input + VarStartPos).TrimStartAndEnd().ReplaceCharWithEscapedChar(), nullptr);
 					}
 
