Added a IWYU pragma to the generated header so IWYU doesn't add the generated headers. They need to be added manually because there should only be one generated header included per header file.

diff --git a/Engine/Source/Programs/Shared/EpicGames.UHT/Exporters/CodeGen/UhtHeaderCodeGeneratorHFile.cs b/Engine/Source/Programs/Shared/EpicGames.UHT/Exporters/CodeGen/UhtHeaderCodeGeneratorHFile.cs
--- a/Engine/Source/Programs/Shared/EpicGames.UHT/Exporters/CodeGen/UhtHeaderCodeGeneratorHFile.cs
+++ b/Engine/Source/Programs/Shared/EpicGames.UHT/Exporters/CodeGen/UhtHeaderCodeGeneratorHFile.cs
@@ -40,6 +40,7 @@ namespace EpicGames.UHT.Exporters.CodeGen
 				StringBuilder builder = borrower.StringBuilder;
 
 				builder.Append(HeaderCopyright);
+				builder.Append("// IWYU pragma: private, include \"").Append(this.HeaderFile.IncludeFilePath).Append("\"\r\n");
 				builder.Append("#include \"UObject/ObjectMacros.h\"\r\n");
 				builder.Append("#include \"UObject/ScriptMacros.h\"\r\n");
 				if (headerInfo.NeedsPushModelHeaders)
diff --git a/Engine/Source/Programs/UnrealHeaderTool/Private/CodeGenerator.cpp b/Engine/Source/Programs/UnrealHeaderTool/Private/CodeGenerator.cpp
--- a/Engine/Source/Programs/UnrealHeaderTool/Private/CodeGenerator.cpp
+++ b/Engine/Source/Programs/UnrealHeaderTool/Private/CodeGenerator.cpp
@@ -3247,10 +3247,11 @@ void FNativeClassHeaderGenerator::ExportConstructorsMacros(FOutputDevice& OutGen
 	EnhancedUObjectConstructorsMacroCall.Logf(TEXT("\t%s\r\n"), *EnhMacroName);
 }
 
-bool FNativeClassHeaderGenerator::WriteHeader(FGeneratedFileInfo& FileInfo, const FString& InBodyText, const TSet<FString>& InAdditionalHeaders, const TSet<FString>& ForwardDeclarations)
+bool FNativeClassHeaderGenerator::WriteHeader(FGeneratedCPP& FileInfo, const FString& InBodyText, const TSet<FString>& InAdditionalHeaders, const TSet<FString>& ForwardDeclarations)
 {
 	FUHTStringBuilder GeneratedHeaderTextWithCopyright;
 	GeneratedHeaderTextWithCopyright.Log(HeaderCopyright);
+	GeneratedHeaderTextWithCopyright.Logf(TEXT("// IWYU pragma: private, include \"%s\"\r\n"), *FileInfo.SourceFile.GetIncludePath());
 	GeneratedHeaderTextWithCopyright.Log(TEXT("#include \"UObject/ObjectMacros.h\"\r\n"));
 	GeneratedHeaderTextWithCopyright.Log(TEXT("#include \"UObject/ScriptMacros.h\"\r\n"));
 
@@ -3283,7 +3284,7 @@ bool FNativeClassHeaderGenerator::WriteHeader(FGeneratedFileInfo& FileInfo, cons
 	GeneratedHeaderTextWithCopyright.Log(InBodyText);
 	GeneratedHeaderTextWithCopyright.Log(EnableDeprecationWarnings);
 
-	const bool bHasChanged = SaveHeaderIfChanged(FileInfo, MoveTemp(GeneratedHeaderTextWithCopyright));
+	const bool bHasChanged = SaveHeaderIfChanged(FileInfo.Header, MoveTemp(GeneratedHeaderTextWithCopyright));
 	return bHasChanged;
 }
 
@@ -5963,7 +5964,7 @@ void FNativeClassHeaderGenerator::WriteSourceFile(FGeneratedCPP& GeneratedCPP)
 				AdditionalHeaders.Add(FString(TEXT("Net/Core/PushModel/PushModelMacros.h")));
 			}
 
-			bool bHasChanged = WriteHeader(GeneratedCPP.Header, GeneratedHeaderText, AdditionalHeaders, GeneratedCPP.ForwardDeclarations);
+			bool bHasChanged = WriteHeader(GeneratedCPP, GeneratedHeaderText, AdditionalHeaders, GeneratedCPP.ForwardDeclarations);
 			WriteSource(Module, GeneratedCPP.Source, GeneratedCPPText, &SourceFile, GeneratedCPP.CrossModuleReferences, GeneratedCPP.ExportFlags);
 
 			SourceFile.SetGeneratedFilename(MoveTemp(GeneratedCPP.Header.GetFilename()));
diff --git a/Engine/Source/Programs/UnrealHeaderTool/Private/NativeClassExporter.h b/Engine/Source/Programs/UnrealHeaderTool/Private/NativeClassExporter.h
--- a/Engine/Source/Programs/UnrealHeaderTool/Private/NativeClassExporter.h
+++ b/Engine/Source/Programs/UnrealHeaderTool/Private/NativeClassExporter.h
@@ -381,7 +381,7 @@ private:
 	/**
 	 * After all of the dependency checking, but before actually exporting the class, set up the generated code
 	 */
-	static bool WriteHeader(FGeneratedFileInfo& FileInfo, const FString& InBodyText, const TSet<FString>& InAdditionalHeaders, const TSet<FString>& ForwardDeclarations);
+	static bool WriteHeader(FGeneratedCPP& FileInfo, const FString& InBodyText, const TSet<FString>& InAdditionalHeaders, const TSet<FString>& ForwardDeclarations);
 
 	/**
 	 * Write the body of a source file using a standard format
