UnrealHeaderTool: Remove unused code in CodeGenerator.cpp

#rnx

diff --git a/Engine/Source/Programs/UnrealHeaderTool/Private/CodeGenerator.cpp b/Engine/Source/Programs/UnrealHeaderTool/Private/CodeGenerator.cpp
--- a/Engine/Source/Programs/UnrealHeaderTool/Private/CodeGenerator.cpp
+++ b/Engine/Source/Programs/UnrealHeaderTool/Private/CodeGenerator.cpp
@@ -7133,9 +7133,6 @@ void ParseSourceFiles(TArray<FUnrealSourceFile*>& OrderedSourceFiles)
 			const FManifestModule& Module = ParsedCPP.PackageDef.GetModule();
 			FUnrealSourceFile& SourceFile = ParsedCPP.SourceFile;
 
-			FString ModuleRelativeFilename = SourceFile.GetFilename();
-			ConvertToBuildIncludePath(Module, ModuleRelativeFilename);
-
 			auto ParseSource = [&ParsedCPP]()
 			{
 				++GSourcesParsing;
