Updated C++ UHT to match C# UHT global delegate changes

diff --git a/Engine/Source/Programs/Shared/EpicGames.UHT/Exporters/CodeGen/UhtHeaderCodeGenerator.cs b/Engine/Source/Programs/Shared/EpicGames.UHT/Exporters/CodeGen/UhtHeaderCodeGenerator.cs
--- a/Engine/Source/Programs/Shared/EpicGames.UHT/Exporters/CodeGen/UhtHeaderCodeGenerator.cs
+++ b/Engine/Source/Programs/Shared/EpicGames.UHT/Exporters/CodeGen/UhtHeaderCodeGenerator.cs
@@ -286,7 +286,9 @@ namespace EpicGames.UHT.Exporters.CodeGen
 			++tabs;
 
 			if (addEventParameterStruct)
+			{
 				AppendEventParameter(builder, function, functionName, UhtPropertyTextType.EventParameterMember, true, tabs, "\r\n");
+			}
 
 			string eventParameterStructName = GetEventStructParametersName(function.Outer, functionName);
 
diff --git a/Engine/Source/Programs/Shared/EpicGames.UHT/Exporters/CodeGen/UhtHeaderCodeGeneratorHFile.cs b/Engine/Source/Programs/Shared/EpicGames.UHT/Exporters/CodeGen/UhtHeaderCodeGeneratorHFile.cs
--- a/Engine/Source/Programs/Shared/EpicGames.UHT/Exporters/CodeGen/UhtHeaderCodeGeneratorHFile.cs
+++ b/Engine/Source/Programs/Shared/EpicGames.UHT/Exporters/CodeGen/UhtHeaderCodeGeneratorHFile.cs
@@ -333,7 +333,9 @@ namespace EpicGames.UHT.Exporters.CodeGen
 				}
 
 				//if (!function.FunctionFlags.HasAnyFlags(EFunctionFlags.RequiredAPI)) // TODO: This requires too much fixup for now
+				//{
 				//	addAPI = false
+				//}
 
 				if (addAPI)
 				{
diff --git a/Engine/Source/Programs/UnrealHeaderTool/Private/CodeGenerator.cpp b/Engine/Source/Programs/UnrealHeaderTool/Private/CodeGenerator.cpp
--- a/Engine/Source/Programs/UnrealHeaderTool/Private/CodeGenerator.cpp
+++ b/Engine/Source/Programs/UnrealHeaderTool/Private/CodeGenerator.cpp
@@ -3601,7 +3601,7 @@ void FNativeClassHeaderGenerator::ExportGeneratedStructBodyMacros(FOutputDevice&
 		Out.Logf(TEXT("\tif (!%s)\r\n"), *OuterSingletonName);
 		Out.Logf(TEXT("\t{\r\n"));
 
-		Out.Logf(TEXT("\t\t%s = GetStaticStruct(%s, %s, TEXT(\"%s\"));\r\n"),
+		Out.Logf(TEXT("\t\t%s = GetStaticStruct(%s, (UObject*)%s, TEXT(\"%s\"));\r\n"),
 			*OuterSingletonName, *ChoppedSingletonName, *OuterName, *ActualStructName);
 
 		// if this struct has RigVM methods - we need to register the method to our central
@@ -3858,7 +3858,7 @@ void FNativeClassHeaderGenerator::ExportGeneratedEnumInitCode(FOutputDevice& Out
 
 		Out.Logf(TEXT("\t\tif (!Z_Registration_Info_UEnum_%s.OuterSingleton)\r\n"), *EnumNameCpp);
 		Out.Logf(TEXT("\t\t{\r\n"));
-		Out.Logf(TEXT("\t\t\tZ_Registration_Info_UEnum_%s.OuterSingleton = GetStaticEnum(%s, %s, TEXT(\"%s\"));\r\n"), *EnumNameCpp, *SingletonName, *PackageSingletonName, *EnumNameCpp);
+		Out.Logf(TEXT("\t\t\tZ_Registration_Info_UEnum_%s.OuterSingleton = GetStaticEnum(%s, (UObject*)%s, TEXT(\"%s\"));\r\n"), *EnumNameCpp, *SingletonName, *PackageSingletonName, *EnumNameCpp);
 		Out.Logf(TEXT("\t\t}\r\n"));
 		Out.Logf(TEXT("\t\treturn Z_Registration_Info_UEnum_%s.OuterSingleton;\r\n"), *EnumNameCpp);
 
@@ -3977,9 +3977,6 @@ bool FNativeClassHeaderGenerator::WillExportEventParms(FUnrealFunctionDefinition
 
 void WriteEventFunctionPrologue(FOutputDevice& Output, int32 Indent, const FParmsAndReturnProperties& Parameters, FUnrealObjectDefinitionInfo& FunctionOuterDef, const TCHAR* FunctionName)
 {
-	// now the body - first we need to declare a struct which will hold the parameters for the event/delegate call
-	Output.Logf(TEXT("\r\n%s{\r\n"), FCString::Tab(Indent));
-
 	// declare and zero-initialize the parameters and return value, if applicable
 	if (!Parameters.HasParms())
 		return;
@@ -4035,7 +4032,6 @@ void WriteEventFunctionEpilogue(FOutputDevice& Output, int32 Indent, const FParm
 		bool bBoolProperty = Parameters.Return->IsBooleanOrBooleanStaticArray();
 		Output.Logf(TEXT("%sreturn %sParms.%s;\r\n"), FCString::Tab(Indent + 1), bBoolProperty ? TEXT("!!") : TEXT(""), *Parameters.Return->GetName());
 	}
-	Output.Logf(TEXT("%s}\r\n"), FCString::Tab(Indent));
 }
 
 void FNativeClassHeaderGenerator::ExportDelegateDeclaration(FOutputDevice& Out, FReferenceGatherers& OutReferenceGatherers, const FUnrealSourceFile& SourceFile, FUnrealFunctionDefinitionInfo& FunctionDef) const
@@ -4064,15 +4060,26 @@ void FNativeClassHeaderGenerator::ExportDelegateDeclaration(FOutputDevice& Out,
 	);
 
 	FUHTStringBuilder DelegateOutput;
-	DelegateOutput.Log(TEXT("static "));
 
 	// export the line that looks like: int32 Main(const FString& Parms)
-	ExportNativeFunctionHeader(DelegateOutput, OutReferenceGatherers.ForwardDeclarations, FunctionDef, FunctionData, EExportFunctionType::Event, EExportFunctionHeaderStyle::Declaration, *ExtraParam, *GetAPIString());
+	ExportNativeFunctionHeader(DelegateOutput, OutReferenceGatherers.ForwardDeclarations, FunctionDef, FunctionData, EExportFunctionType::Event, EExportFunctionHeaderStyle::Definition, *ExtraParam, *GetAPIString());
 
-	// Only exporting function prototype
-	DelegateOutput.Logf(TEXT(";\r\n"));
+	FParmsAndReturnProperties Parameters = GetFunctionParmsAndReturn(FunctionDef);
+
+	DelegateOutput.Logf(TEXT("\r\n%s{\r\n"), FCString::Tab(0));
+	ExportEventParm(DelegateOutput, OutReferenceGatherers.ForwardDeclarations, FunctionDef, /*Indent=*/ 1, /*bOutputConstructor=*/ true, EExportingState::Normal);
+	WriteEventFunctionPrologue(DelegateOutput, 0, Parameters, *FunctionDef.GetOuter(), *DelegateName);
+	{
+		const TCHAR* DelegateType = bIsMulticastDelegate ? TEXT("ProcessMulticastDelegate") : TEXT("ProcessDelegate");
+		const TCHAR* DelegateArg = Parameters.HasParms() ? TEXT("&Parms") : TEXT("NULL");
+		DelegateOutput.Logf(TEXT("\t%s.%s<UObject>(%s);\r\n"), *DelegateName, DelegateType, DelegateArg);
+	}
+	WriteEventFunctionEpilogue(DelegateOutput, 0, Parameters);
+	DelegateOutput.Logf(TEXT("%s}\r\n"), FCString::Tab(0));
 
 	ExportFunction(Out, OutReferenceGatherers, SourceFile, FunctionDef, false);
+
+	Out.Log(DelegateOutput);
 }
 
 void FNativeClassHeaderGenerator::ExportDelegateDefinition(FOutputDevice& Out, FReferenceGatherers& OutReferenceGatherers, const FUnrealSourceFile& SourceFile, FUnrealFunctionDefinitionInfo& FunctionDef) const
@@ -4083,18 +4090,11 @@ void FNativeClassHeaderGenerator::ExportDelegateDefinition(FOutputDevice& Out, F
 
 	check(FunctionDef.HasAnyFunctionFlags(FUNC_Delegate));
 
-	// Export parameters structs for all delegates.  We'll need these to declare our delegate execution function.
-	FUHTStringBuilder DelegateOutput;
-	ExportEventParm(DelegateOutput, OutReferenceGatherers.ForwardDeclarations, FunctionDef, /*Indent=*/ 0, /*bOutputConstructor=*/ true, EExportingState::Normal);
-
 	const bool bIsMulticastDelegate = FunctionDef.HasAnyFunctionFlags( FUNC_MulticastDelegate );
 
 	// Unmangle the function name
 	const FString DelegateName = FunctionDef.GetName().LeftChop(HEADER_GENERATED_DELEGATE_SIGNATURE_SUFFIX_LENGTH);
 
-	// Always export delegate wrapper functions as inline
-	FunctionData.FunctionExportFlags |= FUNCEXPORT_Inline;
-
 	// Add class name to beginning of function, to avoid collisions with other classes with the same delegate name in this scope
 	check(FunctionData.MarshallAndCallName.StartsWith(DelegateStr));
 	FString ShortName = *FunctionData.MarshallAndCallName + UE_ARRAY_COUNT(DelegateStr) - 1;
@@ -4107,20 +4107,29 @@ void FNativeClassHeaderGenerator::ExportDelegateDefinition(FOutputDevice& Out, F
 		*DelegateName
 	);
 
-	DelegateOutput.Log(TEXT("static "));
+	bool bAddAPI = true;
 
-	// export the line that looks like: int32 Main(const FString& Parms)
-	ExportNativeFunctionHeader(DelegateOutput, OutReferenceGatherers.ForwardDeclarations, FunctionDef, FunctionData, EExportFunctionType::Event, EExportFunctionHeaderStyle::Declaration, *ExtraParam, *GetAPIString());
+	FUHTStringBuilder DelegateOutput;
+	if (FUnrealClassDefinitionInfo* ClassDef = FunctionDef.GetOwnerClass())
+	{
+		DelegateOutput.Log(TEXT("static "));
+		if (ClassDef->HasAnyClassFlags(CLASS_RequiredAPI))
+		{
+			bAddAPI = false;
+		}
+	}
 
-	FParmsAndReturnProperties Parameters = GetFunctionParmsAndReturn(FunctionDef);
+	//if ((FunctionData.FunctionExportFlags & FUNCEXPORT_RequiredAPI) == 0)  // TODO: This requires too much fixup for now
+	//	bAddAPI = false;
 
-	WriteEventFunctionPrologue(DelegateOutput, 0, Parameters, *FunctionDef.GetOuter(), *DelegateName);
+	if (bAddAPI)
 	{
-		const TCHAR* DelegateType = bIsMulticastDelegate ? TEXT( "ProcessMulticastDelegate" ) : TEXT( "ProcessDelegate" );
-		const TCHAR* DelegateArg  = Parameters.HasParms() ? TEXT("&Parms") : TEXT("NULL");
-		DelegateOutput.Logf(TEXT("\t%s.%s<UObject>(%s);\r\n"), *DelegateName, DelegateType, DelegateArg);
+		DelegateOutput.Log(*GetAPIString());
 	}
-	WriteEventFunctionEpilogue(DelegateOutput, 0, Parameters);
+
+	// export the line that looks like: int32 Main(const FString& Parms)
+	ExportNativeFunctionHeader(DelegateOutput, OutReferenceGatherers.ForwardDeclarations, FunctionDef, FunctionData, EExportFunctionType::Event, EExportFunctionHeaderStyle::Declaration, *ExtraParam, TEXT(""));
+	DelegateOutput.Logf(TEXT(";\r\n"));
 
 	FString MacroName = SourceFile.GetGeneratedMacroName(FunctionData.MacroLine, TEXT("_DELEGATE"));
 	WriteMacro(Out, MacroName, DelegateOutput);
@@ -4497,7 +4506,10 @@ void FNativeClassHeaderGenerator::ExportNativeFunctionHeader(
 	FString FunctionName;
 	if (FunctionHeaderStyle == EExportFunctionHeaderStyle::Definition)
 	{
-		FunctionName = FString::Printf(TEXT("%s::"), *UHTCastChecked<FUnrealClassDefinitionInfo>(FunctionDef.GetOuter()).GetAlternateNameCPP(bIsInterface || FunctionType == EExportFunctionType::Interface));
+		if (FUnrealClassDefinitionInfo* OuterClass = UHTCast<FUnrealClassDefinitionInfo>(FunctionDef.GetOuter()))
+		{
+			FunctionName = FString::Printf(TEXT("%s::"), *OuterClass->GetAlternateNameCPP(bIsInterface || FunctionType == EExportFunctionType::Interface));
+		}
 	}
 
 	if (FunctionType == EExportFunctionType::Interface)
@@ -5470,6 +5482,7 @@ void FNativeClassHeaderGenerator::ExportCallbackFunctions(
 
 		if (ExportCallbackType != EExportCallbackType::Interface)
 		{
+			OutCpp.Logf(TEXT("\r\n%s{\r\n"), FCString::Tab(1));
 			WriteEventFunctionPrologue(OutCpp, /*Indent=*/ 1, Parameters, ClassDef, *FunctionName);
 			{
 				// Cast away const just in case, because ProcessEvent isn't const
@@ -5481,6 +5494,7 @@ void FNativeClassHeaderGenerator::ExportCallbackFunctions(
 				);
 			}
 			WriteEventFunctionEpilogue(OutCpp, /*Indent=*/ 1, Parameters);
+			OutCpp.Logf(TEXT("%s}\r\n"), FCString::Tab(1));
 		}
 		else
 		{
