Fixed infinite loop in C++ UHT recursive include checks.

diff --git a/Engine/Source/Programs/UnrealHeaderTool/Private/CodeGenerator.cpp b/Engine/Source/Programs/UnrealHeaderTool/Private/CodeGenerator.cpp
--- a/Engine/Source/Programs/UnrealHeaderTool/Private/CodeGenerator.cpp
+++ b/Engine/Source/Programs/UnrealHeaderTool/Private/CodeGenerator.cpp
@@ -7015,7 +7015,7 @@ void PrepareTypesForParsing(TArray<FUnrealPackageDefinitionInfo*>& PackageDefs)
 	// Does nothing now
 }
 
-void TopologicalRecursion(FUnrealSourceFile& First, FUnrealSourceFile& Visit)
+void TopologicalRecursion(FUnrealSourceFile& First, FUnrealSourceFile& Visit, TSet<FUnrealSourceFile*> SignaledFiles)
 {
 	check(Visit.GetTopologicalState() == ETopologicalState::Temporary);
 	for (FHeaderProvider& Header : Visit.GetIncludes())
@@ -7024,10 +7024,15 @@ void TopologicalRecursion(FUnrealSourceFile& First, FUnrealSourceFile& Visit)
 		{
 			if (Include->GetTopologicalState() == ETopologicalState::Temporary)
 			{
+				if (SignaledFiles.Contains(Include))
+				{
+					break;
+				}
+				SignaledFiles.Add(Include);
 				UE_LOG(LogCompile, Error, TEXT("%s includes/requires %s"), *Visit.GetFilename(), *Include->GetFilename());
 				if (&First != Include)
 				{
-					TopologicalRecursion(First, *Include);
+					TopologicalRecursion(First, *Include, SignaledFiles);
 				}
 				break;
 			}
@@ -7081,7 +7086,8 @@ void TopologicalSort(TArray<FUnrealSourceFile*>& OrderedSourceFiles)
 			if (FUnrealSourceFile* Recusion = TopologicalVisit(OrderedSourceFiles, *SourceFile))
 			{
 				UE_LOG(LogCompile, Error, TEXT("Circular dependency detected:"));
-				TopologicalRecursion(*Recusion, *Recusion);
+				TSet<FUnrealSourceFile*> SignaledFiles;
+				TopologicalRecursion(*Recusion, *Recusion, SignaledFiles);
 				FResults::SetResult(ECompilationResult::OtherCompilationError);
 				bCircularDependencyDetected = true;
 			}
