UHT: Fix RigVM related compat issues

diff --git a/Engine/Source/Programs/UnrealHeaderTool/Private/CodeGenerator.cpp b/Engine/Source/Programs/UnrealHeaderTool/Private/CodeGenerator.cpp
--- a/Engine/Source/Programs/UnrealHeaderTool/Private/CodeGenerator.cpp
+++ b/Engine/Source/Programs/UnrealHeaderTool/Private/CodeGenerator.cpp
@@ -3409,7 +3409,7 @@ void FNativeClassHeaderGenerator::ExportGeneratedStructBodyMacros(FOutputDevice&
 	const FString& ChoppedSingletonName = ScriptStructDef.GetSingletonNameChopped(true);
 
 	const FString RigVMExecuteContextDeclaration = TEXT("FRigVMExtendedExecuteContext& RigVMExecuteContext");
-	const FString RigVMExecuteContextPublicDeclaration = FString::Printf(TEXT("const %s& RigVMExecuteContext"), *ScriptStructDef.GetRigVMInfo().ExecuteContextType);
+	const FString RigVMExecuteContextPublicDeclaration = FString::Printf(TEXT("%s& ExecuteContext"), *ScriptStructDef.GetRigVMInfo().ExecuteContextType);
 	TArray<FString> RigVMVirtualFuncProlog, RigVMVirtualFuncEpilog, RigVMStubProlog;
 
 	// for RigVM methods we need to generated a macro used for implementing the static method
@@ -3475,12 +3475,12 @@ void FNativeClassHeaderGenerator::ExportGeneratedStructBodyMacros(FOutputDevice&
 			}
 		}
 
-		FString StructMembers = StructRigVMInfo.Members.Declarations(false, TEXT(", \\\r\n\t\t"), true, false);
+		FString StructMembers = StructRigVMInfo.Members.Declarations(false, TEXT(", \\\r\n\t\t"), true, false, false, false);
 
 		OutGeneratedHeaderText.Log(TEXT("\r\n"));
 		for (const FRigVMMethodInfo& MethodInfo : StructRigVMInfo.Methods)
 		{
-			FString ParameterSuffix = MethodInfo.Parameters.Declarations(true, TEXT(", \\\r\n\t\t"));
+			FString ParameterSuffix = MethodInfo.Parameters.Declarations(!StructMembers.IsEmpty(), TEXT(", \\\r\n\t\t"));
 			FString RigVMParameterPrefix2 = RigVMExecuteContextPublicDeclaration + FString((StructMembers.IsEmpty() && ParameterSuffix.IsEmpty()) ? TEXT("") : TEXT(", \\\r\n\t\t"));
 			OutGeneratedHeaderText.Logf(TEXT("#define %s_%s() \\\r\n"), *StructNameCPP, *MethodInfo.Name);
 			OutGeneratedHeaderText.Logf(TEXT("\t%s %s::Static%s( \\\r\n\t\t%s%s%s \\\r\n\t)\r\n"), *MethodInfo.ReturnType, *StructNameCPP, *MethodInfo.Name, *RigVMParameterPrefix2, *StructMembers, *ParameterSuffix);
@@ -3503,12 +3503,12 @@ void FNativeClassHeaderGenerator::ExportGeneratedStructBodyMacros(FOutputDevice&
 		FString RigVMMethodsDeclarations;
 		if (StructRigVMInfo.bHasRigVM)
 		{
-			FString StructMembers = StructRigVMInfo.Members.Declarations(false, TEXT(",\r\n\t\t"), true, false);
+			FString StructMembers = StructRigVMInfo.Members.Declarations(false, TEXT(",\r\n\t\t"), true, false, true, false);
 			for (const FRigVMMethodInfo& MethodInfo : StructRigVMInfo.Methods)
 			{
-				FString StructMembersForStub = StructRigVMInfo.Members.Names(false, TEXT(",\r\n\t\t\t"), false);
-				FString ParameterSuffix = MethodInfo.Parameters.Declarations(true, TEXT(",\r\n\t\t"));
-				FString ParameterNamesSuffix = MethodInfo.Parameters.Names(true, TEXT(",\r\n\t\t\t"));
+				FString StructMembersForStub = StructRigVMInfo.Members.Names(false, TEXT(",\r\n\t\t\t"), false, true, false);
+				FString ParameterSuffix = MethodInfo.Parameters.Declarations(!StructMembers.IsEmpty(), TEXT(",\r\n\t\t"));
+				FString ParameterNamesSuffix = MethodInfo.Parameters.Names(!StructMembersForStub.IsEmpty(), TEXT(",\r\n\t\t\t"));
 				FString RigVMParameterPrefix2 = RigVMExecuteContextPublicDeclaration + FString((StructMembers.IsEmpty() && ParameterSuffix.IsEmpty()) ? TEXT("") : TEXT(",\r\n\t\t"));
 				FString RigVMParameterPrefix4 = FString(TEXT("RigVMExecuteContext.GetPublicData<")) +
 					StructRigVMInfo.ExecuteContextType +
@@ -3766,14 +3766,14 @@ void FNativeClassHeaderGenerator::ExportGeneratedStructBodyMacros(FOutputDevice&
 	// The static method is implemented by the user using a macro.
 	if (StructRigVMInfo.bHasRigVM)
 	{
-		FString StructMembersForVirtualFunc = StructRigVMInfo.Members.Names(false, TEXT(",\r\n\t\t"), true);
+		FString StructMembersForVirtualFunc = StructRigVMInfo.Members.Names(false, TEXT(",\r\n\t\t"), true, true, false);
 
 		for (const FRigVMMethodInfo& MethodInfo : StructRigVMInfo.Methods)
 		{
 			Out.Log(TEXT("\r\n"));
 
 			FString ParameterDeclaration = MethodInfo.Parameters.Declarations(false, TEXT(",\r\n\t\t"));
-			FString ParameterSuffix = MethodInfo.Parameters.Names(true, TEXT(",\r\n\t\t"));
+			FString ParameterSuffix = MethodInfo.Parameters.Names(!StructMembersForVirtualFunc.IsEmpty(), TEXT(",\r\n\t\t"));
 			FString RigVMParameterPrefix3 = FString(TEXT("TemporaryExecuteContext")) + FString((StructMembersForVirtualFunc.IsEmpty() && ParameterSuffix.IsEmpty()) ? TEXT("") : TEXT(",\r\n\t\t"));
 
 			// implement the virtual function body.
diff --git a/Engine/Source/Programs/UnrealHeaderTool/Private/HeaderParser.cpp b/Engine/Source/Programs/UnrealHeaderTool/Private/HeaderParser.cpp
--- a/Engine/Source/Programs/UnrealHeaderTool/Private/HeaderParser.cpp
+++ b/Engine/Source/Programs/UnrealHeaderTool/Private/HeaderParser.cpp
@@ -6728,17 +6728,18 @@ void FHeaderParser::ParseRigVMMethodParameters(FUnrealStructDefinitionInfo& Stru
 
 		if (Parameter.IsExecuteContext())
 		{
+			const FString ExecuteContextName = Parameter.GetExecuteContextTypeName();
 			if(StructRigVMInfo.ExecuteContextMember.IsEmpty())
 			{
 				StructRigVMInfo.ExecuteContextMember = Parameter.Name;
 			}
 			if(StructRigVMInfo.ExecuteContextType == TEXT("FRigVMExecuteContext"))
 			{
-				StructRigVMInfo.ExecuteContextType = Parameter.Type;
+				StructRigVMInfo.ExecuteContextType = ExecuteContextName;
 			}
-			else if(StructRigVMInfo.ExecuteContextType != Parameter.Type)
+			else if(StructRigVMInfo.ExecuteContextType != ExecuteContextName)
 			{
-				LogError(TEXT("RigVM Struct '%s' contains properties of varying execute context type %s vs %s."), *StructDef.GetName(), *StructRigVMInfo.ExecuteContextType, *Parameter.Type);
+				LogError(TEXT("RigVM Struct '%s' contains properties of varying execute context type %s vs %s."), *StructDef.GetName(), *StructRigVMInfo.ExecuteContextType, *ExecuteContextName);
 			}
 		}
 		else
diff --git a/Engine/Source/Programs/UnrealHeaderTool/Private/ParserHelper.cpp b/Engine/Source/Programs/UnrealHeaderTool/Private/ParserHelper.cpp
--- a/Engine/Source/Programs/UnrealHeaderTool/Private/ParserHelper.cpp
+++ b/Engine/Source/Programs/UnrealHeaderTool/Private/ParserHelper.cpp
@@ -424,23 +424,28 @@ bool FAdvancedDisplayParameterHandler::CanMarkMore() const
 
 bool FRigVMParameter::IsExecuteContext() const
 {
-	if(PropertyDef)
+	if(PropertyDef && PropertyDef->IsStructOrStructStaticArray())
 	{
-		if(const FProperty* Property = PropertyDef->GetPropertySafe())
+		if(const FUnrealStructDefinitionInfo* StructDef = PropertyDef->GetPropertyBase().ScriptStructDef)
 		{
-			if(const FStructProperty* StructProperty = CastField<FStructProperty>(Property))
+			while(StructDef != nullptr)
 			{
-				const UScriptStruct* Struct = StructProperty->Struct;
-				while(Struct != nullptr)
+				if(StructDef->GetNameCPP() == TEXT("FRigVMExecuteContext"))
 				{
-					if(Struct->GetStructCPPName() == TEXT("FRigVMExecuteContext"))
-					{
-						return true;
-					}
-					Struct = Cast<UScriptStruct>(Struct->GetSuperStruct());
+					return true;
 				}
+				StructDef = StructDef->GetSuperStruct();
 			}
 		}
 	}
 	return false;
 }
+
+FString FRigVMParameter::GetExecuteContextTypeName() const
+{
+	if(IsExecuteContext())
+	{
+		return PropertyDef->GetPropertyBase().ScriptStructDef->GetNameCPP();
+	}
+	return FString();
+}
diff --git a/Engine/Source/Programs/UnrealHeaderTool/Private/ParserHelper.h b/Engine/Source/Programs/UnrealHeaderTool/Private/ParserHelper.h
--- a/Engine/Source/Programs/UnrealHeaderTool/Private/ParserHelper.h
+++ b/Engine/Source/Programs/UnrealHeaderTool/Private/ParserHelper.h
@@ -640,6 +640,7 @@ struct FRigVMParameter
 	}
 
 	bool IsExecuteContext() const;
+	FString GetExecuteContextTypeName() const; 
 };
 
 /**
@@ -663,7 +664,7 @@ public:
 		return Parameters.Add(InParameter);
 	}
 
-	FString Names(bool bLeadingSeparator = false, const TCHAR* Separator = TEXT(", "), bool bCastType = false, bool bIncludeEditorOnly = true) const
+	FString Names(bool bLeadingSeparator = false, const TCHAR* Separator = TEXT(", "), bool bCastType = false, bool bIncludeEditorOnly = true, bool bIncludeExecuteContext = true) const
 	{
 		if (Parameters.Num() == 0)
 		{
@@ -676,6 +677,10 @@ public:
 			{
 				continue;
 			}
+			if (!bIncludeExecuteContext && Parameter.IsExecuteContext())
+			{
+				continue;
+			}
 			NameArray.Add(Parameter.NameOriginal(bCastType));
 		}
 
@@ -692,7 +697,7 @@ public:
 		return Joined;
 	}
 
-	FString Declarations(bool bLeadingSeparator = false, const TCHAR* Separator = TEXT(", "), bool bCastType = false, bool bCastName = false, bool bIncludeEditorOnly = true) const
+	FString Declarations(bool bLeadingSeparator = false, const TCHAR* Separator = TEXT(", "), bool bCastType = false, bool bCastName = false, bool bIncludeEditorOnly = true, bool bIncludeExecuteContext = true) const
 	{
 		if (Parameters.Num() == 0)
 		{
@@ -705,6 +710,10 @@ public:
 			{
 				continue;
 			}
+			if (!bIncludeExecuteContext && Parameter.IsExecuteContext())
+			{
+				continue;
+			}
 			DeclarationArray.Add(Parameter.Variable(bCastType, bCastName));
 		}
 
diff --git a/Engine/Source/Programs/UnrealHeaderTool/Resources/UHTDebugging/RigVMStructHeader.h b/Engine/Source/Programs/UnrealHeaderTool/Resources/UHTDebugging/RigVMStructHeader.h
--- a/Engine/Source/Programs/UnrealHeaderTool/Resources/UHTDebugging/RigVMStructHeader.h
+++ b/Engine/Source/Programs/UnrealHeaderTool/Resources/UHTDebugging/RigVMStructHeader.h
@@ -1,6 +1,7 @@
 // Copyright Epic Games, Inc. All Rights Reserved.
 
 #pragma once
+#include "RigVMCore/RigVMExecuteContext.h"
 
 #include "RigVMStructHeader.generated.h"
 
@@ -39,6 +40,18 @@ struct FRigVMStructBase
 	virtual FRigVMStructUpgradeInfo GetUpgradeInfo() const {};
 };
 
+USTRUCT(BlueprintType)
+struct FRigVMExecuteContext
+{
+	GENERATED_BODY()
+};
+
+USTRUCT(BlueprintType)
+struct FUHTTestExecuteContext : public FRigVMExecuteContext
+{
+	GENERATED_BODY()
+};
+
 USTRUCT(meta = (Deprecated = "5.0.0"))
 struct FRigVMMethodStruct : public FRigVMStructBase
 {
@@ -116,6 +129,36 @@ struct CONTROLRIG_API FRigUnit_Control
 	}
 
 	/** The transform of this control */
-	UPROPERTY(EditAnywhere, Category="Control")
+	UPROPERTY(EditAnywhere, Category="Control", meta=(Input))
 	float Factor;
+};
+
+USTRUCT()
+struct CONTROLRIG_API FRigUnit_Mutable
+{
+	GENERATED_BODY()
+
+	FRigUnit_Mutable()
+	{
+	}
+
+	RIGVM_METHOD()
+	virtual void Execute(float Factor);
+
+	/** The transform of this control */
+	UPROPERTY(EditAnywhere, Category="Control", meta=(Input, Output))
+	FUHTTestExecuteContext ExecuteContext;
+};
+
+USTRUCT()
+struct CONTROLRIG_API FRigUnit_BeginExecution
+{
+	GENERATED_BODY()
+
+	RIGVM_METHOD()
+	virtual void Execute(const FRigVMExecuteContext& Context) override;
+
+	// The execution result
+	UPROPERTY(EditAnywhere, Transient, DisplayName = "Execute", Category = "BeginExecution", meta = (Output))
+	FUHTTestExecuteContext ExecuteContext;
 };
\ No newline at end of file
