Fixed C# UHT to properly generate an error when a function's argument name shadows a property in the class or super classes.
Fixed C++ UHT to no longer check for parameter name shadowing for delegates.

diff --git a/Engine/Source/Programs/Shared/EpicGames.UHT/Types/UhtFunction.cs b/Engine/Source/Programs/Shared/EpicGames.UHT/Types/UhtFunction.cs
--- a/Engine/Source/Programs/Shared/EpicGames.UHT/Types/UhtFunction.cs
+++ b/Engine/Source/Programs/Shared/EpicGames.UHT/Types/UhtFunction.cs
@@ -752,6 +752,9 @@ namespace EpicGames.UHT.Types
 							this.LogError($"Can't override delegate signature function '{SourceName}'");
 						}
 					}
+
+					// Disable shadow checks for delegates
+					options &= ~UhtValidationOptions.Shadowing;
 					break;
 			}
 
diff --git a/Engine/Source/Programs/Shared/EpicGames.UHT/Utils/UhtSymbolTable.cs b/Engine/Source/Programs/Shared/EpicGames.UHT/Utils/UhtSymbolTable.cs
--- a/Engine/Source/Programs/Shared/EpicGames.UHT/Utils/UhtSymbolTable.cs
+++ b/Engine/Source/Programs/Shared/EpicGames.UHT/Utils/UhtSymbolTable.cs
@@ -265,7 +265,21 @@ namespace EpicGames.UHT.Utils
 			// If requested, skip self
 			if (options.HasAnyFlags(UhtFindOptions.NoSelf) && currentType == startingType)
 			{
-				currentType = ((UhtStruct)currentType).Super;
+				if (currentType is UhtFunction)
+				{
+					if (currentType.Outer is UhtStruct outerStructType)
+					{
+						currentType = outerStructType;
+					}
+					else
+					{
+						currentType = null;
+					}
+				}
+				else
+				{
+					currentType = ((UhtStruct)currentType).Super;
+				}
 			}
 
 			// Search the chain
diff --git a/Engine/Source/Programs/UnrealHeaderTool/Private/HeaderParser.cpp b/Engine/Source/Programs/UnrealHeaderTool/Private/HeaderParser.cpp
--- a/Engine/Source/Programs/UnrealHeaderTool/Private/HeaderParser.cpp
+++ b/Engine/Source/Programs/UnrealHeaderTool/Private/HeaderParser.cpp
@@ -4874,40 +4874,43 @@ FUnrealPropertyDefinitionInfo& FHeaderParser::GetVarNameAndDim
 
 	// Make sure it doesn't conflict.
 	FString VarName(Identifier);
-	int32 OuterContextCount = 0;
-	FUnrealFunctionDefinitionInfo* ExistingFunctionDef = FindFunction(ParentStruct, *VarName, true, nullptr);
-	FUnrealPropertyDefinitionInfo* ExistingPropertyDef = FindProperty(ParentStruct, *VarName, true);
-
-	if (ExistingFunctionDef != nullptr || ExistingPropertyDef != nullptr)
+	FUnrealFunctionDefinitionInfo* ParentFunction = UHTCast<FUnrealFunctionDefinitionInfo>(&ParentStruct);
+	if (VariableCategory == EVariableCategory::Member || (ParentFunction != nullptr && ParentFunction->GetFunctionType() == EFunctionType::Function))
 	{
-		bool bErrorDueToShadowing = true;
+		FUnrealFunctionDefinitionInfo* ExistingFunctionDef = FindFunction(ParentStruct, *VarName, true, nullptr);
+		FUnrealPropertyDefinitionInfo* ExistingPropertyDef = FindProperty(ParentStruct, *VarName, true);
 
-		if (ExistingFunctionDef && (VariableCategory != EVariableCategory::Member))
+		if (ExistingFunctionDef != nullptr || ExistingPropertyDef != nullptr)
 		{
-			// A function parameter with the same name as a method is allowed
-			bErrorDueToShadowing = false;
-		}
+			bool bErrorDueToShadowing = true;
 
-		//@TODO: This exception does not seem sound either, but there is enough existing code that it will need to be
-		// fixed up first before the exception it is removed.
-		if (ExistingPropertyDef)
- 		{
- 			const bool bExistingPropDeprecated = ExistingPropertyDef->HasAnyPropertyFlags(CPF_Deprecated);
- 			const bool bNewPropDeprecated = (VariableCategory == EVariableCategory::Member) && ((VarProperty.PropertyFlags & CPF_Deprecated) != 0);
- 			if (bNewPropDeprecated || bExistingPropDeprecated)
- 			{
- 				// if this is a property and one of them is deprecated, ignore it since it will be removed soon
- 				bErrorDueToShadowing = false;
- 			}
- 		}
+			if (ExistingFunctionDef && (VariableCategory != EVariableCategory::Member))
+			{
+				// A function parameter with the same name as a method is allowed
+				bErrorDueToShadowing = false;
+			}
 
-		if (bErrorDueToShadowing)
-		{
-			Throwf(TEXT("%s: '%s' cannot be defined in '%s' as it is already defined in scope '%s' (shadowing is not allowed)"),
-				HintText,
-				*FString(Identifier),
-				*ParentStruct.GetName(),
-				ExistingFunctionDef ? *ExistingFunctionDef->GetOuter()->GetName() : *ExistingPropertyDef->GetFullName());
+			//@TODO: This exception does not seem sound either, but there is enough existing code that it will need to be
+			// fixed up first before the exception it is removed.
+			if (ExistingPropertyDef)
+			{
+				const bool bExistingPropDeprecated = ExistingPropertyDef->HasAnyPropertyFlags(CPF_Deprecated);
+				const bool bNewPropDeprecated = (VariableCategory == EVariableCategory::Member) && ((VarProperty.PropertyFlags & CPF_Deprecated) != 0);
+				if (bNewPropDeprecated || bExistingPropDeprecated)
+				{
+					// if this is a property and one of them is deprecated, ignore it since it will be removed soon
+					bErrorDueToShadowing = false;
+				}
+			}
+
+			if (bErrorDueToShadowing)
+			{
+				Throwf(TEXT("%s: '%s' cannot be defined in '%s' as it is already defined in scope '%s' (shadowing is not allowed)"),
+					HintText,
+					*FString(Identifier),
+					*ParentStruct.GetName(),
+					ExistingFunctionDef ? *ExistingFunctionDef->GetOuter()->GetName() : *ExistingPropertyDef->GetFullName());
+			}
 		}
 	}
 
