C++ UHT: Allow True and False as identifiers for variables

This is already supported by the new header tool.

diff --git a/Engine/Source/Programs/UnrealHeaderTool/Private/HeaderParser.cpp b/Engine/Source/Programs/UnrealHeaderTool/Private/HeaderParser.cpp
--- a/Engine/Source/Programs/UnrealHeaderTool/Private/HeaderParser.cpp
+++ b/Engine/Source/Programs/UnrealHeaderTool/Private/HeaderParser.cpp
@@ -4810,11 +4810,27 @@ FUnrealPropertyDefinitionInfo& FHeaderParser::GetVarNameAndDim
 	else
 	{
 		FToken VarToken;
-		if (!GetIdentifier(VarToken))
+		if (!GetToken(VarToken, false))
 		{
 			Throwf(TEXT("Missing variable name") );
 		}
 
+		// allow True and False as identifiers for Variables
+		if(VarToken.TokenType == ETokenType::TrueConst || VarToken.TokenType == ETokenType::FalseConst)
+		{
+			const FString FirstChar(1, VarToken.Value.GetData());
+			if(FirstChar.ToUpper().Equals(FirstChar, ESearchCase::CaseSensitive))
+			{
+				VarToken.TokenType = ETokenType::Identifier;
+			}
+		}
+
+		if(!VarToken.IsIdentifier())
+		{
+			UngetToken(VarToken);
+			Throwf(TEXT("Missing variable name") );
+		}
+
 		switch (LayoutMacroType)
 		{
 			case ELayoutMacroType::Array:
diff --git a/Engine/Source/Programs/UnrealHeaderTool/Resources/UHTDebugging/RigVMStructHeader.h b/Engine/Source/Programs/UnrealHeaderTool/Resources/UHTDebugging/RigVMStructHeader.h
--- a/Engine/Source/Programs/UnrealHeaderTool/Resources/UHTDebugging/RigVMStructHeader.h
+++ b/Engine/Source/Programs/UnrealHeaderTool/Resources/UHTDebugging/RigVMStructHeader.h
@@ -183,3 +183,49 @@ struct FRigVMLazyStruct
 	UPROPERTY(meta = (Output))
 	float Result;
 };
+
+/**
+ * Executes either the True or False branch based on the condition
+ */
+USTRUCT(meta = (DisplayName = "Branch", Keywords = "If"))
+struct RIGVM_API FRigVMFunction_ControlFlowBranch
+{
+	GENERATED_BODY()
+
+	FRigVMFunction_ControlFlowBranch()
+	{
+		Condition = false;
+		BlockToRun = NAME_None;
+	}
+
+	RIGVM_METHOD()
+	virtual void Execute();
+
+	virtual const TArray<FName>& GetControlFlowBlocks_Impl() const override
+	{
+		static const TArray<FName> Blocks = {
+			GET_MEMBER_NAME_CHECKED(FRigVMFunction_ControlFlowBranch, True),
+			GET_MEMBER_NAME_CHECKED(FRigVMFunction_ControlFlowBranch, False),
+			ForLoopCompletedPinName
+		};
+		return Blocks;
+	}
+
+	UPROPERTY(meta=(Input, DisplayName="Execute"))
+	FRigVMExecuteContext ExecuteContext;
+
+	UPROPERTY(meta=(Input))
+	bool Condition;
+
+	UPROPERTY(meta=(Output))
+	FRigVMExecuteContext True ;
+
+	UPROPERTY(meta=(Output))
+	FRigVMExecuteContext False;
+
+	UPROPERTY(meta=(Output))
+	FRigVMExecuteContext Completed;
+
+	UPROPERTY(meta=(Singleton))
+	FName BlockToRun;
+};
