RigVM: Remove support for opaque attributes

diff --git a/Engine/Source/Programs/UnrealHeaderTool/Private/CodeGenerator.cpp b/Engine/Source/Programs/UnrealHeaderTool/Private/CodeGenerator.cpp
--- a/Engine/Source/Programs/UnrealHeaderTool/Private/CodeGenerator.cpp
+++ b/Engine/Source/Programs/UnrealHeaderTool/Private/CodeGenerator.cpp
@@ -3506,10 +3506,13 @@ void FNativeClassHeaderGenerator::ExportGeneratedStructBodyMacros(FOutputDevice&
 		OutGeneratedHeaderText.Log(TEXT("\r\n"));
 		for (const FRigVMMethodInfo& MethodInfo : StructRigVMInfo.Methods)
 		{
-			FString ParameterSuffix = MethodInfo.Parameters.Declarations(!StructMembers.IsEmpty(), TEXT(", \\\r\n\t\t"));
-			FString RigVMParameterPrefix2 = RigVMExecuteContextPublicDeclaration + FString((StructMembers.IsEmpty() && ParameterSuffix.IsEmpty()) ? TEXT("") : TEXT(", \\\r\n\t\t"));
+			FString RigVMParameterPrefix2 = RigVMExecuteContextPublicDeclaration + FString(StructMembers.IsEmpty() ? TEXT("") : TEXT(", \\\r\n\t\t"));
+			if(StructRigVMInfo.ExecuteContextMember.IsEmpty())
+			{
+				RigVMParameterPrefix2 = TEXT("const ") + RigVMParameterPrefix2;
+			}
 			OutGeneratedHeaderText.Logf(TEXT("#define %s_%s() \\\r\n"), *StructNameCPP, *MethodInfo.Name);
-			OutGeneratedHeaderText.Logf(TEXT("\t%s %s::Static%s( \\\r\n\t\t%s%s%s \\\r\n\t)\r\n"), *MethodInfo.ReturnType, *StructNameCPP, *MethodInfo.Name, *RigVMParameterPrefix2, *StructMembers, *ParameterSuffix);
+			OutGeneratedHeaderText.Logf(TEXT("\t%s %s::Static%s( \\\r\n\t\t%s%s \\\r\n\t)\r\n"), *MethodInfo.ReturnType, *StructNameCPP, *MethodInfo.Name, *RigVMParameterPrefix2, *StructMembers);
 		}
 		OutGeneratedHeaderText.Log(TEXT("\r\n"));
 	}
@@ -3533,30 +3536,21 @@ void FNativeClassHeaderGenerator::ExportGeneratedStructBodyMacros(FOutputDevice&
 			for (const FRigVMMethodInfo& MethodInfo : StructRigVMInfo.Methods)
 			{
 				FString StructMembersForStub = StructRigVMInfo.Members.Names(false, TEXT(",\r\n\t\t\t"), false, true, false);
-				FString ParameterSuffix = MethodInfo.Parameters.Declarations(!StructMembers.IsEmpty(), TEXT(",\r\n\t\t"));
-				FString ParameterNamesSuffix = MethodInfo.Parameters.Names(!StructMembersForStub.IsEmpty(), TEXT(",\r\n\t\t\t"));
-				FString RigVMParameterPrefix2 = RigVMExecuteContextPublicDeclaration + FString((StructMembers.IsEmpty() && ParameterSuffix.IsEmpty()) ? TEXT("") : TEXT(",\r\n\t\t"));
+				FString RigVMParameterPrefix2 = RigVMExecuteContextPublicDeclaration + FString(StructMembers.IsEmpty() ? TEXT("") : TEXT(",\r\n\t\t"));
+				if(StructRigVMInfo.ExecuteContextMember.IsEmpty())
+				{
+					RigVMParameterPrefix2 = TEXT("const ") + RigVMParameterPrefix2;
+				}
 				FString RigVMParameterPrefix4 = FString(TEXT("RigVMExecuteContext.GetPublicData<")) +
 					StructRigVMInfo.ExecuteContextType +
 					FString(TEXT(">()")) +
-					FString((StructMembersForStub.IsEmpty() && ParameterSuffix.IsEmpty()) ? TEXT("") : TEXT(",\r\n\t\t\t"));
+					FString(StructMembersForStub.IsEmpty() ? TEXT("") : TEXT(",\r\n\t\t\t"));
 
-				RigVMMethodsDeclarations += FString::Printf(TEXT("\tstatic %s Static%s(\r\n\t\t%s%s%s\r\n\t);\r\n"), *MethodInfo.ReturnType, *MethodInfo.Name, *RigVMParameterPrefix2, *StructMembers, *ParameterSuffix);
+				RigVMMethodsDeclarations += FString::Printf(TEXT("\t%s %s(%s%s& InExecuteContext);\r\n"), *MethodInfo.ReturnType, *MethodInfo.Name, StructRigVMInfo.ExecuteContextMember.IsEmpty() ? TEXT("const ") : TEXT(""), *StructRigVMInfo.ExecuteContextType);
+				RigVMMethodsDeclarations += FString::Printf(TEXT("\tstatic %s Static%s(\r\n\t\t%s%s\r\n\t);\r\n"), *MethodInfo.ReturnType, *MethodInfo.Name, *RigVMParameterPrefix2, *StructMembers);
 				RigVMMethodsDeclarations += FString::Printf(TEXT("\tFORCEINLINE_DEBUGGABLE static %s RigVM%s(\r\n\t\t%s,\r\n\t\tFRigVMMemoryHandleArray RigVMMemoryHandles\r\n\t)\r\n"), *MethodInfo.ReturnType, *MethodInfo.Name, *RigVMExecuteContextDeclaration);
 				RigVMMethodsDeclarations += FString::Printf(TEXT("\t{\r\n"));
 
-				// implement inline stub method body
-				if (MethodInfo.Parameters.Num() > 0)
-				{
-					//RigVMMethodsDeclarations += FString::Printf(TEXT("\t\tensure(RigVMUserData.Num() == %d);\r\n"), MethodInfo.Parameters.Num());
-					for (int32 ParameterIndex = 0; ParameterIndex < MethodInfo.Parameters.Num(); ParameterIndex++)
-					{
-						const FRigVMParameter& Parameter = MethodInfo.Parameters[ParameterIndex];
-						RigVMMethodsDeclarations += FString::Printf(TEXT("\t\t%s = *(%s*)RigVMExecuteContext.OpaqueArguments[%d];\r\n"), *Parameter.Declaration(), *Parameter.TypeNoRef(), ParameterIndex);
-					}
-					RigVMMethodsDeclarations += FString::Printf(TEXT("\t\t\r\n"));
-				}
-
 				if (RigVMStubProlog.Num() > 0)
 				{
 					for (const FString& RigVMStubPrologLine : RigVMStubProlog)
@@ -3566,7 +3560,7 @@ void FNativeClassHeaderGenerator::ExportGeneratedStructBodyMacros(FOutputDevice&
 					RigVMMethodsDeclarations += FString::Printf(TEXT("\t\t\r\n"));
 				}
 
-				RigVMMethodsDeclarations += FString::Printf(TEXT("\t\t%sStatic%s(\r\n\t\t\t%s%s%s\r\n\t\t);\r\n"), *MethodInfo.ReturnPrefix(), *MethodInfo.Name, *RigVMParameterPrefix4, *StructMembersForStub, *ParameterNamesSuffix);
+				RigVMMethodsDeclarations += FString::Printf(TEXT("\t\t%sStatic%s(\r\n\t\t\t%s%s\r\n\t\t);\r\n"), *MethodInfo.ReturnPrefix(), *MethodInfo.Name, *RigVMParameterPrefix4, *StructMembersForStub);
 				RigVMMethodsDeclarations += FString::Printf(TEXT("\t}\r\n"));
 			}
 		}
@@ -3646,11 +3640,6 @@ void FNativeClassHeaderGenerator::ExportGeneratedStructBodyMacros(FOutputDevice&
 					const FRigVMParameter& Parameter = StructRigVMInfo.Members[MemberIndex];
 					Out.Logf(TEXT("\t\t%s.Emplace(TEXT(\"%s\"), TEXT(\"%s\"));\r\n"), *ArgumentsName, *Parameter.NameOriginal(), *Parameter.TypeOriginal());
 				}
-				for (int32 ParameterIndex = 0; ParameterIndex < MethodInfo.Parameters.Num(); ParameterIndex++)
-				{
-					const FRigVMParameter& Parameter = MethodInfo.Parameters[ParameterIndex];
-					Out.Logf(TEXT("\t\t%s.Emplace(TEXT(\"%s\"), TEXT(\"%s\"));\r\n"), *ArgumentsName, *Parameter.NameOriginal(), *Parameter.TypeOriginal());
-				}
 				Out.Logf(TEXT("\t\tFRigVMRegistry::Get().Register(TEXT(\"%s::%s\"), &%s::RigVM%s, %s, %s);\r\n"),
 					*StructNameCPP, *MethodInfo.Name, *StructNameCPP, *MethodInfo.Name, *OuterSingletonName, *ArgumentsName);
 			}
@@ -3798,12 +3787,10 @@ void FNativeClassHeaderGenerator::ExportGeneratedStructBodyMacros(FOutputDevice&
 		{
 			Out.Log(TEXT("\r\n"));
 
-			FString ParameterDeclaration = MethodInfo.Parameters.Declarations(false, TEXT(",\r\n\t\t"));
-			FString ParameterSuffix = MethodInfo.Parameters.Names(!StructMembersForVirtualFunc.IsEmpty(), TEXT(",\r\n\t\t"));
-			FString RigVMParameterPrefix3 = FString(TEXT("TemporaryExecuteContext")) + FString((StructMembersForVirtualFunc.IsEmpty() && ParameterSuffix.IsEmpty()) ? TEXT("") : TEXT(",\r\n\t\t"));
+			FString RigVMParameterPrefix3 = FString(TEXT("InExecuteContext")) + FString(StructMembersForVirtualFunc.IsEmpty() ? TEXT("") : TEXT(",\r\n\t\t"));
 
-			// implement the virtual function body.
-			Out.Logf(TEXT("%s %s::%s(%s)\r\n"), *MethodInfo.ReturnType, *StructNameCPP, *MethodInfo.Name, *ParameterDeclaration);
+			// implement the virtual function body without execute context
+			Out.Logf(TEXT("%s %s::%s()\r\n"), *MethodInfo.ReturnType, *StructNameCPP, *MethodInfo.Name);
 			Out.Log(TEXT("{\r\n"));
 			if(StructRigVMInfo.ExecuteContextMember.IsEmpty())
 			{
@@ -3814,6 +3801,24 @@ void FNativeClassHeaderGenerator::ExportGeneratedStructBodyMacros(FOutputDevice&
 				Out.Logf(TEXT("\t%s& TemporaryExecuteContext = %s;\r\n"), *StructRigVMInfo.ExecuteContextType, *StructRigVMInfo.ExecuteContextMember);
 			}
 			Out.Log(TEXT("\tTemporaryExecuteContext.Initialize();\r\n"));
+			Out.Log(TEXT("\t"));
+
+			if(MethodInfo.ReturnType != TEXT("void"))
+			{
+				Out.Log(TEXT("return "));
+			}
+			Out.Logf(TEXT("%s(TemporaryExecuteContext);\r\n"), *MethodInfo.Name);
+			Out.Log(TEXT("}\r\n\r\n"));
+
+			// implement the virtual function body.
+			FString ParameterDeclaration = FString::Printf(TEXT("%s& InExecuteContext"), *StructRigVMInfo.ExecuteContextType);
+			if(StructRigVMInfo.ExecuteContextMember.IsEmpty())
+			{
+				ParameterDeclaration = TEXT("const ") + ParameterDeclaration;
+			}
+			
+			Out.Logf(TEXT("%s %s::%s(%s)\r\n"), *MethodInfo.ReturnType, *StructNameCPP, *MethodInfo.Name, *ParameterDeclaration);
+			Out.Log(TEXT("{\r\n"));
 
 			if(RigVMVirtualFuncProlog.Num() > 0)
 			{
@@ -3824,7 +3829,7 @@ void FNativeClassHeaderGenerator::ExportGeneratedStructBodyMacros(FOutputDevice&
 				Out.Log(TEXT("\t\r\n"));
 			}
 
-			Out.Logf(TEXT("\t%sStatic%s(\r\n\t\t%s%s%s\r\n\t);\r\n"), *MethodInfo.ReturnPrefix(), *MethodInfo.Name, *RigVMParameterPrefix3, *StructMembersForVirtualFunc, *ParameterSuffix);
+			Out.Logf(TEXT("\t%sStatic%s(\r\n\t\t%s%s\r\n\t);\r\n"), *MethodInfo.ReturnPrefix(), *MethodInfo.Name, *RigVMParameterPrefix3, *StructMembersForVirtualFunc);
 
 			if (RigVMVirtualFuncEpilog.Num() > 0)
 			{
diff --git a/Engine/Source/Programs/UnrealHeaderTool/Private/HeaderParser.cpp b/Engine/Source/Programs/UnrealHeaderTool/Private/HeaderParser.cpp
--- a/Engine/Source/Programs/UnrealHeaderTool/Private/HeaderParser.cpp
+++ b/Engine/Source/Programs/UnrealHeaderTool/Private/HeaderParser.cpp
@@ -6637,10 +6637,24 @@ void FHeaderParser::CompileRigVMMethodDeclaration(FUnrealStructDefinitionInfo& S
 	}
 
 	FRigVMStructInfo& StructRigVMInfo = StructDef.GetRigVMInfo();
+
+	// disable support for opaque arguments
+	if (MethodInfo.Parameters.Num() > 0)
+	{
+		LogError(TEXT("RIGVM_METHOD F%s::%s has %d parameters. Since 5.2 parameters are no longer allowed for RIGVM_METHOD functions."), *StructRigVMInfo.Name, *MethodInfo.Name, MethodInfo.Parameters.Num());
+		MethodInfo.Parameters = FRigVMParameterArray();
+	}
+	
 	StructRigVMInfo.bHasRigVM = true;
 	StructRigVMInfo.Name = StructDef.GetName();
 	StructRigVMInfo.Methods.Add(MethodInfo);
 	StructRigVMInfo.ExecuteContextType = TEXT("FRigVMExecuteContext");
+
+	FString ExecuteContextTypeMetadata;
+	if(StructDef.GetStringMetaDataHierarchical(TEXT("ExecuteContext"), &ExecuteContextTypeMetadata))
+	{
+		StructRigVMInfo.ExecuteContextType = ExecuteContextTypeMetadata;
+	}
 }
 
 const FName FHeaderParser::NAME_InputText(TEXT("Input"));
diff --git a/Engine/Source/Programs/UnrealHeaderTool/Resources/UHTDebugging/RigVMStructHeader.h b/Engine/Source/Programs/UnrealHeaderTool/Resources/UHTDebugging/RigVMStructHeader.h
--- a/Engine/Source/Programs/UnrealHeaderTool/Resources/UHTDebugging/RigVMStructHeader.h
+++ b/Engine/Source/Programs/UnrealHeaderTool/Resources/UHTDebugging/RigVMStructHeader.h
@@ -52,7 +52,7 @@ struct FUHTTestExecuteContext : public FRigVMExecuteContext
 	GENERATED_BODY()
 };
 
-USTRUCT(meta = (Deprecated = "5.0.0"))
+USTRUCT(meta = (Deprecated = "5.0.0", ExecuteContext = "FUHTTestExecuteContext"))
 struct FRigVMMethodStruct : public FRigVMStructBase
 {
 	GENERATED_BODY()
@@ -61,10 +61,10 @@ struct FRigVMMethodStruct : public FRigVMStructBase
 	void Clear();
 
 	RIGVM_METHOD()
-	virtual void Execute(bool bAdditionalFlag = false, const FString& InString = TEXT("")) override;
+	virtual void Execute() override;
 
 	RIGVM_METHOD()
-	void Compute(float TestFloat);
+	void Compute();
 
 	UPROPERTY(meta = (Input))
 	float A;
@@ -143,7 +143,7 @@ struct CONTROLRIG_API FRigUnit_Mutable
 	}
 
 	RIGVM_METHOD()
-	virtual void Execute(float Factor);
+	virtual void Execute();
 
 	/** The transform of this control */
 	UPROPERTY(EditAnywhere, Category="Control", meta=(Input, Output))
@@ -156,7 +156,7 @@ struct CONTROLRIG_API FRigUnit_BeginExecution
 	GENERATED_BODY()
 
 	RIGVM_METHOD()
-	virtual void Execute(const FRigVMExecuteContext& Context) override;
+	virtual void Execute() override;
 
 	// The execution result
 	UPROPERTY(EditAnywhere, Transient, DisplayName = "Execute", Category = "BeginExecution", meta = (Output))
diff --git a/Engine/Source/Runtime/RigVM/Private/RigVMCore/RigVM.cpp b/Engine/Source/Runtime/RigVM/Private/RigVMCore/RigVM.cpp
--- a/Engine/Source/Runtime/RigVM/Private/RigVMCore/RigVM.cpp
+++ b/Engine/Source/Runtime/RigVM/Private/RigVMCore/RigVM.cpp
@@ -983,10 +983,10 @@ bool URigVM::ResumeExecution()
 	return false;
 }
 
-bool URigVM::ResumeExecution(TArrayView<URigVMMemoryStorage*> Memory, TArrayView<void*> AdditionalArguments, const FName& InEntryName)
+bool URigVM::ResumeExecution(TArrayView<URigVMMemoryStorage*> Memory, const FName& InEntryName)
 {
 	ResumeExecution();
-	return Execute(Memory, AdditionalArguments, InEntryName) != ERigVMExecuteResult::Failed;
+	return Execute(Memory, InEntryName) != ERigVMExecuteResult::Failed;
 }
 
 #endif
@@ -1549,7 +1549,7 @@ bool URigVM::ShouldHaltAtInstruction(const FName& InEventName, const uint16 Inst
 }
 #endif
 
-bool URigVM::Initialize(TArrayView<URigVMMemoryStorage*> Memory, TArrayView<void*> AdditionalArguments, bool bInitializeMemory)
+bool URigVM::Initialize(TArrayView<URigVMMemoryStorage*> Memory, bool bInitializeMemory)
 {
 	if (ExecutingThreadId != INDEX_NONE)
 	{
@@ -1590,7 +1590,6 @@ bool URigVM::Initialize(TArrayView<URigVMMemoryStorage*> Memory, TArrayView<void
 
 	Context.Reset();
 	Context.SliceOffsets.AddZeroed(Instructions.Num());
-	Context.OpaqueArguments = AdditionalArguments;
 	FRigVMExecuteContext& ContextPublicData = Context.GetPublicData<>();
 
 	TGuardValue<URigVM*> VMInContext(Context.VM, this);
@@ -1760,14 +1759,13 @@ bool URigVM::Initialize(TArrayView<URigVMMemoryStorage*> Memory, TArrayView<void
 	return true;
 }
 
-ERigVMExecuteResult URigVM::Execute(TArrayView<URigVMMemoryStorage*> Memory, TArrayView<void*> AdditionalArguments, const FName& InEntryName)
+ERigVMExecuteResult URigVM::Execute(TArrayView<URigVMMemoryStorage*> Memory, const FName& InEntryName)
 {
 	// if this the first entry being executed - get ready for execution
 	const bool bIsRootEntry = EntriesBeingExecuted.IsEmpty();
 
 	TGuardValue<FName> EntryNameGuard(CurrentEntryName, InEntryName);
 	TGuardValue<bool> RootEntryGuard(bCurrentlyRunningRootEntry, bIsRootEntry);
-	TGuardValue<TArrayView<void*>> AdditionalArgumentsGuard(CurrentAdditionalArguments, AdditionalArguments);
 
 	if(bIsRootEntry)
 	{
@@ -1822,7 +1820,6 @@ ERigVMExecuteResult URigVM::Execute(TArrayView<URigVMMemoryStorage*> Memory, TAr
 	{
 		Context.Reset();
 		Context.SliceOffsets.AddZeroed(Instructions.Num());
-		Context.OpaqueArguments = AdditionalArguments;
 	}
 
 	TGuardValue<URigVM*> VMInContext(Context.VM, this);
@@ -2841,7 +2838,7 @@ ERigVMExecuteResult URigVM::ExecuteInstructions(int32 InFirstInstruction, int32
 				{
 					// this will restore the public data after invoking the entry
 					TGuardValue<FRigVMExecuteContext> PublicDataGuard(ContextPublicData, ContextPublicData);
-					const ERigVMExecuteResult ExecuteResult = Execute(CurrentMemory, CurrentAdditionalArguments, Op.EntryName);
+					const ERigVMExecuteResult ExecuteResult = Execute(CurrentMemory, Op.EntryName);
 					if(ExecuteResult != ERigVMExecuteResult::Succeeded)
 					{
 						return CurrentExecuteResult = ExecuteResult;
@@ -2974,7 +2971,7 @@ ERigVMExecuteResult URigVM::ExecuteInstructions(int32 InFirstInstruction, int32
 
 bool URigVM::Execute(const FName& InEntryName)
 {
-	return Execute(TArray<URigVMMemoryStorage*>(), TArrayView<void*>(), InEntryName) != ERigVMExecuteResult::Failed;
+	return Execute(TArray<URigVMMemoryStorage*>(), InEntryName) != ERigVMExecuteResult::Failed;
 }
 
 ERigVMExecuteResult URigVM::ExecuteLazyBranch(const FRigVMBranchInfo& InBranchToRun)
diff --git a/Engine/Source/Runtime/RigVM/Private/RigVMCore/RigVMDispatchFactory.cpp b/Engine/Source/Runtime/RigVM/Private/RigVMCore/RigVMDispatchFactory.cpp
--- a/Engine/Source/Runtime/RigVM/Private/RigVMCore/RigVMDispatchFactory.cpp
+++ b/Engine/Source/Runtime/RigVM/Private/RigVMCore/RigVMDispatchFactory.cpp
@@ -140,6 +140,11 @@ const TArray<FName>& FRigVMDispatchFactory::GetControlFlowBlocks() const
 	return Blocks;
 }
 
+bool FRigVMDispatchFactory::SupportsExecuteContextStruct(const UScriptStruct* InExecuteContextStruct) const
+{
+	return InExecuteContextStruct->IsChildOf(GetExecuteContextStruct());
+}
+
 const TArray<FName>& FRigVMDispatchFactory::GetControlFlowBlocks_Impl() const
 {
 	static const TArray<FName> EmptyArray;
diff --git a/Engine/Source/Runtime/RigVM/Private/RigVMCore/RigVMFunction.cpp b/Engine/Source/Runtime/RigVM/Private/RigVMCore/RigVMFunction.cpp
--- a/Engine/Source/Runtime/RigVM/Private/RigVMCore/RigVMFunction.cpp
+++ b/Engine/Source/Runtime/RigVM/Private/RigVMCore/RigVMFunction.cpp
@@ -4,6 +4,7 @@
 #include "RigVMCore/RigVMTemplate.h"
 #include "RigVMCore/RigVMRegistry.h"
 #include "RigVMCore/RigVMExternalVariable.h"
+#include "RigVMCore/RigVMStruct.h"
 #include "UObject/Package.h"
 
 FString FRigVMFunction::GetName() const
@@ -97,17 +98,6 @@ const TArray<TRigVMTypeIndex>& FRigVMFunction::GetArgumentTypeIndices() const
 	return ArgumentTypeIndices;
 }
 
-bool FRigVMFunction::IsAdditionalArgument(const FRigVMFunctionArgument& InArgument) const
-{
-#if WITH_EDITOR
-	if (Struct)
-	{
-		return Struct->FindPropertyByName(InArgument.Name) == nullptr;
-	}
-#endif
-	return false;
-}
-
 const FRigVMTemplate* FRigVMFunction::GetTemplate() const
 {
 	if(TemplateIndex == INDEX_NONE)
@@ -123,3 +113,49 @@ const FRigVMTemplate* FRigVMFunction::GetTemplate() const
 
 	return Template;
 }
+
+const UScriptStruct* FRigVMFunction::GetExecuteContextStruct() const
+{
+	if(Factory)
+	{
+		return Factory->GetExecuteContextStruct();
+	}
+	if(Struct)
+	{
+#if WITH_EDITOR
+		FString ExecuteContextName;
+		if(Struct->GetStringMetaDataHierarchical(FRigVMStruct::ExecuteContextName, &ExecuteContextName))
+		{
+			const FRigVMTemplateArgumentType& Type = FRigVMRegistry::Get().FindTypeFromCPPType(ExecuteContextName);
+			if(const UScriptStruct* ExecuteContextStruct = Cast<UScriptStruct>(Type.CPPTypeObject))
+			{
+				return ExecuteContextStruct;
+			}
+		}
+#endif
+
+		for (TFieldIterator<FProperty> It(Struct); It; ++It)
+		{
+			const FProperty* Property = *It;
+
+			if(const FArrayProperty* ArrayProperty = CastField<FArrayProperty>(Property))
+			{
+				Property = ArrayProperty->Inner;
+			}
+			if(const FStructProperty* StructProperty = CastField<FStructProperty>(Property))
+			{
+				if(StructProperty->Struct->IsChildOf(FRigVMExecuteContext::StaticStruct()))
+				{
+					return StructProperty->Struct;
+				}
+			}
+		}
+	}
+
+	return FRigVMExecuteContext::StaticStruct();
+}
+
+bool FRigVMFunction::SupportsExecuteContextStruct(const UScriptStruct* InExecuteContextStruct) const
+{
+	return InExecuteContextStruct->IsChildOf(GetExecuteContextStruct());
+}
diff --git a/Engine/Source/Runtime/RigVM/Private/RigVMCore/RigVMNativized.cpp b/Engine/Source/Runtime/RigVM/Private/RigVMCore/RigVMNativized.cpp
--- a/Engine/Source/Runtime/RigVM/Private/RigVMCore/RigVMNativized.cpp
+++ b/Engine/Source/Runtime/RigVM/Private/RigVMCore/RigVMNativized.cpp
@@ -27,15 +27,13 @@ void URigVMNativized::Reset(bool IsIgnoringArchetypeRef)
 	ByteCodeStorage.Reset();
 }
 
-bool URigVMNativized::Initialize(TArrayView<URigVMMemoryStorage*> Memory, TArrayView<void*> AdditionalArguments,
-	bool bInitializeMemory)
+bool URigVMNativized::Initialize(TArrayView<URigVMMemoryStorage*> Memory, bool bInitializeMemory)
 {
 	// nothing to do here 
 	return true;
 }
 
-ERigVMExecuteResult URigVMNativized::Execute(TArrayView<URigVMMemoryStorage*> Memory, TArrayView<void*> AdditionalArguments,
-	const FName& InEntryName)
+ERigVMExecuteResult URigVMNativized::Execute(TArrayView<URigVMMemoryStorage*> Memory, const FName& InEntryName)
 {
 	// to be implemented by the generated code
 	return ERigVMExecuteResult::Failed;
diff --git a/Engine/Source/Runtime/RigVM/Private/RigVMCore/RigVMTemplate.cpp b/Engine/Source/Runtime/RigVM/Private/RigVMCore/RigVMTemplate.cpp
--- a/Engine/Source/Runtime/RigVM/Private/RigVMCore/RigVMTemplate.cpp
+++ b/Engine/Source/Runtime/RigVM/Private/RigVMCore/RigVMTemplate.cpp
@@ -963,20 +963,27 @@ FString FRigVMTemplate::GetArgumentMetaData(const FName& InArgumentName, const F
 
 #endif
 
-bool FRigVMTemplate::IsCompatible(const FRigVMTemplate& InOther) const
+bool FRigVMTemplate::Merge(const FRigVMTemplate& InOther)
 {
 	if (!IsValid() || !InOther.IsValid())
 	{
 		return false;
 	}
 
-	return Notation == InOther.Notation;
-}
+	if(Notation != InOther.Notation)
+	{
+		return false;
+	}
 
-bool FRigVMTemplate::Merge(const FRigVMTemplate& InOther)
-{
-	if (!IsCompatible(InOther))
+	if(InOther.GetExecuteContextStruct() != GetExecuteContextStruct())
 	{
+		// find the previously defined permutation.
+		UE_LOG(LogRigVM, Display, TEXT("RigVMFunction '%s' cannot be merged into the '%s' template. ExecuteContext Types differ ('%s' vs '%s' from '%s')."),
+			*InOther.GetPrimaryPermutation()->Name,
+			*GetNotation().ToString(),
+			*InOther.GetExecuteContextStruct()->GetStructCPPName(),
+			*GetExecuteContextStruct()->GetStructCPPName(),
+			*GetPrimaryPermutation()->Name);
 		return false;
 	}
 
@@ -1001,7 +1008,7 @@ bool FRigVMTemplate::Merge(const FRigVMTemplate& InOther)
 		{
 			// find the previously defined permutation.
 			UE_LOG(LogRigVM, Display, TEXT("RigVMFunction '%s' cannot be merged into the '%s' template. It collides with '%s'."),
-				*InOther.GetPermutation(0)->Name,
+				*InOther.GetPrimaryPermutation()->Name,
 				*GetNotation().ToString(),
 				*GetPermutation(PermutationIndex)->Name);
 			return false;
@@ -1107,6 +1114,21 @@ const TArray<FRigVMExecuteArgument>& FRigVMTemplate::GetExecuteArguments() const
 	return ExecuteArguments;
 }
 
+const UScriptStruct* FRigVMTemplate::GetExecuteContextStruct() const
+{
+	if(const FRigVMDispatchFactory* Factory = GetDispatchFactory())
+	{
+		return Factory->GetExecuteContextStruct();
+	}
+	check(!Permutations.IsEmpty());
+	return GetPrimaryPermutation()->GetExecuteContextStruct();
+}
+
+bool FRigVMTemplate::SupportsExecuteContextStruct(const UScriptStruct* InExecuteContextStruct) const
+{
+	return InExecuteContextStruct->IsChildOf(GetExecuteContextStruct());
+}
+
 bool FRigVMTemplate::ArgumentSupportsTypeIndex(const FName& InArgumentName, TRigVMTypeIndex InTypeIndex, TRigVMTypeIndex* OutTypeIndex) const
 {
 	if (const FRigVMTemplateArgument* Argument = FindArgument(InArgumentName))
@@ -1468,7 +1490,7 @@ FString FRigVMTemplate::GetCategory() const
 	}
 	
 	FString Category;
-	GetPermutation(0)->Struct->GetStringMetaDataHierarchical(FRigVMStruct::CategoryMetaName, &Category);
+	GetPrimaryPermutation()->Struct->GetStringMetaDataHierarchical(FRigVMStruct::CategoryMetaName, &Category);
 
 	if (Category.IsEmpty())
 	{
diff --git a/Engine/Source/Runtime/RigVM/Public/RigVMCore/RigVM.h b/Engine/Source/Runtime/RigVM/Public/RigVMCore/RigVM.h
--- a/Engine/Source/Runtime/RigVM/Public/RigVMCore/RigVM.h
+++ b/Engine/Source/Runtime/RigVM/Public/RigVMCore/RigVM.h
@@ -210,12 +210,12 @@ public:
 	}
 
 	// Initializes all execute ops and their memory.
-	virtual bool Initialize(TArrayView<URigVMMemoryStorage*> Memory, TArrayView<void*> AdditionalArguments, bool bInitializeMemory = true);
+	virtual bool Initialize(TArrayView<URigVMMemoryStorage*> Memory, bool bInitializeMemory = true);
 
 	// Executes the VM.
 	// You can optionally provide external memory to the execution
 	// and provide optional additional operands.
-	virtual ERigVMExecuteResult Execute(TArrayView<URigVMMemoryStorage*> Memory, TArrayView<void*> AdditionalArguments, const FName& InEntryName = NAME_None);
+	virtual ERigVMExecuteResult Execute(TArrayView<URigVMMemoryStorage*> Memory, const FName& InEntryName = NAME_None);
 
 	// Executes the VM.
 	// You can optionally provide external memory to the execution
@@ -348,7 +348,7 @@ public:
 
 	const void SetFirstEntryEventInEventQueue(const FName& InFirstEventName) { FirstEntryEventInQueue = InFirstEventName; }
 
-	bool ResumeExecution(TArrayView<URigVMMemoryStorage*> Memory, TArrayView<void*> AdditionalArguments, const FName& InEntryName = NAME_None);
+	bool ResumeExecution(TArrayView<URigVMMemoryStorage*> Memory, const FName& InEntryName = NAME_None);
 	bool ResumeExecution();
 #endif
 
@@ -848,7 +848,6 @@ private:
 	FName CurrentEntryName;
 	bool bCurrentlyRunningRootEntry;
 	TArrayView<URigVMMemoryStorage*> CurrentMemory;
-	TArrayView<void*> CurrentAdditionalArguments;
 
 #if WITH_EDITOR
 protected:
diff --git a/Engine/Source/Runtime/RigVM/Public/RigVMCore/RigVMDispatchFactory.h b/Engine/Source/Runtime/RigVM/Public/RigVMCore/RigVMDispatchFactory.h
--- a/Engine/Source/Runtime/RigVM/Public/RigVMCore/RigVMDispatchFactory.h
+++ b/Engine/Source/Runtime/RigVM/Public/RigVMCore/RigVMDispatchFactory.h
@@ -103,12 +103,12 @@ public:
 	// Returns the execute context support for this dispatch factory
 	virtual UScriptStruct* GetExecuteContextStruct() const { return FRigVMExecuteContext::StaticStruct(); }
 
+	// Returns truf if this factory supports a given executecontext struct
+	bool SupportsExecuteContextStruct(const UScriptStruct* InExecuteContextStruct) const;
+
 	// registered needed types during registration of the factory
 	virtual void RegisterDependencyTypes() const {}
 
-	// returns opaque arguments expected to be passed to this factory in FRigVMExtendedExecuteContext::OpaqueArguments
-	virtual TArray<TPair<FName,FString>> GetOpaqueArguments() const { return TArray<TPair<FName,FString>>(); }
-
 	// returns the arguments of the template
 	virtual TArray<FRigVMTemplateArgument> GetArguments() const { return TArray<FRigVMTemplateArgument>(); }
 
diff --git a/Engine/Source/Runtime/RigVM/Public/RigVMCore/RigVMExecuteContext.h b/Engine/Source/Runtime/RigVM/Public/RigVMCore/RigVMExecuteContext.h
--- a/Engine/Source/Runtime/RigVM/Public/RigVMCore/RigVMExecuteContext.h
+++ b/Engine/Source/Runtime/RigVM/Public/RigVMCore/RigVMExecuteContext.h
@@ -402,7 +402,6 @@ struct RIGVM_API FRigVMExtendedExecuteContext
 		}
 
 		VM = Other.VM;
-		OpaqueArguments = Other.OpaqueArguments;
 		Slices = Other.Slices;
 		SliceOffsets = Other.SliceOffsets;
 		return *this;
@@ -517,7 +516,6 @@ struct RIGVM_API FRigVMExtendedExecuteContext
 
 	FStructOnScope PublicDataScope;
 	URigVM* VM;
-	TArrayView<void*> OpaqueArguments;
 	TArray<FRigVMSlice> Slices;
 	TArray<uint16> SliceOffsets;
 	double LastExecutionMicroSeconds;
diff --git a/Engine/Source/Runtime/RigVM/Public/RigVMCore/RigVMFunction.h b/Engine/Source/Runtime/RigVM/Public/RigVMCore/RigVMFunction.h
--- a/Engine/Source/Runtime/RigVM/Public/RigVMCore/RigVMFunction.h
+++ b/Engine/Source/Runtime/RigVM/Public/RigVMCore/RigVMFunction.h
@@ -20,7 +20,6 @@ struct FRigVMExtendedExecuteContext;
 struct FRigVMMemoryHandle;
 
 typedef TArrayView<FRigVMMemoryHandle> FRigVMMemoryHandleArray;
-typedef TArrayView<void*> FRigVMUserDataArray;
 
 typedef void (*FRigVMFunctionPtr)(FRigVMExtendedExecuteContext& RigVMExecuteContext, FRigVMMemoryHandleArray RigVMMemoryHandles);
 
@@ -155,6 +154,7 @@ struct RIGVM_API FRigVMFunction
 	FString GetModuleRelativeHeaderPath() const;
 	const TArray<FRigVMFunctionArgument>& GetArguments() const { return Arguments; }
 	const TArray<TRigVMTypeIndex>& GetArgumentTypeIndices() const;
-	bool IsAdditionalArgument(const FRigVMFunctionArgument& InArgument) const;
 	const FRigVMTemplate* GetTemplate() const;
+	const UScriptStruct* GetExecuteContextStruct() const;
+	bool SupportsExecuteContextStruct(const UScriptStruct* InExecuteContextStruct) const;
 };
diff --git a/Engine/Source/Runtime/RigVM/Public/RigVMCore/RigVMNativized.h b/Engine/Source/Runtime/RigVM/Public/RigVMCore/RigVMNativized.h
--- a/Engine/Source/Runtime/RigVM/Public/RigVMCore/RigVMNativized.h
+++ b/Engine/Source/Runtime/RigVM/Public/RigVMCore/RigVMNativized.h
@@ -50,8 +50,8 @@ public:
 	virtual bool IsNativized() const override { return true; }
 	virtual void Empty() override { return; }
 	virtual void CopyFrom(URigVM* InVM, bool bDeferCopy = false, bool bReferenceLiteralMemory = false, bool bReferenceByteCode = false, bool bCopyExternalVariables = false, bool bCopyDynamicRegisters = false) override { return; }
-	virtual bool Initialize(TArrayView<URigVMMemoryStorage*> Memory, TArrayView<void*> AdditionalArguments, bool bInitializeMemory = true) override;
-	virtual ERigVMExecuteResult Execute(TArrayView<URigVMMemoryStorage*> Memory, TArrayView<void*> AdditionalArguments, const FName& InEntryName = NAME_None) override;
+	virtual bool Initialize(TArrayView<URigVMMemoryStorage*> Memory, bool bInitializeMemory = true) override;
+	virtual ERigVMExecuteResult Execute(TArrayView<URigVMMemoryStorage*> Memory, const FName& InEntryName = NAME_None) override;
 	virtual int32 AddRigVMFunction(UScriptStruct* InRigVMStruct, const FName& InMethodName) override { return INDEX_NONE; }
 	virtual FString GetRigVMFunctionName(int32 InFunctionIndex) const override { return FString(); }
 	virtual URigVMMemoryStorage* GetMemoryByType(ERigVMMemoryType InMemoryType, bool bCreateIfNeeded = true) override { return nullptr; }
@@ -76,13 +76,12 @@ public:
 protected:
 
 	template<typename ExecuteContextType = FRigVMExecuteContext>
-	ExecuteContextType& UpdateContext(TArrayView<void*> AdditionalArguments, int32 InNumberInstructions, const FName& InEntryName)
+	ExecuteContextType& UpdateContext(int32 InNumberInstructions, const FName& InEntryName)
 	{
 		UpdateExternalVariables();
 	
 		Context.Reset();
 		Context.VM = this;
-		Context.OpaqueArguments = AdditionalArguments;
 		Context.SliceOffsets.AddZeroed(InNumberInstructions);
 		Context.GetPublicData<ExecuteContextType>().EventName = InEntryName;
 		return Context.GetPublicData<ExecuteContextType>();
diff --git a/Engine/Source/Runtime/RigVM/Public/RigVMCore/RigVMStructTest.h b/Engine/Source/Runtime/RigVM/Public/RigVMCore/RigVMStructTest.h
--- a/Engine/Source/Runtime/RigVM/Public/RigVMCore/RigVMStructTest.h
+++ b/Engine/Source/Runtime/RigVM/Public/RigVMCore/RigVMStructTest.h
@@ -66,7 +66,7 @@ private:
 		virtual FString GetBeautifiedTestName() const override { return TEXT(RIGVMSTRUCT_TEST_STRINGIFY(RigVM.RigVMFunctions.TRigVMStruct)); } \
 		void Execute() \
 		{ \
-			Unit.Execute(); \
+			Unit.Execute(ExecuteContext); \
 		} \
 	}; \
 	namespace\
diff --git a/Engine/Source/Runtime/RigVM/Public/RigVMCore/RigVMTemplate.h b/Engine/Source/Runtime/RigVM/Public/RigVMCore/RigVMTemplate.h
--- a/Engine/Source/Runtime/RigVM/Public/RigVMCore/RigVMTemplate.h
+++ b/Engine/Source/Runtime/RigVM/Public/RigVMCore/RigVMTemplate.h
@@ -249,6 +249,7 @@ protected:
 	friend struct FRigVMRegistry;
 	friend struct FRigVMStructUpgradeInfo;
 	friend struct FRigVMSetLibraryTemplateAction;
+	friend class URigVMCompiler;
 };
 
 /**
@@ -279,9 +280,6 @@ public:
 	// Returns the name of the template
 	FName GetName() const;
 
-	// returns true if this template is compatible with another one
-	bool IsCompatible(const FRigVMTemplate& InOther) const;
-
 	// returns true if this template can merge another one
 	bool Merge(const FRigVMTemplate& InOther);
 
@@ -303,6 +301,12 @@ public:
 	// returns an argument given a name (or nullptr)
 	const FRigVMExecuteArgument* FindExecuteArgument(const FName& InArgumentName) const;
 
+	// returns the top level execute context struct this template uses
+	const UScriptStruct* GetExecuteContextStruct() const;
+
+	// returns true if this template supports a given execute context struct
+	bool SupportsExecuteContextStruct(const UScriptStruct* InExecuteContextStruct) const;
+
 	// returns true if a given arg supports a type
 	bool ArgumentSupportsTypeIndex(const FName& InArgumentName, TRigVMTypeIndex InTypeIndex, TRigVMTypeIndex* OutTypeIndex = nullptr) const;
 
