Update C++ UHT to match enum changes in C# UHT

#rnx

diff --git a/Engine/Source/Programs/UnrealHeaderTool/Private/BaseParser.cpp b/Engine/Source/Programs/UnrealHeaderTool/Private/BaseParser.cpp
--- a/Engine/Source/Programs/UnrealHeaderTool/Private/BaseParser.cpp
+++ b/Engine/Source/Programs/UnrealHeaderTool/Private/BaseParser.cpp
@@ -1271,9 +1271,13 @@ EUnderlyingEnumType FBaseParser::ParseUnderlyingEnumType()
 		{
 			UnderlyingType = EUnderlyingEnumType::int64;
 		}
+		else if (BaseToken.IsValue(TEXT("int"), ESearchCase::CaseSensitive))
+		{
+			UnderlyingType = EUnderlyingEnumType::Int;
+		}
 		else
 		{
-			Throwf(TEXT("Unsupported enum class base type: %s"), *BaseToken.GetTokenValue());
+			Throwf(TEXT("Unsupported enum base type: %s"), *BaseToken.GetTokenValue());
 		}
 	}
 	else
diff --git a/Engine/Source/Programs/UnrealHeaderTool/Private/BaseParser.h b/Engine/Source/Programs/UnrealHeaderTool/Private/BaseParser.h
--- a/Engine/Source/Programs/UnrealHeaderTool/Private/BaseParser.h
+++ b/Engine/Source/Programs/UnrealHeaderTool/Private/BaseParser.h
@@ -25,7 +25,8 @@ enum class EUnderlyingEnumType
 	int8,
 	int16,
 	int32,
-	int64
+	int64,
+	Int,
 };
 
 using FMetaData = TMap<FName, FString>;
diff --git a/Engine/Source/Programs/UnrealHeaderTool/Private/CodeGenerator.cpp b/Engine/Source/Programs/UnrealHeaderTool/Private/CodeGenerator.cpp
--- a/Engine/Source/Programs/UnrealHeaderTool/Private/CodeGenerator.cpp
+++ b/Engine/Source/Programs/UnrealHeaderTool/Private/CodeGenerator.cpp
@@ -3335,6 +3335,32 @@ FString FNativeClassHeaderGenerator::GetClassFlagExportText(FUnrealClassDefiniti
 	return StaticClassFlagText;
 }
 
+static FString GetUnderlyingType(FUnrealEnumDefinitionInfo& EnumDef)
+{
+	FString UnderlyingTypeString;
+
+	if (EnumDef.GetUnderlyingType() != EUnderlyingEnumType::Unspecified)
+	{
+		UnderlyingTypeString = TEXT(" : ");
+
+		switch (EnumDef.GetUnderlyingType())
+		{
+		case EUnderlyingEnumType::int8:        UnderlyingTypeString += TNameOf<int8>::GetName();	break;
+		case EUnderlyingEnumType::int16:       UnderlyingTypeString += TNameOf<int16>::GetName();	break;
+		case EUnderlyingEnumType::int32:       UnderlyingTypeString += TNameOf<int32>::GetName();	break;
+		case EUnderlyingEnumType::int64:       UnderlyingTypeString += TNameOf<int64>::GetName();	break;
+		case EUnderlyingEnumType::uint8:       UnderlyingTypeString += TNameOf<uint8>::GetName();	break;
+		case EUnderlyingEnumType::uint16:      UnderlyingTypeString += TNameOf<uint16>::GetName();	break;
+		case EUnderlyingEnumType::uint32:      UnderlyingTypeString += TNameOf<uint32>::GetName();	break;
+		case EUnderlyingEnumType::uint64:      UnderlyingTypeString += TNameOf<uint64>::GetName();	break;
+		case EUnderlyingEnumType::Int:         UnderlyingTypeString += TEXT("int");                 break;
+		default:
+			check(false);
+		}
+	}
+	return UnderlyingTypeString;
+}
+
 /**
 * Exports the header text for the list of enums specified
 *
@@ -3360,33 +3386,27 @@ void FNativeClassHeaderGenerator::ExportEnum(FOutputDevice& Out, FUnrealEnumDefi
 
 	if (EnumDef.GetCppForm() == UEnum::ECppForm::EnumClass)
 	{
-		FString UnderlyingTypeString;
-
-		if (EnumDef.GetUnderlyingType() != EUnderlyingEnumType::Unspecified)
-		{
-			UnderlyingTypeString = TEXT(" : ");
-
-			switch (EnumDef.GetUnderlyingType())
-			{
-			case EUnderlyingEnumType::int8:        UnderlyingTypeString += TNameOf<int8>::GetName();	break;
-			case EUnderlyingEnumType::int16:       UnderlyingTypeString += TNameOf<int16>::GetName();	break;
-			case EUnderlyingEnumType::int32:       UnderlyingTypeString += TNameOf<int32>::GetName();	break;
-			case EUnderlyingEnumType::int64:       UnderlyingTypeString += TNameOf<int64>::GetName();	break;
-			case EUnderlyingEnumType::uint8:       UnderlyingTypeString += TNameOf<uint8>::GetName();	break;
-			case EUnderlyingEnumType::uint16:      UnderlyingTypeString += TNameOf<uint16>::GetName();	break;
-			case EUnderlyingEnumType::uint32:      UnderlyingTypeString += TNameOf<uint32>::GetName();	break;
-			case EUnderlyingEnumType::uint64:      UnderlyingTypeString += TNameOf<uint64>::GetName();	break;
-			default:
-				check(false);
-			}
-		}
+		Out.Logf(TEXT("\r\n"));
+		Out.Logf(TEXT("enum class %s%s;\r\n"), *EnumDef.GetCppType(), *GetUnderlyingType(EnumDef));
 
-		Out.Logf( TEXT("\r\n") );
-		Out.Logf( TEXT("enum class %s%s;\r\n"), *EnumDef.GetCppType(), *UnderlyingTypeString );
-		
 		// Add TIsUEnumClass typetraits
-		Out.Logf( TEXT("template<> struct TIsUEnumClass<%s> { enum { Value = true }; };\r\n"), *EnumDef.GetCppType());
+		Out.Logf(TEXT("template<> struct TIsUEnumClass<%s> { enum { Value = true }; };\r\n"), *EnumDef.GetCppType());
+	}
+	else if (EnumDef.GetCppForm() == UEnum::ECppForm::Regular && EnumDef.GetUnderlyingType() != EUnderlyingEnumType::Unspecified)
+	{
+		Out.Logf(TEXT("\r\n"));
+		Out.Logf(TEXT("enum %s%s;\r\n"), *EnumDef.GetCppType(), *GetUnderlyingType(EnumDef));
+	}
+	else if (EnumDef.GetCppForm() == UEnum::ECppForm::Namespaced && EnumDef.GetUnderlyingType() != EUnderlyingEnumType::Unspecified)
+	{
+		TArray<FString> Parts;
+		EnumDef.GetCppType().ParseIntoArray(Parts, TEXT("::"));
+		Out.Logf(TEXT("\r\n"));
+		Out.Logf(TEXT("namespace %s { enum %s%s; }\r\n"), *Parts[0], *Parts[1], *GetUnderlyingType(EnumDef));
+	}
 
+	if (EnumDef.GetCppForm() == UEnum::ECppForm::EnumClass || EnumDef.GetUnderlyingType() != EUnderlyingEnumType::Unspecified)
+	{
 		// Forward declare the StaticEnum<> specialisation for enum classes
 		Out.Logf( TEXT("template<> %sUEnum* StaticEnum<%s>();\r\n"), *GetAPIString(), *EnumDef.GetCppType());
 		Out.Logf( TEXT("\r\n") );
diff --git a/Engine/Source/Programs/UnrealHeaderTool/Private/HeaderParser.cpp b/Engine/Source/Programs/UnrealHeaderTool/Private/HeaderParser.cpp
--- a/Engine/Source/Programs/UnrealHeaderTool/Private/HeaderParser.cpp
+++ b/Engine/Source/Programs/UnrealHeaderTool/Private/HeaderParser.cpp
@@ -1125,6 +1125,11 @@ FUnrealEnumDefinitionInfo& FHeaderParser::CompileEnum()
 	if (CppForm == UEnum::ECppForm::EnumClass)
 	{
 		UnderlyingType = ParseUnderlyingEnumType();
+		if (UnderlyingType != EUnderlyingEnumType::uint8 && EnumMetaData.Contains(FHeaderParserNames::NAME_BlueprintType))
+		{
+			Throwf(TEXT("Invalid BlueprintType enum base - currently only uint8 supported"));
+		}
+
 		if (UnderlyingType != EnumDef.GetUnderlyingType())
 		{
 			Throwf(TEXT("ICE: Mismatch of underlying enum type between pre-parser and parser"));
@@ -1132,19 +1137,8 @@ FUnrealEnumDefinitionInfo& FHeaderParser::CompileEnum()
 	}
 	else if (CppForm == UEnum::ECppForm::Regular)
 	{
-		if (MatchSymbol(TEXT(':')))
-		{
-			FToken BaseToken;
-			if (!GetIdentifier(BaseToken))
-			{
-				Throwf(TEXT("Missing enum base"));
-			}
-
-			if (!BaseToken.IsValue(TEXT("int"), ESearchCase::CaseSensitive))
-			{
-				LogError(TEXT("Regular enums only support 'int' as the value size"));
-			}
-		}
+		UnderlyingType = ParseUnderlyingEnumType();
+		EnumDef.SetUnderlyingType(UnderlyingType);
 	}
 	else
 	{
@@ -1154,11 +1148,6 @@ FUnrealEnumDefinitionInfo& FHeaderParser::CompileEnum()
 		}
 	}
 
-	if (UnderlyingType != EUnderlyingEnumType::uint8 && EnumMetaData.Contains(FHeaderParserNames::NAME_BlueprintType))
-	{
-		Throwf(TEXT("Invalid BlueprintType enum base - currently only uint8 supported"));
-	}
-
 	EnumDef.GetDefinitionRange().Start = &Input[InputPos];
 
 	// Get opening brace.
@@ -1181,19 +1170,8 @@ FUnrealEnumDefinitionInfo& FHeaderParser::CompileEnum()
 
 			EnumDef.SetCppType(FString::Printf(TEXT("%s::%s"), *EnumIdentifier, *InnerEnumToken.GetTokenValue()));
 
-			if (MatchSymbol(TEXT(':')))
-			{
-				FToken BaseToken;
-				if (!GetIdentifier(BaseToken))
-				{
-					Throwf(TEXT("Missing enum base"));
-				}
-
-				if (!BaseToken.IsValue(TEXT("int"), ESearchCase::CaseSensitive))
-				{
-					LogError(TEXT("Namespace enums only support 'int' as the value size"));
-				}
-			}
+			UnderlyingType = ParseUnderlyingEnumType();
+			EnumDef.SetUnderlyingType(UnderlyingType);
 
 			RequireSymbol( TEXT('{'), TEXT("'Enum'") );
 		}
@@ -6903,7 +6881,8 @@ void FHeaderParser::ParseParameterList(FUnrealFunctionDefinitionInfo& FunctionDe
 			{
 				if (FUnrealEnumDefinitionInfo* EnumDef = Property.AsEnum())
 				{
-					if (EnumDef->GetUnderlyingType() != EUnderlyingEnumType::uint8 &&
+					if (EnumDef->GetCppForm() == UEnum::ECppForm::EnumClass &&
+						EnumDef->GetUnderlyingType() != EUnderlyingEnumType::uint8 &&
 						EnumDef->GetUnderlyingType() != EUnderlyingEnumType::Unspecified)
 					{
 						Throwf(TEXT("Invalid enum param for Blueprints - currently only uint8 supported"));
@@ -9220,11 +9199,7 @@ TSharedRef<FUnrealTypeDefinitionInfo> FHeaderPreParser::ParseEnumDeclaration(con
 	}
 
 	// Read base for enum class
-	EUnderlyingEnumType UnderlyingType = EUnderlyingEnumType::uint8;
-	if (CppForm == UEnum::ECppForm::EnumClass)
-	{
-		UnderlyingType = ParseUnderlyingEnumType();
-	}
+	EUnderlyingEnumType UnderlyingType = ParseUnderlyingEnumType();
 
 	TSharedRef<FUnrealEnumDefinitionInfo> EnumDef = MakeShareable(new FUnrealEnumDefinitionInfo(SourceFile, InLineNumber, FString(EnumToken.Value), FName(EnumToken.Value, FNAME_Add), CppForm, UnderlyingType));
 	return EnumDef;
diff --git a/Engine/Source/Programs/UnrealHeaderTool/Private/UnrealTypeDefinitionInfo.h b/Engine/Source/Programs/UnrealHeaderTool/Private/UnrealTypeDefinitionInfo.h
--- a/Engine/Source/Programs/UnrealHeaderTool/Private/UnrealTypeDefinitionInfo.h
+++ b/Engine/Source/Programs/UnrealHeaderTool/Private/UnrealTypeDefinitionInfo.h
@@ -2012,6 +2012,14 @@ public:
 		return UnderlyingType;
 	}
 
+	/**
+	 * Set the underlying enumeration type
+	 */
+	void SetUnderlyingType(EUnderlyingEnumType InUnderlyingType)
+	{
+		UnderlyingType = InUnderlyingType;
+	}
+
 	/**
 	 * Return true if the enumeration is editor only
 	 */
