Ported C# UHT include changes to C++ UHT.
Removed redundant check for return value since it is covered in the loop above.

diff --git a/Engine/Source/Programs/Shared/EpicGames.UHT/Exporters/CodeGen/UhtHeaderCodeGeneratorCppFile.cs b/Engine/Source/Programs/Shared/EpicGames.UHT/Exporters/CodeGen/UhtHeaderCodeGeneratorCppFile.cs
--- a/Engine/Source/Programs/Shared/EpicGames.UHT/Exporters/CodeGen/UhtHeaderCodeGeneratorCppFile.cs
+++ b/Engine/Source/Programs/Shared/EpicGames.UHT/Exporters/CodeGen/UhtHeaderCodeGeneratorCppFile.cs
@@ -67,20 +67,18 @@ namespace EpicGames.UHT.Exporters.CodeGen
 								}
 							}
 
-							bool RequireIncludeForClasses = IsRpcFunction(function) && ShouldExportFunction(function);
+							bool requireIncludeForClasses = IsRpcFunction(function) && ShouldExportFunction(function);
 
 							foreach (UhtProperty property in function.Properties)
 							{
-								AddIncludeForProperty(property, addedIncludes, includesToAdd, RequireIncludeForClasses);
+								AddIncludeForProperty(property, requireIncludeForClasses, addedIncludes, includesToAdd);
 							}
-							if (function.ReturnProperty != null)
-								AddIncludeForProperty(function.ReturnProperty, addedIncludes, includesToAdd, RequireIncludeForClasses);
 						}
 
 						// Properties
 						foreach (UhtProperty property in structObj.Properties)
 						{
-							AddIncludeForProperty(property, addedIncludes, includesToAdd, false);
+							AddIncludeForProperty(property, false, addedIncludes, includesToAdd);
 						}
 					}
 
@@ -355,7 +353,7 @@ namespace EpicGames.UHT.Exporters.CodeGen
 			}
 		}
 
-		private void AddIncludeForType(UhtProperty uhtProperty, HashSet<UhtHeaderFile> addedIncludes, IList<string> includesToAdd, bool requireIncludeForClasses)
+		private void AddIncludeForType(UhtProperty uhtProperty, bool requireIncludeForClasses, HashSet<UhtHeaderFile> addedIncludes, IList<string> includesToAdd)
 		{
 			if (uhtProperty is UhtStructProperty structProperty)
 			{
@@ -375,18 +373,18 @@ namespace EpicGames.UHT.Exporters.CodeGen
 			}
 		}
 
-		private void AddIncludeForProperty(UhtProperty property, HashSet<UhtHeaderFile> addedIncludes, IList<string> includesToAdd, bool requireIncludeForClasses)
+		private void AddIncludeForProperty(UhtProperty property, bool requireIncludeForClasses, HashSet<UhtHeaderFile> addedIncludes, IList<string> includesToAdd)
 		{
-			AddIncludeForType(property, addedIncludes, includesToAdd, requireIncludeForClasses);
+			AddIncludeForType(property, requireIncludeForClasses, addedIncludes, includesToAdd);
 
 			if (property is UhtContainerBaseProperty containerProperty)
 			{
-				AddIncludeForType(containerProperty.ValueProperty, addedIncludes, includesToAdd, false);
+				AddIncludeForType(containerProperty.ValueProperty, false, addedIncludes, includesToAdd);
 			}
 
 			if (property is UhtMapProperty mapProperty)
 			{
-				AddIncludeForType(mapProperty.KeyProperty, addedIncludes, includesToAdd, false);
+				AddIncludeForType(mapProperty.KeyProperty, false, addedIncludes, includesToAdd);
 			}
 		}
 
diff --git a/Engine/Source/Programs/UnrealHeaderTool/Private/CodeGenerator.cpp b/Engine/Source/Programs/UnrealHeaderTool/Private/CodeGenerator.cpp
--- a/Engine/Source/Programs/UnrealHeaderTool/Private/CodeGenerator.cpp
+++ b/Engine/Source/Programs/UnrealHeaderTool/Private/CodeGenerator.cpp
@@ -5616,7 +5616,7 @@ void FNativeClassHeaderGenerator::ApplyAlternatePropertyExportText(FUnrealProper
 	}
 }
 
-static void AddIncludeForType(const FUnrealTypeDefinitionInfo* PropertyTypeDef, TArray<FString>& RelativeIncludes)
+static void AddIncludeForType(const TSharedRef<FUnrealPropertyDefinitionInfo>& PropertyDef, const FUnrealTypeDefinitionInfo* PropertyTypeDef, bool bRequireIncludeForClasses, TArray<FString>& RelativeIncludes)
 {
 	if (UHTCast<FUnrealScriptStructDefinitionInfo>(PropertyTypeDef) != nullptr)
 	{
@@ -5626,15 +5626,27 @@ static void AddIncludeForType(const FUnrealTypeDefinitionInfo* PropertyTypeDef,
 			RelativeIncludes.AddUnique(MoveTemp(Header));
 		}
 	}
+	else if (bRequireIncludeForClasses)
+	{
+		const FPropertyBase& PropertyBase = PropertyDef->GetPropertyBase();
+		if (PropertyBase.Type == CPT_ObjectReference && PropertyBase.ClassDef->IsChildOf(*GUClassDef))
+		{
+			if (PropertyBase.ClassDef->HasSource() && !PropertyBase.ClassDef->GetUnrealSourceFile().IsNoExportTypes())
+			{
+				FString Header = GetBuildPath(PropertyBase.ClassDef->GetUnrealSourceFile());
+				RelativeIncludes.AddUnique(MoveTemp(Header));
+			}
+		}
+	}
 }
 
-static void AddIncludeForProperty(const TSharedRef<FUnrealPropertyDefinitionInfo>& PropertyDef, TArray<FString>& RelativeIncludes)
+static void AddIncludeForProperty(const TSharedRef<FUnrealPropertyDefinitionInfo>& PropertyDef, bool bRequireIncludeForClasses, TArray<FString>& RelativeIncludes)
 {
 	const FPropertyBase& PropertyBase = PropertyDef->GetPropertyBase();
-	AddIncludeForType(PropertyBase.TypeDef, RelativeIncludes);
+	AddIncludeForType(PropertyDef, PropertyBase.TypeDef, bRequireIncludeForClasses, RelativeIncludes);
 	if (PropertyBase.MapKeyProp != nullptr)
 	{
-		AddIncludeForType(PropertyBase.MapKeyProp->TypeDef, RelativeIncludes);
+		AddIncludeForType(PropertyDef, PropertyBase.MapKeyProp->TypeDef, false, RelativeIncludes);
 	}
 }
 
@@ -5663,7 +5675,8 @@ bool FNativeClassHeaderGenerator::WriteSource(const FManifestModule& Module, FGe
 				// Functions
 				for (const TSharedRef<FUnrealFunctionDefinitionInfo>& FunctionDef : Struct->GetFunctions())
 				{
-					if (!(FunctionDef->GetFunctionData().FunctionExportFlags & FUNCEXPORT_CppStatic) && FunctionDef->HasAnyFunctionFlags(FUNC_NetValidate))
+					const FFuncInfo& FunctionData = FunctionDef->GetFunctionData();
+					if (!(FunctionData.FunctionExportFlags & FUNCEXPORT_CppStatic) && FunctionDef->HasAnyFunctionFlags(FUNC_NetValidate))
 					{
 						if (!bAddedCoreNetHeader)
 						{
@@ -5671,16 +5684,22 @@ bool FNativeClassHeaderGenerator::WriteSource(const FManifestModule& Module, FGe
 							bAddedCoreNetHeader = true;
 						}
 					}
+
+					bool bRequireIncludeForClasses =
+						FunctionDef->HasAnyFunctionFlags(FUNC_Native) &&
+						(FunctionData.FunctionExportFlags & FUNCEXPORT_CustomThunk) == 0 &&
+						ShouldExportFunction(*FunctionDef);
+
 					for (const TSharedRef<FUnrealPropertyDefinitionInfo>& PropertyDef : FunctionDef->GetProperties())
 					{
-						AddIncludeForProperty(PropertyDef, RelativeIncludes);
+						AddIncludeForProperty(PropertyDef, bRequireIncludeForClasses, RelativeIncludes);
 					}
 				}
 
 				// Properties
 				for (const TSharedRef<FUnrealPropertyDefinitionInfo>& PropertyDef : Struct->GetProperties())
 				{
-					AddIncludeForProperty(PropertyDef, RelativeIncludes);
+					AddIncludeForProperty(PropertyDef, false, RelativeIncludes);
 				}
 			}
 
