Porting C# UHT changes for class constructor to C++ UHT

diff --git a/Engine/Source/Programs/UnrealHeaderTool/Private/CodeGenerator.cpp b/Engine/Source/Programs/UnrealHeaderTool/Private/CodeGenerator.cpp
--- a/Engine/Source/Programs/UnrealHeaderTool/Private/CodeGenerator.cpp
+++ b/Engine/Source/Programs/UnrealHeaderTool/Private/CodeGenerator.cpp
@@ -3107,15 +3107,16 @@ void ExportStandardConstructorsMacro(FOutputDevice& Out, FUnrealClassDefinitionI
 /**
  * Generates constructor definition.
  *
- * @param Out Output device to generate to.
+ * @param OutGeneratedHeaderText Output device for header
+ * @param OutGeneratedCppText Output device for cpp
  * @param Class Class to generate constructor for.
  * @param API API string for this constructor.
  */
-void ExportConstructorDefinition(FOutputDevice& Out, FUnrealClassDefinitionInfo& ClassDef, const TCHAR* API, const TCHAR* ClassCPPName)
+void ExportConstructorDefinition(FOutputDevice& OutGeneratedHeaderText, FOutputDevice& OutGeneratedCppText, FUnrealClassDefinitionInfo& ClassDef, const TCHAR* API, const TCHAR* ClassCPPName)
 {
 	if (!ClassDef.IsConstructorDeclared())
 	{
-		Out.Logf(TEXT("\t/** Standard constructor, called after all reflected properties have been initialized */\r\n"));
+		OutGeneratedHeaderText.Logf(TEXT("\t/** Standard constructor, called after all reflected properties have been initialized */\r\n"));
 
 		// Assume super class has OI constructor, this may not always be true but we should always be able to check this.
 		// In any case, it will default to old behaviour before we even checked this.
@@ -3140,18 +3141,26 @@ void ExportConstructorDefinition(FOutputDevice& Out, FUnrealClassDefinitionInfo&
 		}
 		if (bSuperClassObjectInitializerConstructorDeclared)
 		{
-			Out.Logf(TEXT("\t%s_API %s(const FObjectInitializer& ObjectInitializer = FObjectInitializer::Get()) : Super(ObjectInitializer) { };\r\n"), API, ClassCPPName);
+			OutGeneratedHeaderText.Logf(TEXT("\t%s_API %s(const FObjectInitializer& ObjectInitializer = FObjectInitializer::Get());\r\n"), API, ClassCPPName);
+			if (!ClassDef.UsesGeneratedBodyLegacy())
+			{
+				OutGeneratedCppText.Logf(TEXT("\t%s::%s(const FObjectInitializer& ObjectInitializer) : Super(ObjectInitializer) {}\r\n"), ClassCPPName, ClassCPPName);
+			}
 			ClassDef.MarkObjectInitializerConstructorDeclared();
 		}
 		else
 		{
-			Out.Logf(TEXT("\t%s_API %s() { };\r\n"), API, ClassCPPName);
+			OutGeneratedHeaderText.Logf(TEXT("\t%s_API %s();\r\n"), API, ClassCPPName);
+			if (!ClassDef.UsesGeneratedBodyLegacy())
+			{
+				OutGeneratedCppText.Logf(TEXT("\t%s::%s() {}\r\n"), ClassCPPName, ClassCPPName);
+			}
 			ClassDef.MarkDefaultConstructorDeclared();
 		}
 
 		ClassDef.MarkConstructorDeclared();
 	}
-	ExportCopyConstructorDefinition(Out, API, ClassCPPName);
+	ExportCopyConstructorDefinition(OutGeneratedHeaderText, API, ClassCPPName);
 }
 
 /**
@@ -3197,12 +3206,12 @@ void ExportDefaultConstructorCallDefinition(FOutputDevice& Out, FUnrealClassDefi
  * @param Class Class to generate constructor for.
  * @param API API string for this constructor.
  */
-void ExportEnhancedConstructorsMacro(FOutputDevice& Out, FUnrealClassDefinitionInfo& ClassDef, const TCHAR* API, const TCHAR* ClassCPPName)
+void ExportEnhancedConstructorsMacro(FOutputDevice& OutGeneratedHeaderText, FOutputDevice& OutGeneratedCppText, FUnrealClassDefinitionInfo& ClassDef, const TCHAR* API, const TCHAR* ClassCPPName)
 {
-	ExportConstructorDefinition(Out, ClassDef, API, ClassCPPName);
-	ExportVTableHelperCtorAndCaller(Out, ClassDef, API, ClassCPPName);
-	ExportDefaultConstructorCallDefinition(Out, ClassDef, ClassCPPName);
-	ExportDestructorDefinition(Out, ClassDef, API, ClassCPPName);
+	ExportConstructorDefinition(OutGeneratedHeaderText, OutGeneratedCppText, ClassDef, API, ClassCPPName);
+	ExportVTableHelperCtorAndCaller(OutGeneratedHeaderText, ClassDef, API, ClassCPPName);
+	ExportDefaultConstructorCallDefinition(OutGeneratedHeaderText, ClassDef, ClassCPPName);
+	ExportDestructorDefinition(OutGeneratedHeaderText, ClassDef, API, ClassCPPName);
 }
 
 /**
@@ -3221,7 +3230,7 @@ FString GetBuildPath(const FUnrealSourceFile& SourceFile)
 	return Out;
 }
 
-void FNativeClassHeaderGenerator::ExportConstructorsMacros(FOutputDevice& OutGeneratedHeaderText, FOutputDevice& Out, FOutputDevice& StandardUObjectConstructorsMacroCall, FOutputDevice& EnhancedUObjectConstructorsMacroCall, const FString& ConstructorsMacroPrefix, FUnrealClassDefinitionInfo& ClassDef, const TCHAR* APIArg)
+void FNativeClassHeaderGenerator::ExportConstructorsMacros(FOutputDevice& OutGeneratedHeaderText, FOutputDevice& OutGeneratedCppText, FOutputDevice& StandardUObjectConstructorsMacroCall, FOutputDevice& EnhancedUObjectConstructorsMacroCall, const FString& ConstructorsMacroPrefix, FUnrealClassDefinitionInfo& ClassDef, const TCHAR* APIArg)
 {
 	const FString ClassCPPName = ClassDef.GetAlternateNameCPP();
 
@@ -3231,16 +3240,16 @@ void FNativeClassHeaderGenerator::ExportConstructorsMacros(FOutputDevice& OutGen
 	FString EnhMacroName = ConstructorsMacroPrefix + TEXT("_ENHANCED_CONSTRUCTORS");
 
 	ExportStandardConstructorsMacro(StdMacro, ClassDef, APIArg, *ClassCPPName);
-	ExportEnhancedConstructorsMacro(EnhMacro, ClassDef, APIArg, *ClassCPPName);
+	ExportEnhancedConstructorsMacro(EnhMacro, OutGeneratedCppText, ClassDef, APIArg, *ClassCPPName);
 
 	if (!ClassDef.IsCustomVTableHelperConstructorDeclared())
 	{
-		Out.Logf(TEXT("\tDEFINE_VTABLE_PTR_HELPER_CTOR(%s);" LINE_TERMINATOR_ANSI), *ClassCPPName);
+		OutGeneratedCppText.Logf(TEXT("\tDEFINE_VTABLE_PTR_HELPER_CTOR(%s);" LINE_TERMINATOR_ANSI), *ClassCPPName);
 	}
 
 	if (!ClassDef.IsDestructorDeclared())
 	{
-		Out.Logf(TEXT("\t%s::~%s() {}" LINE_TERMINATOR_ANSI), *ClassCPPName, *ClassCPPName);
+		OutGeneratedCppText.Logf(TEXT("\t%s::~%s() {}" LINE_TERMINATOR_ANSI), *ClassCPPName, *ClassCPPName);
 	}
 
 	OutGeneratedHeaderText.Log(Macroize(*StdMacroName, *StdMacro));
diff --git a/Engine/Source/Programs/UnrealHeaderTool/Private/HeaderParser.cpp b/Engine/Source/Programs/UnrealHeaderTool/Private/HeaderParser.cpp
--- a/Engine/Source/Programs/UnrealHeaderTool/Private/HeaderParser.cpp
+++ b/Engine/Source/Programs/UnrealHeaderTool/Private/HeaderParser.cpp
@@ -5069,6 +5069,8 @@ bool FHeaderParser::CompileDeclaration(TArray<FUnrealFunctionDefinitionInfo*>& D
 
 		if (Token.IsIdentifier(TEXT("GENERATED_IINTERFACE_BODY"), ESearchCase::CaseSensitive))
 		{
+			// Due to the way IInterface and UInterface share the same class instance, we don't set the legacy body
+			//ClassDef.MarkUsesGeneratedBodyLegacy();
 			CurrentAccessSpecifier = ACCESS_Public;
 		}
 
@@ -5099,6 +5101,7 @@ bool FHeaderParser::CompileDeclaration(TArray<FUnrealFunctionDefinitionInfo*>& D
 		if (Token.IsIdentifier(TEXT("GENERATED_UINTERFACE_BODY"), ESearchCase::CaseSensitive))
 		{
 			CurrentAccessSpecifier = ACCESS_Public;
+			ClassDef.MarkUsesGeneratedBodyLegacy();
 		}
 		return true;
 	}
@@ -5120,6 +5123,7 @@ bool FHeaderParser::CompileDeclaration(TArray<FUnrealFunctionDefinitionInfo*>& D
 		}
 		else
 		{
+			ClassDef.MarkUsesGeneratedBodyLegacy();
 			CurrentAccessSpecifier = ACCESS_Public;
 		}
 
diff --git a/Engine/Source/Programs/UnrealHeaderTool/Private/NativeClassExporter.h b/Engine/Source/Programs/UnrealHeaderTool/Private/NativeClassExporter.h
--- a/Engine/Source/Programs/UnrealHeaderTool/Private/NativeClassExporter.h
+++ b/Engine/Source/Programs/UnrealHeaderTool/Private/NativeClassExporter.h
@@ -661,13 +661,14 @@ private:
 	/**
 	 * Exports macros that manages UObject constructors.
 	 * 
-	 * @param	VTableOut								The destination to write vtable helpers to.
+	 * @param	OutGeneratedHeaderText  				The destination to write header code.
+	 * @param	OutGeneratedCppText  					The destination to write cpp code.
 	 * @param	StandardUObjectConstructorsMacroCall	The destination to write standard constructor macros to.
 	 * @param	EnhancedUObjectConstructorsMacroCall	The destination to write enhanced constructor macros to.
 	 * @param	ConstructorsMacroPrefix					Prefix for constructors macro.
 	 * @param	Class									Class for which to export macros.
 	 */
-	static void ExportConstructorsMacros(FOutputDevice& OutGeneratedHeaderText, FOutputDevice& VTableOut, FOutputDevice& StandardUObjectConstructorsMacroCall, FOutputDevice& EnhancedUObjectConstructorsMacroCall, const FString& ConstructorsMacroPrefix, FUnrealClassDefinitionInfo& ClassDef, const TCHAR* APIArg);
+	static void ExportConstructorsMacros(FOutputDevice& OutGeneratedHeaderText, FOutputDevice& OutGeneratedCppText, FOutputDevice& StandardUObjectConstructorsMacroCall, FOutputDevice& EnhancedUObjectConstructorsMacroCall, const FString& ConstructorsMacroPrefix, FUnrealClassDefinitionInfo& ClassDef, const TCHAR* APIArg);
 
 	/**
 	 * Gets string with function return type.
diff --git a/Engine/Source/Programs/UnrealHeaderTool/Private/UnrealTypeDefinitionInfo.h b/Engine/Source/Programs/UnrealHeaderTool/Private/UnrealTypeDefinitionInfo.h
--- a/Engine/Source/Programs/UnrealHeaderTool/Private/UnrealTypeDefinitionInfo.h
+++ b/Engine/Source/Programs/UnrealHeaderTool/Private/UnrealTypeDefinitionInfo.h
@@ -3423,6 +3423,22 @@ public:
 		bNoExport = true;
 	}
 
+	/**
+	 * Return true if the class uses legacy generated body
+	 */
+	bool UsesGeneratedBodyLegacy() const
+	{
+		return bUsesGeneratedBodyLegacy;
+	}
+
+	/**
+	 * Mark the class as using legacy generated body
+	 */
+	void MarkUsesGeneratedBodyLegacy()
+	{
+		bUsesGeneratedBodyLegacy = true;
+	}
+
 	/**
 	 * Get the generated body access specifier
 	 */
@@ -3554,6 +3570,9 @@ private:
 	/** True if the class is not to be exported */
 	bool bNoExport = false;
 
+	/** True if the class uses legacy generated body */
+	bool bUsesGeneratedBodyLegacy = false;
+
 	/** True if the class is an interface */
 	bool bIsInterface = false;
 
