Modified C++ UHT to not do super class testing on native interfaces when the interface does not being with an 'I'.

diff --git a/Engine/Source/Programs/UnrealHeaderTool/Private/HeaderParser.cpp b/Engine/Source/Programs/UnrealHeaderTool/Private/HeaderParser.cpp
--- a/Engine/Source/Programs/UnrealHeaderTool/Private/HeaderParser.cpp
+++ b/Engine/Source/Programs/UnrealHeaderTool/Private/HeaderParser.cpp
@@ -6294,10 +6294,16 @@ FUnrealClassDefinitionInfo& FHeaderParser::CompileClassDeclaration()
 	return ClassDef;
 }
 
-FUnrealClassDefinitionInfo* FHeaderParser::ParseInterfaceNameDeclaration(FString& DeclaredInterfaceName, FString& RequiredAPIMacroIfPresent)
+FUnrealClassDefinitionInfo* FHeaderParser::ParseInterfaceNameDeclaration(FString& DeclaredInterfaceName, FString& RequiredAPIMacroIfPresent, bool bIsNativeInterface)
 {
 	ParseNameWithPotentialAPIMacroPrefix(/*out*/ DeclaredInterfaceName, /*out*/ RequiredAPIMacroIfPresent, TEXT("interface"));
 
+	// If we are expecting an native interface name and it doesn't start with 'I', then don't bother doing anything more
+	if (bIsNativeInterface && DeclaredInterfaceName[0] != 'I')
+	{
+		return nullptr;
+	}
+
 	FUnrealClassDefinitionInfo* ClassDef = FUnrealClassDefinitionInfo::FindClass(*GetClassNameWithPrefixRemoved(*DeclaredInterfaceName));
 	if (ClassDef == nullptr)
 	{
@@ -6344,7 +6350,7 @@ bool FHeaderParser::TryParseIInterfaceClass()
 	// Get a class name
 	FString DeclaredInterfaceName;
 	FString RequiredAPIMacroIfPresent;
-	if (ParseInterfaceNameDeclaration(/*out*/ DeclaredInterfaceName, /*out*/ RequiredAPIMacroIfPresent) == nullptr)
+	if (ParseInterfaceNameDeclaration(/*out*/ DeclaredInterfaceName, /*out*/ RequiredAPIMacroIfPresent, true) == nullptr)
 	{
 		return false;
 	}
@@ -6401,7 +6407,7 @@ void FHeaderParser::CompileInterfaceDeclaration()
 
 	// New style files have the interface name / extends afterwards
 	RequireIdentifier(TEXT("class"), ESearchCase::CaseSensitive, TEXT("Interface declaration"));
-	FUnrealClassDefinitionInfo* InterfaceClassDef = ParseInterfaceNameDeclaration(/*out*/ DeclaredInterfaceName, /*out*/ RequiredAPIMacroIfPresent);
+	FUnrealClassDefinitionInfo* InterfaceClassDef = ParseInterfaceNameDeclaration(/*out*/ DeclaredInterfaceName, /*out*/ RequiredAPIMacroIfPresent, false);
 	check(InterfaceClassDef);
 	InterfaceClassDef->GetDefinitionRange().Start = &Input[InputPos];
 
diff --git a/Engine/Source/Programs/UnrealHeaderTool/Private/HeaderParser.h b/Engine/Source/Programs/UnrealHeaderTool/Private/HeaderParser.h
--- a/Engine/Source/Programs/UnrealHeaderTool/Private/HeaderParser.h
+++ b/Engine/Source/Programs/UnrealHeaderTool/Private/HeaderParser.h
@@ -411,7 +411,7 @@ protected:
 	void CompileRigVMMethodDeclaration(FUnrealStructDefinitionInfo& StructDef);
 	void ParseRigVMMethodParameters(FUnrealStructDefinitionInfo& StructDef);
 
-	FUnrealClassDefinitionInfo* ParseInterfaceNameDeclaration(FString& DeclaredInterfaceName, FString& RequiredAPIMacroIfPresent);
+	FUnrealClassDefinitionInfo* ParseInterfaceNameDeclaration(FString& DeclaredInterfaceName, FString& RequiredAPIMacroIfPresent, bool bIsNativeInterface);
 	bool TryParseIInterfaceClass();
 
 	bool CompileStatement(TArray<FUnrealFunctionDefinitionInfo*>& DelegatesToFixup);
