Removing engine types from C++ UHT.

C++ UHT can no longer run external script generators but can still be used to validate C# UHT.

diff --git a/Engine/Source/Programs/UnrealHeaderTool/Private/CodeGenerator.cpp b/Engine/Source/Programs/UnrealHeaderTool/Private/CodeGenerator.cpp
--- a/Engine/Source/Programs/UnrealHeaderTool/Private/CodeGenerator.cpp
+++ b/Engine/Source/Programs/UnrealHeaderTool/Private/CodeGenerator.cpp
@@ -1704,7 +1704,7 @@ FString GetEventStructParamsName(FUnrealObjectDefinitionInfo& OuterDef, const TC
 	}
 	else if (FUnrealPackageDefinitionInfo* PackageDef = UHTCast<FUnrealPackageDefinitionInfo>(OuterDef))
 	{
-		OuterName = PackageDef->GetPackage()->GetName();
+		OuterName = PackageDef->GetName();
 		OuterName.ReplaceInline(TEXT("/"), TEXT("_"), ESearchCase::CaseSensitive);
 	}
 	else
@@ -1894,10 +1894,9 @@ bool IsDelegateFunction(FUnrealFieldDefinitionInfo& FieldDef)
 
 void FNativeClassHeaderGenerator::ExportGeneratedPackageInitCode(FOutputDevice& Out, const TCHAR* InDeclarations, const TArray<FGeneratedCPP*>& ExportedSorted, uint32 Hash)
 {
-	UPackage* Package = PackageDef.GetPackage();
 	const FString& SingletonName = GetPackageSingletonNameFuncAddr(PackageDef, nullptr);
 
-	FString PackageName = Package->GetName();
+	FString PackageName = PackageDef.GetName();
 	PackageName.ReplaceInline(TEXT("/"), TEXT("_"), ESearchCase::CaseSensitive);
 	FString BodyHashFn = FString::Printf(TEXT("Z_UPackage_%s_BodyHash"), *PackageName);
 	FString DeclarationsHashFn = FString::Printf(TEXT("Z_UPackage_%s_DeclarationsHash"), *PackageName);
@@ -1971,10 +1970,10 @@ void FNativeClassHeaderGenerator::ExportGeneratedPackageInitCode(FOutputDevice&
 	}
 
 	Out.Logf(TEXT("\t\t\tstatic const UECodeGen_Private::FPackageParams PackageParams = {\r\n"));
-	Out.Logf(TEXT("\t\t\t\t%s,\r\n"), *CreateUTF8LiteralString(Package->GetName()));
+	Out.Logf(TEXT("\t\t\t\t%s,\r\n"), *CreateUTF8LiteralString(PackageDef.GetName()));
 	Out.Logf(TEXT("\t\t\t\t%s,\r\n"), SingletonArray);
 	Out.Logf(TEXT("\t\t\t\t%s,\r\n"), SingletonCount);
-	Out.Logf(TEXT("\t\t\t\tPKG_CompiledIn | 0x%08X,\r\n"), Package->GetPackageFlags() & (PKG_ClientOptional | PKG_ServerSideOnly | PKG_EditorOnly | PKG_Developer | PKG_UncookedOnly));
+	Out.Logf(TEXT("\t\t\t\tPKG_CompiledIn | 0x%08X,\r\n"), PackageDef.GetPackageFlags() & (PKG_ClientOptional | PKG_ServerSideOnly | PKG_EditorOnly | PKG_Developer | PKG_UncookedOnly));
 	Out.Logf(TEXT("\t\t\t\t0x%08X,\r\n"), Hash);
 	Out.Logf(TEXT("\t\t\t\t0x%08X,\r\n"), DeclarationsHash);
 	Out.Logf(TEXT("\t\t\t\t%s\r\n"), *MetaDataParams);
@@ -1988,7 +1987,7 @@ void FNativeClassHeaderGenerator::ExportGeneratedPackageInitCode(FOutputDevice&
 	Out.Logf(TEXT("\tstatic FRegisterCompiledInInfo Z_CompiledInDeferPackage_UPackage_%s(%s, TEXT(\"%s\"), Z_Registration_Info_UPackage_%s, CONSTRUCT_RELOAD_VERSION_INFO(FPackageReloadVersionInfo, 0x%08X, 0x%08X));\r\n"),
 		*PackageName,
 		*SingletonName,
-		*Package->GetName(),
+		*PackageDef.GetName(),
 		*PackageName,
 		Hash,
 		DeclarationsHash
@@ -5878,13 +5877,13 @@ void FNativeClassHeaderGenerator::GenerateSourceFile(FGeneratedCPP& GeneratedCPP
 				if (FUnrealEnumDefinitionInfo* EnumDef = UHTCast<FUnrealEnumDefinitionInfo>(FieldDef))
 				{
 					// Is this ever not the case?
-					if (EnumDef->GetOuter()->GetObject()->IsA(UPackage::StaticClass()))
+					if (EnumDef->GetOuter()->AsPackage() != nullptr)
 					{
 						GeneratedFunctionDeclarations.Log(EnumDef->GetExternDecl(true));
 						Generator.ExportGeneratedEnumInitCode(GeneratedCPPText, ReferenceGatherers, SourceFile, *EnumDef);
 							EnumRegs.Add(EnumDef);
-						}
 					}
+				}
 				else if (FUnrealScriptStructDefinitionInfo* ScriptStructDef = UHTCast<FUnrealScriptStructDefinitionInfo>(FieldDef))
 				{
 					if (ScriptStructDef->HasAnyStructFlags(STRUCT_NoExport))
@@ -6219,8 +6218,7 @@ void FNativeClassHeaderGenerator::Generate(
 	TArray<FGeneratedCPP>& GeneratedCPPs
 )
 {
-	UPackage* Package = PackageDef.GetPackage();
-	FString PackageName = FPackageName::GetShortName(Package);
+	FString PackageName = PackageDef.GetShortName();
 	const FManifestModule& Module = PackageDef.GetModule();
 	const bool bWriteClassesH = PackageDef.GetWriteClassesH();
 	const bool bAllowSaveExportedHeaders = Module.SaveExportedHeaders;
@@ -6665,56 +6663,6 @@ void FNativeClassHeaderGenerator::ExportUpdatedHeaders(FString&& PackageName, TA
 	}
 }
 
-/** Get all script plugins based on ini setting */
-void GetScriptPlugins(TArray<IScriptGeneratorPluginInterface*>& ScriptPlugins)
-{
-	if (!GManifest.IsGameTarget)
-	{
-		UE_LOG(LogCompile, Log, TEXT("Script generator plugins only enabled in game targets."));
-	}
-
-	ScriptPlugins = IModularFeatures::Get().GetModularFeatureImplementations<IScriptGeneratorPluginInterface>(TEXT("ScriptGenerator"));
-	UE_LOG(LogCompile, Log, TEXT("Found %d script generator plugins."), ScriptPlugins.Num());
-
-	// Check if we can use these plugins and initialize them
-	for (int32 PluginIndex = ScriptPlugins.Num() - 1; PluginIndex >= 0; --PluginIndex)
-	{
-		IScriptGeneratorPluginInterface* ScriptGenerator = ScriptPlugins[PluginIndex];
-		bool bSupportedPlugin = ScriptGenerator->SupportsTarget(GManifest.TargetName);
-		if (bSupportedPlugin)
-		{
-			// Find the right output directory for this plugin base on its target (Engine-side) plugin name.
-			FString GeneratedCodeModuleName = ScriptGenerator->GetGeneratedCodeModuleName();
-			const FManifestModule* GeneratedCodeModule = NULL;
-			FString OutputDirectory;
-			FString IncludeBase;
-			for (const FManifestModule& Module : GManifest.Modules)
-			{
-				if (Module.Name == GeneratedCodeModuleName)
-				{
-					GeneratedCodeModule = &Module;
-				}
-			}
-			if (GeneratedCodeModule)
-			{
-				UE_LOG(LogCompile, Log, TEXT("Initializing script generator \'%s\'"), *ScriptGenerator->GetGeneratorName());
-				ScriptGenerator->Initialize(GManifest.RootLocalPath, GManifest.RootBuildPath, GeneratedCodeModule->GeneratedIncludeDirectory, GeneratedCodeModule->IncludeBase);
-			}
-			else
-			{
-				// Can't use this plugin
-				UE_LOG(LogCompile, Log, TEXT("Unable to determine output directory for %s. Cannot export script glue with \'%s\'"), *GeneratedCodeModuleName, *ScriptGenerator->GetGeneratorName());
-				bSupportedPlugin = false;
-			}
-		}
-		if (!bSupportedPlugin)
-		{
-			UE_LOG(LogCompile, Log, TEXT("Script generator \'%s\' not supported for target: %s"), *ScriptGenerator->GetGeneratorName(), *GManifest.TargetName);
-			ScriptPlugins.RemoveAt(PluginIndex);
-		}
-	}
-}
-
 /**
  * Tries to resolve super classes for classes defined in the given class
  */
@@ -7048,7 +6996,6 @@ void DefineTypes(TArray<FUnrealPackageDefinitionInfo*>& PackageDefs)
 			FResults::Try([PackageDef, &SourceFile = *SourceFile]()
 			{
 				TArray<TSharedRef<FUnrealTypeDefinitionInfo>>& AllClasses = PackageDef->GetAllClasses();
-				UPackage* Package = PackageDef->GetPackage();
 
 				for (TSharedRef<FUnrealTypeDefinitionInfo>& TypeDef : SourceFile.GetDefinedClasses())
 				{
@@ -7314,33 +7261,6 @@ void PostParseFinalize(TArray<FUnrealPackageDefinitionInfo*>& PackageDefs)
 	FResults::WaitForErrorTasks();
 }
 
-void CreateEngineTypes(TArray<FUnrealPackageDefinitionInfo*>& PackageDefs)
-{
-	auto CreateEngineTypes = [&PackageDefs](ECreateEngineTypesPhase Phase)
-	{
-		for (FUnrealPackageDefinitionInfo* PackageDef : PackageDefs)
-		{
-			FResults::Try([PackageDef, Phase]()
-				{
-					PackageDef->CreateUObjectEngineTypes(Phase);
-				});
-		}
-
-		for (FUnrealPackageDefinitionInfo* PackageDef : PackageDefs)
-		{
-			FResults::Try([PackageDef, Phase]()
-				{
-					PackageDef->CreateUObjectEngineTypes(Phase);
-				});
-		}
-	};
-
-	CreateEngineTypes(ECreateEngineTypesPhase::Phase1);
-	CreateEngineTypes(ECreateEngineTypesPhase::Phase2);
-	CreateEngineTypes(ECreateEngineTypesPhase::Phase3);
-	FResults::WaitForErrorTasks();
-}
-
 void Export(TArray<FUnrealPackageDefinitionInfo*>& PackageDefs, TArray<FUnrealSourceFile*>& OrderedSourceFiles)
 {
 	TArray<FGeneratedCPP> GeneratedCPPs;
@@ -7360,90 +7280,6 @@ void Export(TArray<FUnrealPackageDefinitionInfo*>& PackageDefs, TArray<FUnrealSo
 	FResults::WaitForErrorTasks();
 }
 
-// Exports the class to all available plugins
-void ExportClassToScriptPlugins(const TMap<UClass*, FUnrealSourceFile*>& SourceFileLookup, UClass* Class, const FManifestModule& Module, IScriptGeneratorPluginInterface& ScriptPlugin)
-{
-	check(SourceFileLookup.Find(Class) != nullptr);
-	FUnrealSourceFile* const * SourceFile = SourceFileLookup.Find(Class);
-	if (SourceFile == nullptr)
-	{
-		const FString Empty = TEXT("");
-		ScriptPlugin.ExportClass(Class, Empty, Empty, false);
-	}
-	else
-	{
-		ScriptPlugin.ExportClass(Class, (*SourceFile)->GetFilename(), (*SourceFile)->GetGeneratedFilename(), (*SourceFile)->HasChanged());
-	}
-}
-
-// Exports class tree to all available plugins
-void ExportClassTreeToScriptPlugins(const TMap<UClass*, FUnrealSourceFile*>& SourceFileLookup, const FClassTree* Node, const FManifestModule& Module, IScriptGeneratorPluginInterface& ScriptPlugin)
-{
-	for (int32 ChildIndex = 0; ChildIndex < Node->NumChildren(); ++ChildIndex)
-	{
-		const FClassTree* ChildNode = Node->GetChild(ChildIndex);
-		ExportClassToScriptPlugins(SourceFileLookup, ChildNode->GetClass(), Module, ScriptPlugin);
-	}
-
-	for (int32 ChildIndex = 0; ChildIndex < Node->NumChildren(); ++ChildIndex)
-	{
-		const FClassTree* ChildNode = Node->GetChild(ChildIndex);
-		ExportClassTreeToScriptPlugins(SourceFileLookup, ChildNode, Module, ScriptPlugin);
-	}
-}
-
-void ExportToScriptPlugins(TArray<IScriptGeneratorPluginInterface*>& ScriptPlugins, TArray<FUnrealPackageDefinitionInfo*>& PackageDefs, FString& ExternalDependencies)
-{
-	TMap<UClass*, FUnrealSourceFile*> SourceFileLookup;
-	for (FUnrealPackageDefinitionInfo* PackageDef : PackageDefs)
-	{
-		for (TSharedRef<FUnrealTypeDefinitionInfo> TypeDef : PackageDef->GetAllClasses())
-		{
-			UClass* Class = UHTCastChecked<FUnrealClassDefinitionInfo>(TypeDef).GetClass();
-			SourceFileLookup.Add(Class, &TypeDef->GetUnrealSourceFile());
-		}
-	}
-
-	for (FUnrealPackageDefinitionInfo* PackageDef : PackageDefs)
-	{
-		const FManifestModule& Module = PackageDef->GetModule();
-
-		FClassTree ClassTree(UObject::StaticClass());
-		for (TSharedRef<FUnrealTypeDefinitionInfo> TypeDef : PackageDef->GetAllClasses())
-		{
-			ClassTree.AddClass(UHTCastChecked<FUnrealClassDefinitionInfo>(TypeDef).GetClass());
-		}
-		ClassTree.Validate();
-
-		for (IScriptGeneratorPluginInterface* Plugin : ScriptPlugins)
-		{
-			if (Plugin->ShouldExportClassesForModule(Module.Name, Module.ModuleType, Module.GeneratedIncludeDirectory))
-			{
-				ExportClassToScriptPlugins(SourceFileLookup, ClassTree.GetClass(), Module, *Plugin);
-				ExportClassTreeToScriptPlugins(SourceFileLookup, &ClassTree, Module, *Plugin);
-			}
-		}
-	}
-
-	for (IScriptGeneratorPluginInterface* ScriptGenerator : ScriptPlugins)
-	{
-		ScriptGenerator->FinishExport();
-	}
-
-	// Get a list of external dependencies from each enabled plugin
-	for (IScriptGeneratorPluginInterface* ScriptPlugin : ScriptPlugins)
-	{
-		TArray<FString> PluginExternalDependencies;
-		ScriptPlugin->GetExternalDependencies(PluginExternalDependencies);
-
-		for (const FString& PluginExternalDependency : PluginExternalDependencies)
-		{
-			ExternalDependencies += PluginExternalDependency + LINE_TERMINATOR;
-		}
-	}
-	return;
-}
-
 void WriteExternalDependencies(const FString& ExternalDependencies)
 {
 	FFileHelper::SaveStringToFile(ExternalDependencies, *GManifest.ExternalDependenciesFile);
@@ -7500,7 +7336,6 @@ ECompilationResult::Type UnrealHeaderTool_Main(const FString& ModuleInfoFilename
 	double TotalPrepareModuleTime = FResults::TimedTry([&PackageDefs, &ModuleInfoPath]() { PrepareModules(PackageDefs, ModuleInfoPath); });
 
 	FString ExternalDependencies;
-	TArray<IScriptGeneratorPluginInterface*> ScriptPlugins;
 
 	double TotalPreparseTime = FResults::TimedTry([&PackageDefs, &ModuleInfoPath]() { PreparseSources(PackageDefs, ModuleInfoPath); });
 	double TotalDefineTypesTime = FResults::TimedTry([&PackageDefs]() { DefineTypes(PackageDefs); });
@@ -7521,9 +7356,6 @@ ECompilationResult::Type UnrealHeaderTool_Main(const FString& ModuleInfoFilename
 
 	TotalTopologicalSortTime += FResults::TimedTry([&OrderedSourceFiles]() { TopologicalSort(OrderedSourceFiles); }); // Sort again to include new dependencies
 	double TotalCodeGenTime = FResults::TimedTry([&PackageDefs, &OrderedSourceFiles]() { Export(PackageDefs, OrderedSourceFiles); });
-	double TotalCheckForScriptPluginsTime = FResults::TimedTry([&ScriptPlugins]() { GetScriptPlugins(ScriptPlugins); });
-	double TotalCreateEngineTypesTime = ScriptPlugins.IsEmpty() ? 0.0 : FResults::TimedTry([&PackageDefs]() { CreateEngineTypes(PackageDefs); });
-	double TotalPluginTime = ScriptPlugins.IsEmpty() ? 0.0 : FResults::TimedTry([&ScriptPlugins, &PackageDefs, &ExternalDependencies]() { ExportToScriptPlugins(ScriptPlugins, PackageDefs, ExternalDependencies); });
 	double TotalWriteExternalDependenciesTime = FResults::TimedTry([&ExternalDependencies]() { WriteExternalDependencies(ExternalDependencies); });
 	double TotalSummaryTime = FResults::TimedTry([&PackageDefs]() { GenerateSummary(PackageDefs); });
 
@@ -7567,9 +7399,6 @@ ECompilationResult::Type UnrealHeaderTool_Main(const FString& ModuleInfoFilename
 	UE_LOG(LogCompile, Log, TEXT("Parsing took %.3f seconds"), TotalParseTime);
 	UE_LOG(LogCompile, Log, TEXT("Post parse finalization took %.3f seconds"), TotalPostParseFinalizeTime);
 	UE_LOG(LogCompile, Log, TEXT("Code generation took %.3f seconds"), TotalCodeGenTime);
-	UE_LOG(LogCompile, Log, TEXT("Check for script plugins took %.3f seconds"), TotalCheckForScriptPluginsTime);
-	UE_LOG(LogCompile, Log, TEXT("Create engine types took % .3f seconds"), TotalCreateEngineTypesTime);
-	UE_LOG(LogCompile, Log, TEXT("ScriptPlugin overhead was %.3f seconds"), TotalPluginTime);
 	UE_LOG(LogCompile, Log, TEXT("Write external dependencies overhead was %.3f seconds"), TotalWriteExternalDependenciesTime);
 	UE_LOG(LogCompile, Log, TEXT("Summary generation took %.3f seconds"), TotalSummaryTime);
 	UE_LOG(LogCompile, Log, TEXT("Macroize time was %.3f CPU seconds"), GMacroizeTime);
@@ -7615,7 +7444,6 @@ ECompilationResult::Type UnrealHeaderTool_Main(const FString& ModuleInfoFilename
 
 void ProcessParsedClass(FUnrealClassDefinitionInfo& ClassDef)
 {
-	UPackage* Package = ClassDef.GetPackageDef().GetPackage();
 	const FString& ClassName = ClassDef.GetNameCPP();
 	FString ClassNameStripped = GetClassNameWithPrefixRemoved(*ClassName);
 	const FString& BaseClassName = ClassDef.GetSuperStructInfo().Name;
@@ -7671,7 +7499,6 @@ void ProcessParsedClass(FUnrealClassDefinitionInfo& ClassDef)
 
 void ProcessParsedEnum(FUnrealEnumDefinitionInfo& EnumDef)
 {
-	UPackage* Package = EnumDef.GetPackageDef().GetPackage();
 	const FString& EnumName = EnumDef.GetNameCPP();
 
 	if (TSharedRef<FUnrealTypeDefinitionInfo>* Existing = GTypeDefinitionInfoMap.FindByName(*EnumName))
diff --git a/Engine/Source/Programs/UnrealHeaderTool/Private/HeaderParser.cpp b/Engine/Source/Programs/UnrealHeaderTool/Private/HeaderParser.cpp
--- a/Engine/Source/Programs/UnrealHeaderTool/Private/HeaderParser.cpp
+++ b/Engine/Source/Programs/UnrealHeaderTool/Private/HeaderParser.cpp
@@ -94,6 +94,47 @@ extern std::atomic<int> GSourcesParsing;
 extern std::atomic<int> GSourcesCompleted;
 extern std::atomic<int> GSourcesStalled;
 
+// Singleton class to get the cast flag for a given class name.
+struct ClassCastFlagMap
+{
+	static ClassCastFlagMap& Get();
+
+	// Get Mapped name -> cast flag. Returns CASTCLASS_None if name is not found.
+	EClassCastFlags GetCastFlag(const FString& ClassName) const;
+
+private:
+	ClassCastFlagMap();
+	TMap<FString, EClassCastFlags> CastFlagMap;
+};
+
+
+//////////////////////////////////////////////////////////////////////////
+// ClassCastFlagMap
+
+ClassCastFlagMap::ClassCastFlagMap()
+{
+// Define macro to be applied to all cast class declarations.
+#define DECLARE_CAST_BY_FLAG(ClassName) CastFlagMap.Add(FString(#ClassName), CASTCLASS_##ClassName);
+
+// Now apply the macro to the whole list.
+DECLARE_ALL_CAST_FLAGS;
+
+#undef DECLARE_CAST_BY_FLAG
+}
+
+ClassCastFlagMap& ClassCastFlagMap::Get()
+{
+	static ClassCastFlagMap TheInstance;
+	return TheInstance;
+}
+
+EClassCastFlags ClassCastFlagMap::GetCastFlag(const FString& ClassName) const
+{
+	const EClassCastFlags* ValuePtr = CastFlagMap.Find(ClassName);
+	return ValuePtr ? *ValuePtr : CASTCLASS_None;
+}
+
+
 /*-----------------------------------------------------------------------------
 	Utility functions.
 -----------------------------------------------------------------------------*/
@@ -1881,10 +1922,6 @@ FUnrealScriptStructDefinitionInfo& FHeaderParser::CompileStructDeclaration()
 
 		StructDef.GetSuperStructInfo().Struct = BaseStructDef->AsStruct();
 		StructDef.SetStructFlags(EStructFlags(BaseStructDef->GetStructFlags() & STRUCT_Inherit));
-		if (StructDef.GetScriptStructSafe() != nullptr && BaseStructDef->GetScriptStructSafe())
-		{
-			StructDef.GetScriptStructSafe()->SetSuperStruct(BaseStructDef->GetScriptStructSafe());
-		}
 	}
 
 	Scope->AddType(StructDef);
@@ -9328,6 +9365,7 @@ TSharedRef<FUnrealTypeDefinitionInfo> FHeaderPreParser::ParseStructDeclaration(c
 			{
 				LogError(TEXT("The 'HasDefaults' struct specifier is only valid in the NoExportTypes.h file"));
 			}
+			StructDef->SetScriptStructExportFlags(STRUCTEXPORT_HasDefaults);
 			break;
 
 		case EStructSpecifier::HasNoOpConstructor:
@@ -9335,6 +9373,7 @@ TSharedRef<FUnrealTypeDefinitionInfo> FHeaderPreParser::ParseStructDeclaration(c
 			{
 				LogError(TEXT("The 'HasNoOpConstructor' struct specifier is only valid in the NoExportTypes.h file"));
 			}
+			StructDef->SetScriptStructExportFlags(STRUCTEXPORT_HasNoOpConstructor);
 			break;
 
 		case EStructSpecifier::IsAlwaysAccessible:
@@ -9342,6 +9381,7 @@ TSharedRef<FUnrealTypeDefinitionInfo> FHeaderPreParser::ParseStructDeclaration(c
 			{
 				LogError(TEXT("The 'IsAlwaysAccessible' struct specifier is only valid in the NoExportTypes.h file"));
 			}
+			StructDef->SetScriptStructExportFlags(STRUCTEXPORT_IsAlwaysAccessible);
 			break;
 
 		case EStructSpecifier::IsCoreType:
@@ -9349,6 +9389,7 @@ TSharedRef<FUnrealTypeDefinitionInfo> FHeaderPreParser::ParseStructDeclaration(c
 			{
 				LogError(TEXT("The 'IsCoreType' struct specifier is only valid in the NoExportTypes.h file"));
 			}
+			StructDef->SetScriptStructExportFlags(STRUCTEXPORT_IsCoreType);
 			break;
 		}
 	}
diff --git a/Engine/Source/Programs/UnrealHeaderTool/Private/HeaderParser.h b/Engine/Source/Programs/UnrealHeaderTool/Private/HeaderParser.h
--- a/Engine/Source/Programs/UnrealHeaderTool/Private/HeaderParser.h
+++ b/Engine/Source/Programs/UnrealHeaderTool/Private/HeaderParser.h
@@ -16,7 +16,6 @@ enum class EGeneratedCodeVersion : uint8;
 class FFeedbackContext;
 class UPackage;
 struct FManifestModule;
-class IScriptGeneratorPluginInterface;
 class FStringOutputDevice;
 class FUnrealSourceFile;
 class FScope;
diff --git a/Engine/Source/Programs/UnrealHeaderTool/Private/HeaderProvider.cpp b/Engine/Source/Programs/UnrealHeaderTool/Private/HeaderProvider.cpp
--- a/Engine/Source/Programs/UnrealHeaderTool/Private/HeaderProvider.cpp
+++ b/Engine/Source/Programs/UnrealHeaderTool/Private/HeaderProvider.cpp
@@ -87,7 +87,7 @@ FUnrealSourceFile* FHeaderProvider::Resolve(const FUnrealSourceFile& ParentSourc
 
 		// There is questionable compatibility hack where a source file will always be exported
 		// regardless of having types when it is being included by the SAME package.
-		if (Cache && Cache->GetPackage() == ParentSourceFile.GetPackage())
+		if (Cache && &Cache->GetPackageDef() == &ParentSourceFile.GetPackageDef())
 		{
 			Cache->MarkReferenced();
 		}
diff --git a/Engine/Source/Programs/UnrealHeaderTool/Private/ParserHelper.h b/Engine/Source/Programs/UnrealHeaderTool/Private/ParserHelper.h
--- a/Engine/Source/Programs/UnrealHeaderTool/Private/ParserHelper.h
+++ b/Engine/Source/Programs/UnrealHeaderTool/Private/ParserHelper.h
@@ -53,6 +53,14 @@ enum EPropertyHeaderExportFlags
 	PROPEXPORT_Protected	=0x00000004,	// property should be exported as protected
 };
 
+enum EScriptStructExportFlags
+{
+	STRUCTEXPORT_HasDefaults		= 0x00000001,
+	STRUCTEXPORT_HasNoOpConstructor = 0x00000002,
+	STRUCTEXPORT_IsAlwaysAccessible = 0x00000004,
+	STRUCTEXPORT_IsCoreType			= 0x00000008,
+};
+
 struct EPointerType
 {
 	enum Type
diff --git a/Engine/Source/Programs/UnrealHeaderTool/Private/PropertyTypes.cpp b/Engine/Source/Programs/UnrealHeaderTool/Private/PropertyTypes.cpp
--- a/Engine/Source/Programs/UnrealHeaderTool/Private/PropertyTypes.cpp
+++ b/Engine/Source/Programs/UnrealHeaderTool/Private/PropertyTypes.cpp
@@ -351,15 +351,6 @@ namespace
 		}
 	};
 
-	template<typename TraitsType>
-	struct CreateEngineTypeDispatch
-	{
-		FProperty* operator()(FUnrealPropertyDefinitionInfo& PropDef, FFieldVariant Scope, const FName& Name, EObjectFlags ObjectFlags)
-		{
-			return TraitsType::CreateEngineType(PropDef, Scope, Name, ObjectFlags);
-		}
-	};
-
 	template<typename TraitsType>
 	struct IsSupportedByBlueprintDispatch
 	{
@@ -426,12 +417,6 @@ namespace
 		return PropertyTypeDispatch<CreatePropertyDispatch, bHandleContainers, void>(PropDef);
 	}
 
-	template <bool bHandleContainers>
-	FProperty* CreateEngineTypeHelper(FUnrealPropertyDefinitionInfo& PropDef, FFieldVariant Scope, const FName& Name, EObjectFlags ObjectFlags)
-	{
-		return PropertyTypeDispatch<CreateEngineTypeDispatch, bHandleContainers, FProperty*>(PropDef, Scope, std::ref(Name), ObjectFlags);
-	}
-
 	bool IsSupportedByBlueprintSansContainers(const FUnrealPropertyDefinitionInfo& PropDef, bool bMemberVariable)
 	{
 		return PropertyTypeDispatch<IsSupportedByBlueprintDispatch, false, bool>(PropDef, bMemberVariable);
@@ -572,15 +557,6 @@ struct FPropertyTypeTraitsByte : public FPropertyTypeTraitsNumericBase
 		}
 	}
 
-	static FProperty* CreateEngineType(FUnrealPropertyDefinitionInfo& PropDef, FFieldVariant Scope, const FName& Name, EObjectFlags ObjectFlags)
-	{
-		const FPropertyBase& VarProperty = PropDef.GetPropertyBase();
-		FByteProperty* Result = new FByteProperty(Scope, Name, ObjectFlags);
-		Result->Enum = VarProperty.EnumDef ? VarProperty.EnumDef->GetEnum() : nullptr;
-		check(VarProperty.IntType == EIntType::Sized);
-		return Result;
-	}
-
 	static bool IsSupportedByBlueprint(const FUnrealPropertyDefinitionInfo& PropDef, bool bMemberVariable)
 	{
 		return true;
@@ -640,14 +616,6 @@ struct FPropertyTypeTraitsInt8 : public FPropertyTypeTraitsNumericBase
 	{
 	}
 
-	static FProperty* CreateEngineType(FUnrealPropertyDefinitionInfo& PropDef, FFieldVariant Scope, const FName& Name, EObjectFlags ObjectFlags)
-	{
-		const FPropertyBase& VarProperty = PropDef.GetPropertyBase();
-		FInt8Property* Result = new FInt8Property(Scope, Name, ObjectFlags);
-		check(VarProperty.IntType == EIntType::Sized);
-		return Result;
-	}
-
 	static FString GetEngineClassName(const FUnrealPropertyDefinitionInfo& PropDef)
 	{
 		return FInt8Property::StaticClass()->GetName();
@@ -665,14 +633,6 @@ struct FPropertyTypeTraitsInt16 : public FPropertyTypeTraitsNumericBase
 	{
 	}
 
-	static FProperty* CreateEngineType(FUnrealPropertyDefinitionInfo& PropDef, FFieldVariant Scope, const FName& Name, EObjectFlags ObjectFlags)
-	{
-		const FPropertyBase& VarProperty = PropDef.GetPropertyBase();
-		FInt16Property* Result = new FInt16Property(Scope, Name, ObjectFlags);
-		check(VarProperty.IntType == EIntType::Sized);
-		return Result;
-	}
-
 	static FString GetEngineClassName(const FUnrealPropertyDefinitionInfo& PropDef)
 	{
 		return FInt16Property::StaticClass()->GetName();
@@ -703,13 +663,6 @@ struct FPropertyTypeTraitsInt : public FPropertyTypeTraitsNumericBase
 		PropDef.SetUnsized(VarProperty.IntType == EIntType::Unsized);
 	}
 
-	static FProperty* CreateEngineType(FUnrealPropertyDefinitionInfo& PropDef, FFieldVariant Scope, const FName& Name, EObjectFlags ObjectFlags)
-	{
-		const FPropertyBase& VarProperty = PropDef.GetPropertyBase();
-		FIntProperty* Result = new FIntProperty(Scope, Name, ObjectFlags);
-		return Result;
-	}
-
 	static bool IsSupportedByBlueprint(const FUnrealPropertyDefinitionInfo& PropDef, bool bMemberVariable)
 	{
 		return true;
@@ -743,14 +696,6 @@ struct FPropertyTypeTraitsInt64 : public FPropertyTypeTraitsNumericBase
 	{
 	}
 
-	static FProperty* CreateEngineType(FUnrealPropertyDefinitionInfo& PropDef, FFieldVariant Scope, const FName& Name, EObjectFlags ObjectFlags)
-	{
-		const FPropertyBase& VarProperty = PropDef.GetPropertyBase();
-		FInt64Property* Result = new FInt64Property(Scope, Name, ObjectFlags);
-		check(VarProperty.IntType == EIntType::Sized);
-		return Result;
-	}
-
 	static bool IsSupportedByBlueprint(const FUnrealPropertyDefinitionInfo& PropDef, bool bMemberVariable)
 	{
 		return true;
@@ -773,14 +718,6 @@ struct FPropertyTypeTraitsUInt16 : public FPropertyTypeTraitsNumericBase
 	{
 	}
 
-	static FProperty* CreateEngineType(FUnrealPropertyDefinitionInfo& PropDef, FFieldVariant Scope, const FName& Name, EObjectFlags ObjectFlags)
-	{
-		const FPropertyBase& VarProperty = PropDef.GetPropertyBase();
-		FUInt16Property* Result = new FUInt16Property(Scope, Name, ObjectFlags);
-		check(VarProperty.IntType == EIntType::Sized);
-		return Result;
-	}
-
 	static FString GetEngineClassName(const FUnrealPropertyDefinitionInfo& PropDef)
 	{
 		return FUInt16Property::StaticClass()->GetName();
@@ -800,12 +737,6 @@ struct FPropertyTypeTraitsUInt32 : public FPropertyTypeTraitsNumericBase
 		PropDef.SetUnsized(VarProperty.IntType == EIntType::Unsized);
 	}
 
-	static FProperty* CreateEngineType(FUnrealPropertyDefinitionInfo& PropDef, FFieldVariant Scope, const FName& Name, EObjectFlags ObjectFlags)
-	{
-		FUInt32Property* Result = new FUInt32Property(Scope, Name, ObjectFlags);
-		return Result;
-	}
-
 	static FString GetEngineClassName(const FUnrealPropertyDefinitionInfo& PropDef)
 	{
 		return FUInt32Property::StaticClass()->GetName();
@@ -823,14 +754,6 @@ struct FPropertyTypeTraitsUInt64 : public FPropertyTypeTraitsNumericBase
 	{
 	}
 
-	static FProperty* CreateEngineType(FUnrealPropertyDefinitionInfo& PropDef, FFieldVariant Scope, const FName& Name, EObjectFlags ObjectFlags)
-	{
-		const FPropertyBase& VarProperty = PropDef.GetPropertyBase();
-		FUInt64Property* Result = new FUInt64Property(Scope, Name, ObjectFlags);
-		check(VarProperty.IntType == EIntType::Sized);
-		return Result;
-	}
-
 	static FString GetEngineClassName(const FUnrealPropertyDefinitionInfo& PropDef)
 	{
 		return FUInt64Property::StaticClass()->GetName();
@@ -859,12 +782,6 @@ struct FPropertyTypeTraitsFloat : public FPropertyTypeTraitsNumericBase
 	{
 	}
 
-	static FProperty* CreateEngineType(FUnrealPropertyDefinitionInfo& PropDef, FFieldVariant Scope, const FName& Name, EObjectFlags ObjectFlags)
-	{
-		FFloatProperty* Result = new FFloatProperty(Scope, Name, ObjectFlags);
-		return Result;
-	}
-
 	static bool IsSupportedByBlueprint(const FUnrealPropertyDefinitionInfo& PropDef, bool bMemberVariable)
 	{
 		return true;
@@ -898,12 +815,6 @@ struct FPropertyTypeTraitsDouble : public FPropertyTypeTraitsNumericBase
 	{
 	}
 
-	static FProperty* CreateEngineType(FUnrealPropertyDefinitionInfo& PropDef, FFieldVariant Scope, const FName& Name, EObjectFlags ObjectFlags)
-	{
-		FDoubleProperty* Result = new FDoubleProperty(Scope, Name, ObjectFlags);
-		return Result;
-	}
-
 	static bool IsSupportedByBlueprint(const FUnrealPropertyDefinitionInfo& PropDef, bool bMemberVariable)
 	{
 		return true;
@@ -937,12 +848,6 @@ struct FPropertyTypeTraitsLargeWorldCoordinatesReal : public FPropertyTypeTraits
 	{
 	}
 
-	static FProperty* CreateEngineType(FUnrealPropertyDefinitionInfo& PropDef, FFieldVariant Scope, const FName& Name, EObjectFlags ObjectFlags)
-	{
-		FLargeWorldCoordinatesRealProperty* Result = new FLargeWorldCoordinatesRealProperty(Scope, Name, ObjectFlags);
-		return Result;
-	}
-
 	static bool IsSupportedByBlueprint(const FUnrealPropertyDefinitionInfo& PropDef, bool bMemberVariable)
 	{
 		return true;
@@ -974,15 +879,6 @@ struct FPropertyTypeTraitsBooleanBase : public FPropertyTypeTraitsBase
 	{
 	}
 
-	template <typename SizeType, bool bIsNativeBool>
-	static FProperty* CreateEngineTypeHelper(FUnrealPropertyDefinitionInfo& PropDef, FFieldVariant Scope, const FName& Name, EObjectFlags ObjectFlags)
-	{
-		FBoolProperty* Result = new FBoolProperty(Scope, Name, ObjectFlags);
-		bool bActsLikeNativeBool = bIsNativeBool || PropDef.GetVariableCategory() == EVariableCategory::Return;
-		Result->SetBoolSize(bActsLikeNativeBool ? sizeof(bool) : sizeof(SizeType), bActsLikeNativeBool);
-		return Result;
-	}
-
 	static FString GetCPPTypeHelper(const FUnrealPropertyDefinitionInfo& PropDef, FString* ExtendedTypeText, uint32 CPPExportFlags, bool bIsNativeBool, const TCHAR* SizeText)
 	{
 		if (bIsNativeBool
@@ -1028,11 +924,6 @@ struct FPropertyTypeTraitsBooleanBase : public FPropertyTypeTraitsBase
 
 struct FPropertyTypeTraitsBool : public FPropertyTypeTraitsBooleanBase
 {
-	static FProperty* CreateEngineType(FUnrealPropertyDefinitionInfo& PropDef, FFieldVariant Scope, const FName& Name, EObjectFlags ObjectFlags)
-	{
-		return CreateEngineTypeHelper<bool, true>(PropDef, Scope, Name, ObjectFlags);
-	}
-
 	static FString GetCPPType(const FUnrealPropertyDefinitionInfo& PropDef, FString* ExtendedTypeText, uint32 CPPExportFlags)
 	{
 		return GetCPPTypeHelper(PropDef, ExtendedTypeText, CPPExportFlags, true, TEXT("bool"));
@@ -1046,11 +937,6 @@ struct FPropertyTypeTraitsBool : public FPropertyTypeTraitsBooleanBase
 
 struct FPropertyTypeTraitsBool8 : public FPropertyTypeTraitsBooleanBase
 {
-	static FProperty* CreateEngineType(FUnrealPropertyDefinitionInfo& PropDef, FFieldVariant Scope, const FName& Name, EObjectFlags ObjectFlags)
-	{
-		return CreateEngineTypeHelper<uint8, false>(PropDef, Scope, Name, ObjectFlags);
-	}
-
 	static FString GetCPPType(const FUnrealPropertyDefinitionInfo& PropDef, FString* ExtendedTypeText, uint32 CPPExportFlags)
 	{
 		return GetCPPTypeHelper(PropDef, ExtendedTypeText, CPPExportFlags, false, TEXT("uint8"));
@@ -1064,11 +950,6 @@ struct FPropertyTypeTraitsBool8 : public FPropertyTypeTraitsBooleanBase
 
 struct FPropertyTypeTraitsBool16 : public FPropertyTypeTraitsBooleanBase
 {
-	static FProperty* CreateEngineType(FUnrealPropertyDefinitionInfo& PropDef, FFieldVariant Scope, const FName& Name, EObjectFlags ObjectFlags)
-	{
-		return CreateEngineTypeHelper<uint16, false>(PropDef, Scope, Name, ObjectFlags);
-	}
-
 	static FString GetCPPType(const FUnrealPropertyDefinitionInfo& PropDef, FString* ExtendedTypeText, uint32 CPPExportFlags)
 	{
 		return GetCPPTypeHelper(PropDef, ExtendedTypeText, CPPExportFlags, false, TEXT("uint16"));
@@ -1082,11 +963,6 @@ struct FPropertyTypeTraitsBool16 : public FPropertyTypeTraitsBooleanBase
 
 struct FPropertyTypeTraitsBool32 : public FPropertyTypeTraitsBooleanBase
 {
-	static FProperty* CreateEngineType(FUnrealPropertyDefinitionInfo& PropDef, FFieldVariant Scope, const FName& Name, EObjectFlags ObjectFlags)
-	{
-		return CreateEngineTypeHelper<uint32, false>(PropDef, Scope, Name, ObjectFlags);
-	}
-
 	static FString GetCPPType(const FUnrealPropertyDefinitionInfo& PropDef, FString* ExtendedTypeText, uint32 CPPExportFlags)
 	{
 		return GetCPPTypeHelper(PropDef, ExtendedTypeText, CPPExportFlags, false, TEXT("uint32"));
@@ -1100,11 +976,6 @@ struct FPropertyTypeTraitsBool32 : public FPropertyTypeTraitsBooleanBase
 
 struct FPropertyTypeTraitsBool64 : public FPropertyTypeTraitsBooleanBase
 {
-	static FProperty* CreateEngineType(FUnrealPropertyDefinitionInfo& PropDef, FFieldVariant Scope, const FName& Name, EObjectFlags ObjectFlags)
-	{
-		return CreateEngineTypeHelper<uint64, false>(PropDef, Scope, Name, ObjectFlags);
-	}
-
 	static FString GetCPPType(const FUnrealPropertyDefinitionInfo& PropDef, FString* ExtendedTypeText, uint32 CPPExportFlags)
 	{
 		return GetCPPTypeHelper(PropDef, ExtendedTypeText, CPPExportFlags, false, TEXT("uint64"));
@@ -1186,23 +1057,6 @@ struct FPropertyTypeTraitsEnum : public FPropertyTypeTraitsBase
 			PropDef.GetVariableCategory(), ACCESS_Public, PropDef.GetArrayDimensions(), PropDef.GetUnrealSourceFile(), PropDef.GetLineNumber(), PropDef.GetParsePosition()));
 	}
 
-	static FProperty* CreateEngineType(FUnrealPropertyDefinitionInfo& PropDef, FFieldVariant Scope, const FName& Name, EObjectFlags ObjectFlags)
-	{
-		const FPropertyBase& VarProperty = PropDef.GetPropertyBase();
-
-		if (VarProperty.EnumDef->GetCppForm() != UEnum::ECppForm::EnumClass)
-		{
-			check(VarProperty.Type == EPropertyType::CPT_Byte);
-			return FPropertyTypeTraitsByte::CreateEngineType(PropDef, Scope, Name, ObjectFlags);
-		}
-
-		FEnumProperty* Result = new FEnumProperty(Scope, Name, ObjectFlags);
-		PropDef.SetProperty(Result);
-		Result->UnderlyingProp = CastFieldChecked<FNumericProperty>(FPropertyTraits::CreateEngineType(PropDef.GetValuePropDefRef()));
-		Result->Enum = VarProperty.EnumDef->GetEnum();
-		return Result;
-	}
-
 	static bool IsSupportedByBlueprint(const FUnrealPropertyDefinitionInfo& PropDef, bool bMemberVariable)
 	{
 		return true;
@@ -1345,25 +1199,6 @@ struct FPropertyTypeTraitsObjectReference : public FPropertyTypeTraitsObjectBase
 		return Detector.Changes();
 	}
 
-	static FProperty* CreateEngineType(FUnrealPropertyDefinitionInfo& PropDef, FFieldVariant Scope, const FName& Name, EObjectFlags ObjectFlags)
-	{
-		FPropertyBase& VarProperty = PropDef.GetPropertyBase();
-		check(VarProperty.ClassDef);
-		if (VarProperty.ClassDef->IsChildOf(*GUClassDef))
-		{
-			FClassProperty* Result = new FClassProperty(Scope, Name, ObjectFlags);
-			Result->MetaClass = VarProperty.MetaClassDef ? VarProperty.MetaClassDef->GetClass() : nullptr;
-			Result->PropertyClass = VarProperty.ClassDef->GetClass();
-			return Result;
-		}
-		else
-		{
-			FObjectProperty* Result = new FObjectProperty(Scope, Name, ObjectFlags);
-			Result->PropertyClass = VarProperty.ClassDef->GetClass();
-			return Result;
-		}
-	}
-
 	static bool IsSupportedByBlueprint(const FUnrealPropertyDefinitionInfo& PropDef, bool bMemberVariable)
 	{
 		return true;
@@ -1436,15 +1271,6 @@ struct FPropertyTypeTraitsObjectReference : public FPropertyTypeTraitsObjectBase
 
 struct FPropertyTypeTraitsWeakObjectReference : public FPropertyTypeTraitsObjectBase
 {
-	static FProperty* CreateEngineType(FUnrealPropertyDefinitionInfo& PropDef, FFieldVariant Scope, const FName& Name, EObjectFlags ObjectFlags)
-	{
-		const FPropertyBase& VarProperty = PropDef.GetPropertyBase();
-		check(VarProperty.ClassDef);
-		FWeakObjectProperty* Result = new FWeakObjectProperty(Scope, Name, ObjectFlags);
-		Result->PropertyClass = VarProperty.ClassDef->GetClass();
-		return Result;
-	}
-
 	static bool IsSupportedByBlueprint(const FUnrealPropertyDefinitionInfo& PropDef, bool bMemberVariable)
 	{
 		return bMemberVariable;
@@ -1490,15 +1316,6 @@ struct FPropertyTypeTraitsWeakObjectReference : public FPropertyTypeTraitsObject
 
 struct FPropertyTypeTraitsLazyObjectReference : public FPropertyTypeTraitsObjectBase
 {
-	static FProperty* CreateEngineType(FUnrealPropertyDefinitionInfo& PropDef, FFieldVariant Scope, const FName& Name, EObjectFlags ObjectFlags)
-	{
-		const FPropertyBase& VarProperty = PropDef.GetPropertyBase();
-		check(VarProperty.ClassDef);
-		FLazyObjectProperty* Result = new FLazyObjectProperty(Scope, Name, ObjectFlags);
-		Result->PropertyClass = VarProperty.ClassDef->GetClass();
-		return Result;
-	}
-
 	static FString GetEngineClassName(const FUnrealPropertyDefinitionInfo& PropDef)
 	{
 		return FLazyObjectProperty::StaticClass()->GetName();
@@ -1549,25 +1366,6 @@ struct FPropertyTypeTraitsObjectPtrReference : public FPropertyTypeTraitsObjectB
 		return Detector.Changes();
 	}
 
-	static FProperty* CreateEngineType(FUnrealPropertyDefinitionInfo& PropDef, FFieldVariant Scope, const FName& Name, EObjectFlags ObjectFlags)
-	{
-		FPropertyBase& VarProperty = PropDef.GetPropertyBase();
-		check(VarProperty.ClassDef);
-		if (VarProperty.ClassDef->IsChildOf(*GUClassDef))
-		{
-			FClassPtrProperty* Result = new FClassPtrProperty(Scope, Name, ObjectFlags);
-			Result->MetaClass = VarProperty.MetaClassDef ? VarProperty.MetaClassDef->GetClass() : nullptr;
-			Result->PropertyClass = VarProperty.ClassDef->GetClass();
-			return Result;
-		}
-		else
-		{
-			FObjectPtrProperty* Result = new FObjectPtrProperty(Scope, Name, ObjectFlags);
-			Result->PropertyClass = VarProperty.ClassDef->GetClass();
-			return Result;
-		}
-	}
-
 	static bool IsSupportedByBlueprint(const FUnrealPropertyDefinitionInfo& PropDef, bool bMemberVariable)
 	{
 		return true;
@@ -1637,25 +1435,6 @@ struct FPropertyTypeTraitsObjectPtrReference : public FPropertyTypeTraitsObjectB
 
 struct FPropertyTypeTraitsSoftObjectReference : public FPropertyTypeTraitsObjectBase
 {
-	static FProperty* CreateEngineType(FUnrealPropertyDefinitionInfo& PropDef, FFieldVariant Scope, const FName& Name, EObjectFlags ObjectFlags)
-	{
-		FPropertyBase& VarProperty = PropDef.GetPropertyBase();
-		check(VarProperty.ClassDef);
-		if (VarProperty.ClassDef->IsChildOf(*GUClassDef))
-		{
-			FSoftClassProperty* Result = new FSoftClassProperty(Scope, Name, ObjectFlags);
-			Result->MetaClass = VarProperty.MetaClassDef ? VarProperty.MetaClassDef->GetClass() : nullptr;
-			Result->PropertyClass = VarProperty.ClassDef->GetClass();
-			return Result;
-		}
-		else
-		{
-			FSoftObjectProperty* Result = new FSoftObjectProperty(Scope, Name, ObjectFlags);
-			Result->PropertyClass = VarProperty.ClassDef->GetClass();
-			return Result;
-		}
-	}
-
 	static bool IsSupportedByBlueprint(const FUnrealPropertyDefinitionInfo& PropDef, bool bMemberVariable)
 	{
 		return true;
@@ -1732,16 +1511,6 @@ struct FPropertyTypeTraitsInterface : public FPropertyTypeTraitsBase
 		return FPropertyTypeTraitsObjectBase::PostParseFinalize(PropDef, AddedFlags);
 	}
 
-	static FProperty* CreateEngineType(FUnrealPropertyDefinitionInfo& PropDef, FFieldVariant Scope, const FName& Name, EObjectFlags ObjectFlags)
-	{
-		FPropertyBase& VarProperty = PropDef.GetPropertyBase();
-		check(VarProperty.ClassDef);
-		check(VarProperty.ClassDef->HasAnyClassFlags(CLASS_Interface));
-		FInterfaceProperty* Result = new  FInterfaceProperty(Scope, Name, ObjectFlags);
-		Result->InterfaceClass = VarProperty.ClassDef->GetClass();
-		return Result;
-	}
-
 	static bool IsSupportedByBlueprint(const FUnrealPropertyDefinitionInfo& PropDef, bool bMemberVariable)
 	{
 		return true;
@@ -1822,12 +1591,6 @@ struct FPropertyTypeTraitsName : public FPropertyTypeTraitsBase
 	{
 	}
 
-	static FProperty* CreateEngineType(FUnrealPropertyDefinitionInfo& PropDef, FFieldVariant Scope, const FName& Name, EObjectFlags ObjectFlags)
-	{
-		FNameProperty* Result = new FNameProperty(Scope, Name, ObjectFlags);
-		return Result;
-	}
-
 	static bool IsSupportedByBlueprint(const FUnrealPropertyDefinitionInfo& PropDef, bool bMemberVariable)
 	{
 		return true;
@@ -1867,12 +1630,6 @@ struct FPropertyTypeTraitsString : public FPropertyTypeTraitsBase
 	{
 	}
 
-	static FProperty* CreateEngineType(FUnrealPropertyDefinitionInfo& PropDef, FFieldVariant Scope, const FName& Name, EObjectFlags ObjectFlags)
-	{
-		FStrProperty* Result = new FStrProperty(Scope, Name, ObjectFlags);
-		return Result;
-	}
-
 	static bool IsSupportedByBlueprint(const FUnrealPropertyDefinitionInfo& PropDef, bool bMemberVariable)
 	{
 		return true;
@@ -1949,12 +1706,6 @@ struct FPropertyTypeTraitsText : public FPropertyTypeTraitsBase
 	{
 	}
 
-	static FProperty* CreateEngineType(FUnrealPropertyDefinitionInfo& PropDef, FFieldVariant Scope, const FName& Name, EObjectFlags ObjectFlags)
-	{
-		FTextProperty* Result = new FTextProperty(Scope, Name, ObjectFlags);
-		return Result;
-	}
-
 	static bool IsSupportedByBlueprint(const FUnrealPropertyDefinitionInfo& PropDef, bool bMemberVariable)
 	{
 		return true;
@@ -2227,14 +1978,6 @@ struct FPropertyTypeTraitsStruct : public FPropertyTypeTraitsBase
 		return Detector.Changes();
 	}
 
-	static FProperty* CreateEngineType(FUnrealPropertyDefinitionInfo& PropDef, FFieldVariant Scope, const FName& Name, EObjectFlags ObjectFlags)
-	{
-		FPropertyBase& VarProperty = PropDef.GetPropertyBase();
-		FStructProperty* Result = new FStructProperty(Scope, Name, ObjectFlags);
-		Result->Struct = PropDef.GetPropertyBase().ScriptStructDef->GetScriptStruct();
-		return Result;
-	}
-
 	static bool IsSupportedByBlueprint(const FUnrealPropertyDefinitionInfo& PropDef, bool bMemberVariable)
 	{
 		const FPropertyBase& VarProperty = PropDef.GetPropertyBase();
@@ -2254,7 +1997,7 @@ struct FPropertyTypeTraitsStruct : public FPropertyTypeTraitsBase
 	static FString GetCPPTypeForwardDeclaration(const FUnrealPropertyDefinitionInfo& PropDef)
 	{
 		// Core type structs don't need to forward declare in UHT as every generated.h indirectly includes CoreMinimal.h
-		if (UScriptStruct::ICppStructOps* CppStructOps = UScriptStruct::FindDeferredCppStructOps(PropDef.GetPropertyBase().ScriptStructDef->GetStructPathName()); CppStructOps != nullptr && CppStructOps->IsUECoreType())
+		if (PropDef.GetPropertyBase().ScriptStructDef->HasAnyScriptStructExportFlags(STRUCTEXPORT_IsCoreType))
 		{
 			return FString();
 		}
@@ -2284,16 +2027,6 @@ struct FPropertyTypeTraitsDelegate : public FPropertyTypeTraitsBase
 		return Detector.Changes();
 	}
 
-	static FProperty* CreateEngineType(FUnrealPropertyDefinitionInfo& PropDef, FFieldVariant Scope, const FName& Name, EObjectFlags ObjectFlags)
-	{
-		const FPropertyBase& VarProperty = PropDef.GetPropertyBase();
-		check(VarProperty.FunctionDef);
-
-		FDelegateProperty* Result = new FDelegateProperty(Scope, Name, ObjectFlags);
-		Result->SignatureFunction = VarProperty.FunctionDef->GetFunction();
-		return Result;
-	}
-
 	static bool IsSupportedByBlueprint(const FUnrealPropertyDefinitionInfo& PropDef, bool bMemberVariable)
 	{
 		return true;
@@ -2371,24 +2104,6 @@ struct FPropertyTypeTraitsMulticastDelegate : public FPropertyTypeTraitsBase
 		return Detector.Changes();
 	}
 
-	static FProperty* CreateEngineType(FUnrealPropertyDefinitionInfo& PropDef, FFieldVariant Scope, const FName& Name, EObjectFlags ObjectFlags)
-	{
-		const FPropertyBase& VarProperty = PropDef.GetPropertyBase();
-
-		if (VarProperty.FunctionDef->GetFunctionType() == EFunctionType::SparseDelegate)
-		{
-			FMulticastSparseDelegateProperty* Result = new FMulticastSparseDelegateProperty(Scope, Name, ObjectFlags);
-			Result->SignatureFunction = VarProperty.FunctionDef->GetFunction();
-			return Result;
-		}
-		else
-		{
-			FMulticastInlineDelegateProperty* Result = new FMulticastInlineDelegateProperty(Scope, Name, ObjectFlags);
-			Result->SignatureFunction = VarProperty.FunctionDef->GetFunction();
-			return Result;
-		}
-	}
-
 	static bool IsSupportedByBlueprint(const FUnrealPropertyDefinitionInfo& PropDef, bool bMemberVariable)
 	{
 		return bMemberVariable;
@@ -2472,16 +2187,6 @@ struct FPropertyTypeTraitsFieldPath : public FPropertyTypeTraitsBase
 	{
 	}
 
-	static FProperty* CreateEngineType(FUnrealPropertyDefinitionInfo& PropDef, FFieldVariant Scope, const FName& Name, EObjectFlags ObjectFlags)
-	{
-		const FPropertyBase& VarProperty = PropDef.GetPropertyBase();
-		FFieldPathProperty* Result = new FFieldPathProperty(Scope, Name, ObjectFlags);
-		FFieldClass** FieldClass = FFieldClass::GetNameToFieldClassMap().Find(VarProperty.FieldClassName);
-		check(FieldClass);
-		Result->PropertyClass = *FieldClass;
-		return Result;
-	}
-
 	static bool IsSupportedByBlueprint(const FUnrealPropertyDefinitionInfo& PropDef, bool bMemberVariable)
 	{
 		return true;
@@ -2603,13 +2308,6 @@ struct FPropertyTypeTraitsStaticArray : public FPropertyTypeTraitsBase
 		return PostParseFinalizeHelper<false>(PropDef, AddedFlags);
 	}
 
-	static FProperty* CreateEngineType(FUnrealPropertyDefinitionInfo& PropDef, FFieldVariant Scope, const FName& Name, EObjectFlags ObjectFlags)
-	{
-		FProperty* Property = CreateEngineTypeHelper<false>(PropDef, Scope, Name, ObjectFlags);
-		Property->ArrayDim = 2;
-		return Property;
-	}
-
 	static bool IsSupportedByBlueprint(const FUnrealPropertyDefinitionInfo& PropDef, bool bMemberVariable)
 	{
 		return IsSupportedByBlueprintSansContainers(PropDef, bMemberVariable);
@@ -2670,16 +2368,6 @@ struct FPropertyTypeTraitsDynamicArray : public FPropertyTypeTraitsBase
 		return Detector.Changes();
 	}
 
-	static FProperty* CreateEngineType(FUnrealPropertyDefinitionInfo& PropDef, FFieldVariant Scope, const FName& Name, EObjectFlags ObjectFlags)
-	{
-		FPropertyBase& VarProperty = PropDef.GetPropertyBase();
-
-		FArrayProperty* Array = new FArrayProperty(Scope, Name, ObjectFlags);
-		PropDef.SetProperty(Array);
-		Array->Inner = FPropertyTraits::CreateEngineType(PropDef.GetValuePropDefRef());
-		return Array;
-	}
-
 	static bool IsSupportedByBlueprint(const FUnrealPropertyDefinitionInfo& PropDef, bool bMemberVariable)
 	{
 		return IsSupportedByBlueprintSansContainers(PropDef.GetValuePropDef(), false);
@@ -2756,16 +2444,6 @@ struct FPropertyTypeTraitsSet : public FPropertyTypeTraitsBase
 		return Detector.Changes();
 	}
 
-	static FProperty* CreateEngineType(FUnrealPropertyDefinitionInfo& PropDef, FFieldVariant Scope, const FName& Name, EObjectFlags ObjectFlags)
-	{
-		FPropertyBase& VarProperty = PropDef.GetPropertyBase();
-
-		FSetProperty* Set = new FSetProperty(Scope, Name, ObjectFlags);
-		PropDef.SetProperty(Set);
-		Set->ElementProp = FPropertyTraits::CreateEngineType(PropDef.GetValuePropDefRef());
-		return Set;
-	}
-
 	static bool IsSupportedByBlueprint(const FUnrealPropertyDefinitionInfo& PropDef, bool bMemberVariable)
 	{
 		return IsSupportedByBlueprintSansContainers(PropDef.GetValuePropDef(), false);
@@ -2858,17 +2536,6 @@ struct FPropertyTypeTraitsMap : public FPropertyTypeTraitsBase
 		return Detector.Changes();
 	}
 
-	static FProperty* CreateEngineType(FUnrealPropertyDefinitionInfo& PropDef, FFieldVariant Scope, const FName& Name, EObjectFlags ObjectFlags)
-	{
-		FPropertyBase& VarProperty = PropDef.GetPropertyBase();
-
-		FMapProperty* Map = new FMapProperty(Scope, Name, ObjectFlags);
-		PropDef.SetProperty(Map);
-		Map->KeyProp = FPropertyTraits::CreateEngineType(PropDef.GetKeyPropDefRef());
-		Map->ValueProp = FPropertyTraits::CreateEngineType(PropDef.GetValuePropDefRef());
-		return Map;
-	}
-
 	static bool IsSupportedByBlueprint(const FUnrealPropertyDefinitionInfo& PropDef, bool bMemberVariable)
 	{
 		return
@@ -2958,58 +2625,6 @@ TSharedRef<FUnrealPropertyDefinitionInfo> FPropertyTraits::CreateProperty(const
 	return PropDefRef;
 }
 
-FProperty* FPropertyTraits::CreateEngineType(TSharedRef<FUnrealPropertyDefinitionInfo> PropDefRef)
-{
-	FUnrealPropertyDefinitionInfo& PropDef = *PropDefRef;
-	FUnrealTypeDefinitionInfo* Outer = PropDef.GetOuter();
-	check(Outer);
-
-	EObjectFlags ObjectFlags = RF_Public;
-	if (PropDef.GetVariableCategory() == EVariableCategory::Member && PropDef.GetAccessSpecifier() == ACCESS_Private)
-	{
-		ObjectFlags = RF_NoFlags;
-	}
-
-	// Create the property and attach to the definition
-	FFieldVariant Scope = Outer->AsProperty() ? FFieldVariant(Outer->AsProperty()->GetProperty()) : FFieldVariant(Outer->AsObject()->GetObject());
-	FProperty* Property = CreateEngineTypeHelper<true>(PropDef, Scope, PropDef.GetFName(), ObjectFlags);
-	PropDef.SetProperty(Property);
-
-	// Perform some final initialization from the property base data
-	Property->PropertyFlags = PropDef.GetPropertyBase().PropertyFlags;
-
-	// Special initialization for member variables
-	if (PropDef.GetVariableCategory() == EVariableCategory::Member)
-	{
-		if (PropDef.HasAnyPropertyFlags(CPF_RepNotify))
-		{
-			Property->RepNotifyFunc = PropDef.GetPropertyBase().RepNotifyName;
-		}
-	}
-
-	// Add the meta data to the property
-	for (TPair<FName, FString>& KVP : PropDef.GetMetaDataMap())
-	{
-		Property->SetMetaData(KVP.Key, FString(KVP.Value));
-	}
-
-	// If we are parent to a struct
-	if (FUnrealStructDefinitionInfo* StructDef = UHTCast<FUnrealStructDefinitionInfo>(Outer))
-	{
-
-		// Add to the end of the properties slist
-		FField** Prev = &StructDef->GetStruct()->ChildProperties;
-		for (; *Prev != nullptr; Prev = &(*Prev)->Next)
-		{
-			// No body
-		}
-		check(*Prev == nullptr);
-		Property->Next = nullptr;
-		*Prev = Property;
-	}
-	return Property;
-}
-
 bool FPropertyTraits::IsSupportedByBlueprint(const FUnrealPropertyDefinitionInfo& PropDef, bool bMemberVariable)
 {
 	return PropertyTypeDispatch<IsSupportedByBlueprintDispatch, true, bool>(PropDef, bMemberVariable);
diff --git a/Engine/Source/Programs/UnrealHeaderTool/Private/PropertyTypes.h b/Engine/Source/Programs/UnrealHeaderTool/Private/PropertyTypes.h
--- a/Engine/Source/Programs/UnrealHeaderTool/Private/PropertyTypes.h
+++ b/Engine/Source/Programs/UnrealHeaderTool/Private/PropertyTypes.h
@@ -41,14 +41,6 @@ struct FPropertyTraits
 	 */
 	static TSharedRef<FUnrealPropertyDefinitionInfo> CreateProperty(const FPropertyBase& VarProperty, FUnrealTypeDefinitionInfo& Outer, const FName& Name, EVariableCategory VariableCategory, EAccessSpecifier AccessSpecifier, const TCHAR* Dimensions, FUnrealSourceFile& SourceFile, int LineNumber, int ParsePosition);
 
-	/**
-	 * Given a property, create the underlying engine types
-	 * 
-	 * @param PropDef The property in question
-	 * @param ObjectFlags The flags associated with the property
-	 */
-	static FProperty* CreateEngineType(TSharedRef<FUnrealPropertyDefinitionInfo> PropDefRef);
-
 	/**
 	 * Test to see if the property can be used in a blueprint
 	 * 
diff --git a/Engine/Source/Programs/UnrealHeaderTool/Private/UnrealSourceFile.cpp b/Engine/Source/Programs/UnrealHeaderTool/Private/UnrealSourceFile.cpp
--- a/Engine/Source/Programs/UnrealHeaderTool/Private/UnrealSourceFile.cpp
+++ b/Engine/Source/Programs/UnrealHeaderTool/Private/UnrealSourceFile.cpp
@@ -141,7 +141,7 @@ bool FUnrealSourceFile::HasChanged() const
 
 FString FUnrealSourceFile::GetFileDefineName() const
 {
-	const FString API = FPackageName::GetShortName(GetPackage()).ToUpper();
+	const FString API = GetPackageDef().GetShortUpperName();
 	return FString::Printf(TEXT("%s_%s_generated_h"), *API, *GetStrippedFilename());
 }
 
diff --git a/Engine/Source/Programs/UnrealHeaderTool/Private/UnrealSourceFile.h b/Engine/Source/Programs/UnrealHeaderTool/Private/UnrealSourceFile.h
--- a/Engine/Source/Programs/UnrealHeaderTool/Private/UnrealSourceFile.h
+++ b/Engine/Source/Programs/UnrealHeaderTool/Private/UnrealSourceFile.h
@@ -259,14 +259,6 @@ public:
 		return PackageDef;
 	}
 
-	/**
-	 * Gets package this file is in.
-	 */
-	UPackage* GetPackage() const
-	{
-		return PackageDef.GetPackage();
-	}
-
 	/**
 	 * Gets filename.
 	 */
diff --git a/Engine/Source/Programs/UnrealHeaderTool/Private/UnrealTypeDefinitionInfo.cpp b/Engine/Source/Programs/UnrealHeaderTool/Private/UnrealTypeDefinitionInfo.cpp
--- a/Engine/Source/Programs/UnrealHeaderTool/Private/UnrealTypeDefinitionInfo.cpp
+++ b/Engine/Source/Programs/UnrealHeaderTool/Private/UnrealTypeDefinitionInfo.cpp
@@ -201,9 +201,6 @@ const FString* FUHTMetaData::FindMetaDataHelper(const FName& Key) const
 	{
 		Result = GetMetaDataMap().Find(Key);
 	}
-#if UHT_ENABLE_ENGINE_TYPE_CHECKS
-	CheckFindMetaData(Key, Result);
-#endif
 	return Result;
 }
 
@@ -340,24 +337,6 @@ void FUnrealTypeDefinitionInfo::GetHashTag(FUnrealTypeDefinitionInfo& Referencin
 	}
 }
 
-void FUnrealTypeDefinitionInfo::CreateUObjectEngineTypes(ECreateEngineTypesPhase Phase)
-{
-	EFinalizeState& FinalizeState = CreateUObectEngineTypesState[uint8(Phase)];
-	switch (FinalizeState)
-	{
-	case EFinalizeState::None:
-		FinalizeState = EFinalizeState::InProgress;
-		CreateUObjectEngineTypesInternal(Phase);
-		FinalizeState = EFinalizeState::Finished;
-		break;
-	case EFinalizeState::InProgress:
-		checkf(false, TEXT("Recursive call to CreateUObectEngineTypes detected"))
-		break;
-	case EFinalizeState::Finished:
-		break;
-	}
-}
-
 void FUnrealTypeDefinitionInfo::PostParseFinalize(EPostParseFinalizePhase Phase)
 {
 	EFinalizeState& FinalizeState = PostParseFinalizeState[uint8(Phase)];
@@ -621,27 +600,6 @@ void FUnrealPropertyDefinitionInfo::PostParseFinalizeInternal(EPostParseFinalize
 	}
 }
 
-void FUnrealPropertyDefinitionInfo::CreateUObjectEngineTypesInternal(ECreateEngineTypesPhase Phase)
-{
-	FUnrealTypeDefinitionInfo::CreateUObjectEngineTypesInternal(Phase);
-
-	switch (Phase)
-	{
-	case ECreateEngineTypesPhase::Phase1:
-		break;
-
-	case ECreateEngineTypesPhase::Phase2:
-		break;
-
-	case ECreateEngineTypesPhase::Phase3:
-		if (GetPropertySafe() == nullptr)
-		{
-			FPropertyTraits::CreateEngineType(SharedThis(this));
-		}
-		break;
-	}
-}
-
 void FUnrealPropertyDefinitionInfo::ConcurrentPostParseFinalize()
 {
 	FUnrealTypeDefinitionInfo::ConcurrentPostParseFinalize();
@@ -650,11 +608,6 @@ void FUnrealPropertyDefinitionInfo::ConcurrentPostParseFinalize()
 
 void FUnrealPropertyDefinitionInfo::SetDelegateFunctionSignature(FUnrealFunctionDefinitionInfo& DelegateFunctionDef)
 {
-	if (GetPropertySafe() != nullptr)
-	{
-		FDelegateProperty* DelegateProperty = CastFieldChecked<FDelegateProperty>(PropertyBase.ArrayType == EArrayType::None ? GetProperty() : GetValuePropDef().GetProperty());
-		DelegateProperty->SignatureFunction = DelegateFunctionDef.GetFunction();
-	}
 	PropertyBase.FunctionDef = &DelegateFunctionDef;
 
 #if UHT_ENABLE_DELEGATE_PROPERTY_TAG
@@ -671,12 +624,6 @@ void FUnrealPropertyDefinitionInfo::SetDelegateFunctionSignature(FUnrealFunction
 
 FString FUnrealPropertyDefinitionInfo::GetEngineClassName() const
 {
-#if UHT_ENABLE_ENGINE_TYPE_CHECKS
-	if (GetPropertySafe() != nullptr)
-	{
-		check(FPropertyTraits::GetEngineClassName(*this) == GetPropertySafe()->GetClass()->GetName()); // Validation
-	}
-#endif
 	return FPropertyTraits::GetEngineClassName(*this);
 }
 
@@ -684,14 +631,6 @@ FString FUnrealPropertyDefinitionInfo::GetPathName(FUnrealObjectDefinitionInfo*
 {
 	TStringBuilder<256> ResultString;
 	GetPathName(StopOuter, ResultString);
-#if UHT_ENABLE_ENGINE_TYPE_CHECKS
-	if (GetPropertySafe() != nullptr && (StopOuter == nullptr || StopOuter->GetObjectSafe() != nullptr))
-	{
-		TStringBuilder<256> OtherResultString;
-		GetPropertySafe()->GetPathName(StopOuter ? StopOuter->GetObjectSafe() : nullptr, OtherResultString); // Validation
-		check(FCString::Strcmp(OtherResultString.ToString(), ResultString.ToString()) == 0);
-	}
-#endif
 	return FString(FStringView(ResultString));
 }
 
@@ -725,38 +664,18 @@ FString FUnrealPropertyDefinitionInfo::GetFullName() const
 	FString FullName = GetEngineClassName();
 	FullName += TEXT(" ");
 	FullName += GetPathName();
-#if UHT_ENABLE_ENGINE_TYPE_CHECKS
-	if (GetPropertySafe() != nullptr)
-	{
-		check(FullName == GetPropertySafe()->GetFullName()); // Validation
-	}
-#endif
 	return FullName;
 }
 
 FString FUnrealPropertyDefinitionInfo::GetCPPType(FString* ExtendedTypeText/* = nullptr*/, uint32 CPPExportFlags/* = 0*/) const
 {
 	FString Out = FPropertyTraits::GetCPPType(*this, ExtendedTypeText, CPPExportFlags);
-#if UHT_ENABLE_ENGINE_TYPE_CHECKS
-	if (GetPropertySafe() != nullptr)
-	{
-		FString ExtOutTemp;
-		FString* ExtOutTempPtr = ExtendedTypeText ? &ExtOutTemp : nullptr;
-		check(Out == GetPropertySafe()->GetCPPType(ExtOutTempPtr, CPPExportFlags) && (ExtendedTypeText == nullptr || *ExtendedTypeText == ExtOutTemp)); // Validation
-	}
-#endif
 	return Out;
 }
 
 FString FUnrealPropertyDefinitionInfo::GetCPPTypeForwardDeclaration() const
 {
 	FString Out = FPropertyTraits::GetCPPTypeForwardDeclaration(*this);
-#if UHT_ENABLE_ENGINE_TYPE_CHECKS
-	if (GetPropertySafe() != nullptr)
-	{
-		check(Out == GetPropertySafe()->GetCPPTypeForwardDeclaration()); // Validation
-	}
-#endif
 	return Out;
 }
 
@@ -764,11 +683,6 @@ struct FUHTDisplayNameHelper
 {
 	static FString Get(const FUnrealPropertyDefinitionInfo& Property)
 	{
-		// The GetAuthoredNameForField only does something for user defined structures
-		//if (FUnrealStructDefinitionInfo* OwnerStruct = GetOwnerStruct(Property))
-		//{
-		//	OwnerStruct->GetStruct()->GetAuthoredNameForField(Property.GetProperty());
-		//}
 		return Property.GetName();
 	}
 
@@ -867,12 +781,6 @@ FText FUnrealPropertyDefinitionInfo::GetDisplayNameText() const
 		LocalizedDisplayName = FText::FromString(NativeDisplayName);
 	}
 
-#if UHT_ENABLE_ENGINE_TYPE_CHECKS
-	if (GetPropertySafe())
-	{
-		check(LocalizedDisplayName.ToString() == GetPropertySafe()->GetDisplayNameText().ToString());
-	}
-#endif
 	return LocalizedDisplayName;
 }
 
@@ -921,12 +829,6 @@ FText FUnrealPropertyDefinitionInfo::GetToolTipText(bool bShortTooltip) const
 		LocalizedToolTip = FText::FromString(NativeToolTip);
 	}
 
-#if UHT_ENABLE_ENGINE_TYPE_CHECKS
-	if (GetPropertySafe())
-	{
-		check(LocalizedToolTip.ToString() == GetPropertySafe()->GetToolTipText(bShortTooltip).ToString());
-	}
-#endif
 	return LocalizedToolTip;
 }
 
@@ -1073,12 +975,6 @@ FUnrealPackageDefinitionInfo& FUnrealPropertyDefinitionInfo::GetPackageDef() con
 bool FUnrealPropertyDefinitionInfo::SameType(const FUnrealPropertyDefinitionInfo& Other) const
 {
 	bool bResults = FPropertyTraits::SameType(*this, Other);
-#if UHT_ENABLE_ENGINE_TYPE_CHECKS
-	if (GetPropertySafe() && Other.GetPropertySafe())
-	{
-		check(GetPropertySafe()->SameType(Other.GetPropertySafe()) == bResults)
-	}
-#endif
 	return bResults;
 }
 
@@ -1107,7 +1003,6 @@ FUnrealObjectDefinitionInfo::FUnrealObjectDefinitionInfo(UObject* InObject)
 	, Name(InObject->GetFName())
 	, InternalObjectFlags(InObject->GetInternalFlags())
 {
-	SetObject(InObject);
 }
 
 FString FUnrealObjectDefinitionInfo::GetFullName() const
@@ -1115,12 +1010,6 @@ FString FUnrealObjectDefinitionInfo::GetFullName() const
 	FString FullName = GetEngineClassName();
 	FullName += TEXT(" ");
 	FullName += GetPathName();
-#if UHT_ENABLE_ENGINE_TYPE_CHECKS
-	if (GetObjectSafe() != nullptr)
-	{
-		check(FullName == GetObjectSafe()->GetFullName()); // Validation
-	}
-#endif
 	return FullName;
 }
 
@@ -1128,14 +1017,6 @@ FString FUnrealObjectDefinitionInfo::GetPathName(FUnrealObjectDefinitionInfo* St
 {
 	TStringBuilder<256> ResultBuilder;
 	GetPathName(StopOuter, ResultBuilder);
-#if UHT_ENABLE_ENGINE_TYPE_CHECKS
-	if (GetObjectSafe() != nullptr && (StopOuter == nullptr || StopOuter->GetObjectSafe() != nullptr))
-	{
-		TStringBuilder<256> OtherResultString;
-		GetObjectSafe()->GetPathName(StopOuter ? StopOuter->GetObjectSafe() : nullptr, OtherResultString); // Validation
-		check(FCString::Strcmp(OtherResultString.ToString(), ResultBuilder.ToString()) == 0);
-	}
-#endif
 	return FString(FStringView(ResultBuilder));
 }
 
@@ -1183,53 +1064,25 @@ void FUnrealObjectDefinitionInfo::AddMetaData(TMap<FName, FString>&& InMetaData)
 	// only add if we have some!
 	if (InMetaData.Num())
 	{
-
-		if (TMap<FName, FString>* UObjectMetaDataMap = GetUObjectMetaDataMap())
-		{
-			UObjectMetaDataMap->Append(InMetaData);
-		}
-		else if (UMetaData* UObjectMetaData = GetUObjectMetaData())
-		{
-			UObjectMetaData->SetObjectValues(GetObject(), InMetaData); // This is valid
-		}
-
 		GetMetaDataMap().Append(MoveTemp(InMetaData));
 	}
 }
 
 void FUnrealObjectDefinitionInfo::SetMetaDataHelper(const FName& Key, const TCHAR* InValue)
 {
-	if (TMap<FName, FString>* UObjectMetaDataMap = GetUObjectMetaDataMap())
-	{
-		UObjectMetaDataMap->Add(Key, InValue);
-	}
-	else if (UMetaData* UObjectMetaData = GetUObjectMetaData())
-	{
-		UObjectMetaData->SetValue(GetObject(), Key, InValue); // This is valid
-	}
-
 	GetMetaDataMap().Add(Key, InValue);
 }
 
 FUnrealPackageDefinitionInfo::FUnrealPackageDefinitionInfo(const FManifestModule& InModule, UPackage* InPackage)
 	: FUnrealObjectDefinitionInfo(InPackage)
 	, Module(InModule)
-	, ShortUpperName(FPackageName::GetShortName(InPackage).ToUpper())
+	, ShortName(FPackageName::GetShortName(InPackage))
+	, ShortUpperName(ShortName.ToUpper())
 	, API(FString::Printf(TEXT("%s_API "), *ShortUpperName))
+	, PackageFlags(InPackage->GetPackageFlags())
 {
 }
 
-void FUnrealPackageDefinitionInfo::CreateUObjectEngineTypesInternal(ECreateEngineTypesPhase Phase)
-{
-	for (TSharedRef<FUnrealSourceFile>& LocalSourceFile : GetAllSourceFiles())
-	{
-		for (TSharedRef<FUnrealTypeDefinitionInfo>& TypeDef : LocalSourceFile->GetDefinedTypes())
-		{
-			TypeDef->CreateUObjectEngineTypes(Phase);
-		}
-	}
-}
-
 void FUnrealPackageDefinitionInfo::PostParseFinalizeInternal(EPostParseFinalizePhase Phase)
 {
 	FUnrealObjectDefinitionInfo::PostParseFinalizeInternal(Phase);
@@ -1241,9 +1094,7 @@ void FUnrealPackageDefinitionInfo::PostParseFinalizeInternal(EPostParseFinalizeP
 
 	case EPostParseFinalizePhase::Phase2:
 	{
-		UPackage* Package = GetPackage();
-
-		FString PackageName = Package->GetName();
+		FString PackageName = GetName();
 		PackageName.ReplaceInline(TEXT("/"), TEXT("_"), ESearchCase::CaseSensitive);
 
 		SingletonName.Appendf(TEXT("Z_Construct_UPackage_%s()"), *PackageName);
@@ -1346,12 +1197,6 @@ FText FUnrealFieldDefinitionInfo::GetToolTipText(bool bShortTooltip) const
 		LocalizedToolTip = FText::FromString(NativeToolTip);
 	}
 
-#if UHT_ENABLE_ENGINE_TYPE_CHECKS
-	if (GetFieldSafe())
-	{
-		check(LocalizedToolTip.ToString() == GetFieldSafe()->GetToolTipText(bShortTooltip).ToString());
-	}
-#endif
 	return LocalizedToolTip;
 }
 
@@ -1469,9 +1314,6 @@ FString FUnrealEnumDefinitionInfo::GenerateEnumPrefix() const
 		Prefix = GetName();
 	}
 
-#if UHT_ENABLE_ENGINE_TYPE_CHECKS
-	check(GetEnumSafe() == nullptr || GetEnumSafe()->GenerateEnumPrefix() == Prefix);
-#endif
 	return Prefix;
 }
 
@@ -1489,24 +1331,15 @@ bool FUnrealEnumDefinitionInfo::ContainsExistingMax() const
 {
 	if (GetIndexByName(*GenerateFullEnumName(TEXT("MAX")), EGetByNameFlags::CaseSensitive) != INDEX_NONE)
 	{
-#if UHT_ENABLE_ENGINE_TYPE_CHECKS
-		check(GetEnumSafe() == nullptr || GetEnumSafe()->ContainsExistingMax() == true);
-#endif
 		return true;
 	}
 
 	FName MaxEnumItem = *GenerateFullEnumName(*(GenerateEnumPrefix() + TEXT("_MAX")));
 	if (GetIndexByName(MaxEnumItem, EGetByNameFlags::CaseSensitive) != INDEX_NONE)
 	{
-#if UHT_ENABLE_ENGINE_TYPE_CHECKS
-		check(GetEnumSafe() == nullptr || GetEnumSafe()->ContainsExistingMax() == true);
-#endif
 		return true;
 	}
 
-#if UHT_ENABLE_ENGINE_TYPE_CHECKS
-	check(GetEnumSafe() == nullptr || GetEnumSafe()->ContainsExistingMax() == false);
-#endif
 	return false;
 }
 
@@ -1528,9 +1361,6 @@ int64 FUnrealEnumDefinitionInfo::GetMaxEnumValue() const
 		}
 	}
 
-#if UHT_ENABLE_ENGINE_TYPE_CHECKS
-	check(GetEnumSafe() == nullptr || GetEnumSafe()->GetMaxEnumValue() == MaxValue);
-#endif
 	return MaxValue;
 }
 
@@ -1538,15 +1368,9 @@ FName FUnrealEnumDefinitionInfo::GetNameByIndex(int32 Index) const
 {
 	if (Names.IsValidIndex(Index))
 	{
-#if UHT_ENABLE_ENGINE_TYPE_CHECKS
-		check(GetEnumSafe() == nullptr || GetEnumSafe()->GetNameByIndex(Index) == Names[Index].Key);
-#endif
 		return Names[Index].Key;
 	}
 
-#if UHT_ENABLE_ENGINE_TYPE_CHECKS
-	check(GetEnumSafe() == nullptr || GetEnumSafe()->GetNameByIndex(Index) == NAME_None);
-#endif
 	return NAME_None;
 }
 
@@ -1558,16 +1382,10 @@ bool FUnrealEnumDefinitionInfo::IsValidEnumValue(int64 InValue) const
 		int64 CurrentValue = Names[i].Value;
 		if (CurrentValue == InValue)
 		{
-#if UHT_ENABLE_ENGINE_TYPE_CHECKS
-			check(GetEnumSafe() == nullptr || GetEnumSafe()->IsValidEnumValue(InValue) == true);
-#endif
 			return true;
 		}
 	}
 
-#if UHT_ENABLE_ENGINE_TYPE_CHECKS
-	check(GetEnumSafe() == nullptr || GetEnumSafe()->IsValidEnumValue(InValue) == false);
-#endif
 	return false;
 }
 
@@ -1581,9 +1399,6 @@ int32 FUnrealEnumDefinitionInfo::GetIndexByName(const FName InName, EGetByNameFl
 	{
 		if (Names[Counter].Key.IsEqual(InName, ComparisonMethod))
 		{
-#if UHT_ENABLE_ENGINE_TYPE_CHECKS
-			check(GetEnumSafe() == nullptr || GetEnumSafe()->GetIndexByName(InName, Flags) == Counter);
-#endif
 			return Counter;
 		}
 	}
@@ -1660,34 +1475,6 @@ FString FUnrealEnumDefinitionInfo::GetNameStringByIndex(int32 InIndex) const
 	return FString();
 }
 
-void FUnrealEnumDefinitionInfo::CreateUObjectEngineTypesInternal(ECreateEngineTypesPhase Phase)
-{
-	FUnrealFieldDefinitionInfo::CreateUObjectEngineTypesInternal(Phase);
-
-	switch (Phase)
-	{
-	case ECreateEngineTypesPhase::Phase1:
-	{
-		UPackage* Package = GetPackageDef().GetPackage();
-		const FString& EnumName = GetNameCPP();
-
-		// Create enum definition.
-		UEnum* Enum = new(EC_InternalUseOnlyConstructor, Package, FName(EnumName), RF_Public) UEnum(FObjectInitializer());
-		Enum->SetEnums(Names, CppForm, EnumFlags, false);
-		Enum->CppType = CppType;
-		Enum->GetPackage()->GetMetaData()->SetObjectValues(Enum, GetMetaDataMap());
-		SetObject(Enum);
-		break;
-	}
-
-	case ECreateEngineTypesPhase::Phase2:
-		break;
-
-	case ECreateEngineTypesPhase::Phase3:
-		break;
-	}
-}
-
 FUnrealStructDefinitionInfo::FUnrealStructDefinitionInfo(FUnrealSourceFile& InSourceFile, int32 InLineNumber, FString&& InNameCPP, FName InName, FUnrealObjectDefinitionInfo& InOuter)
 	: FUnrealFieldDefinitionInfo(InSourceFile, InLineNumber, MoveTemp(InNameCPP), InName, InOuter)
 	, StructScope(MakeShared<FStructScope>(*this, &InSourceFile.GetScope().Get()))
@@ -1700,15 +1487,9 @@ bool FUnrealStructDefinitionInfo::IsChildOf(const FUnrealStructDefinitionInfo& S
 	{
 		if (Current == &SomeBase)
 		{
-#if UHT_ENABLE_ENGINE_TYPE_CHECKS
-			check(GetStructSafe() == nullptr || SomeBase.GetStructSafe() == nullptr || GetStructSafe()->IsChildOf(SomeBase.GetStructSafe()) == true);
-#endif
 			return true;
 		}
 	}
-#if UHT_ENABLE_ENGINE_TYPE_CHECKS
-	check(GetStructSafe() == nullptr || SomeBase.GetStructSafe() == nullptr || GetStructSafe()->IsChildOf(SomeBase.GetStructSafe()) == false);
-#endif
 	return false;
 }
 
@@ -1735,63 +1516,6 @@ void FUnrealStructDefinitionInfo::AddProperty(TSharedRef<FUnrealPropertyDefiniti
 	}
 }
 
-void FUnrealStructDefinitionInfo::CreateUObjectEngineTypesInternal(ECreateEngineTypesPhase Phase)
-{
-	FUnrealFieldDefinitionInfo::CreateUObjectEngineTypesInternal(Phase);
-
-	switch (Phase)
-	{
-	case ECreateEngineTypesPhase::Phase1:
-		if (SuperStructInfo.Struct)
-		{
-			SuperStructInfo.Struct->CreateUObjectEngineTypes(Phase);
-		}
-
-		for (const FBaseStructInfo& Info : BaseStructInfos)
-		{
-			if (Info.Struct)
-			{
-				Info.Struct->CreateUObjectEngineTypes(Phase);
-			}
-		}
-		break;
-
-	case ECreateEngineTypesPhase::Phase2:
-		for (TSharedRef<FUnrealFunctionDefinitionInfo>& FunctionDef : Functions)
-		{
-			FunctionDef->CreateUObjectEngineTypes(Phase);
-		}
-		break;
-
-	case ECreateEngineTypesPhase::Phase3:
-		for (TSharedRef<FUnrealPropertyDefinitionInfo>& PropertyDef : Properties)
-		{
-			PropertyDef->CreateUObjectEngineTypes(Phase);
-		}
-
-		for (TSharedRef<FUnrealFunctionDefinitionInfo>& FunctionDef : Functions)
-		{
-			FunctionDef->CreateUObjectEngineTypes(Phase);
-		}
-
-		GetStruct()->Bind();
-
-		// Internals will assert of we are relinking an intrinsic 
-		bool bStaticLink = true;
-		bool bRelinkExistingProperties = true;
-		if (FUnrealClassDefinitionInfo* ClassDef = UHTCast<FUnrealClassDefinitionInfo>(this))
-		{
-			bRelinkExistingProperties = !ClassDef->HasAnyClassFlags(CLASS_Intrinsic);
-			bStaticLink = !ClassDef->IsCompiledIn();
-		}
-		if (bStaticLink)
-		{
-			GetStruct()->StaticLink(bRelinkExistingProperties);
-		}
-		break;
-	}
-}
-
 void FUnrealStructDefinitionInfo::PostParseFinalizeInternal(EPostParseFinalizePhase Phase)
 {
 	FUnrealFieldDefinitionInfo::PostParseFinalizeInternal(Phase);
@@ -1843,17 +1567,6 @@ void FUnrealStructDefinitionInfo::ConcurrentPostParseFinalize()
 	}
 }
 
-void FUnrealStructDefinitionInfo::CreatePropertyEngineTypes()
-{
-	for (TSharedRef<FUnrealPropertyDefinitionInfo>& PropertyDef : Properties)
-	{
-		if (PropertyDef->GetPropertySafe() == nullptr)
-		{
-			FPropertyTraits::CreateEngineType(PropertyDef);
-		}
-	}
-}
-
 void FUnrealStructDefinitionInfo::AddFunction(TSharedRef<FUnrealFunctionDefinitionInfo> FunctionDef)
 {
 	Functions.Add(FunctionDef);
@@ -1880,13 +1593,6 @@ TSharedRef<FScope> FUnrealStructDefinitionInfo::GetScope()
 	}
 }
 
-void FUnrealStructDefinitionInfo::SetObject(UObject* InObject)
-{
-	check(InObject != nullptr);
-
-	FUnrealFieldDefinitionInfo::SetObject(InObject);
-}
-
 FUnrealScriptStructDefinitionInfo::FUnrealScriptStructDefinitionInfo(FUnrealSourceFile& InSourceFile, int32 InLineNumber, FString&& InNameCPP, FName InName)
 	: FUnrealStructDefinitionInfo(InSourceFile, InLineNumber, MoveTemp(InNameCPP), InName, InSourceFile.GetPackageDef())
 { }
@@ -1908,40 +1614,6 @@ uint32 FUnrealScriptStructDefinitionInfo::GetHash(FUnrealTypeDefinitionInfo& Ref
 	return FUnrealStructDefinitionInfo::GetHash(ReferencingDef, bIncludeNoExport);
 }
 
-void FUnrealScriptStructDefinitionInfo::CreateUObjectEngineTypesInternal(ECreateEngineTypesPhase Phase)
-{
-	FUnrealStructDefinitionInfo::CreateUObjectEngineTypesInternal(Phase);
-
-	switch (Phase)
-	{
-	case ECreateEngineTypesPhase::Phase1:
-	{
-		UPackage* Package = GetPackageDef().GetPackage();
-		const FString& StructName = GetNameCPP();
-		FString StructNameStripped = GetClassNameWithPrefixRemoved(*StructName);
-
-		UScriptStruct* ScriptStruct = new(EC_InternalUseOnlyConstructor, Package, FName(StructNameStripped), RF_Public) UScriptStruct(FObjectInitializer());
-		ScriptStruct->StructFlags = StructFlags;
-		ScriptStruct->GetPackage()->GetMetaData()->SetObjectValues(ScriptStruct, GetMetaDataMap());
-		if (FUnrealStructDefinitionInfo* SuperStructDef = GetSuperStruct())
-		{
-			ScriptStruct->SetSuperStruct(SuperStructDef->GetStruct());
-		}
-		SetObject(ScriptStruct);
-
-		// The structure needs to be prepared at this point.  
-		ScriptStruct->PrepareCppStructOps();
-		break;
-	}
-
-	case ECreateEngineTypesPhase::Phase2:
-		break;
-
-	case ECreateEngineTypesPhase::Phase3:
-		break;
-	}
-}
-
 void FUnrealScriptStructDefinitionInfo::PostParseFinalizeInternal(EPostParseFinalizePhase Phase)
 {
 	FUnrealStructDefinitionInfo::PostParseFinalizeInternal(Phase);
@@ -2006,8 +1678,7 @@ void FUnrealScriptStructDefinitionInfo::SetHasInstancedReference()
 
 bool FUnrealScriptStructDefinitionInfo::HasNoOpConstructor() const
 {
-	UScriptStruct::ICppStructOps* CppStructOps = UScriptStruct::FindDeferredCppStructOps(GetStructPathName());
-	return CppStructOps ? CppStructOps->HasNoopConstructor() : false;
+	return HasAnyScriptStructExportFlags(STRUCTEXPORT_HasNoOpConstructor);
 }
 
 FUnrealClassDefinitionInfo::FUnrealClassDefinitionInfo(FUnrealSourceFile& InSourceFile, int32 InLineNumber, FString&& InNameCPP, FName InName, bool bInIsInterface)
@@ -2122,125 +1793,6 @@ uint32 FUnrealClassDefinitionInfo::GetHash(FUnrealTypeDefinitionInfo& Referencin
 	return FUnrealStructDefinitionInfo::GetHash(ReferencingDef, bIncludeNoExport);
 }
 
-void FUnrealClassDefinitionInfo::CreateUObjectEngineTypesInternal(ECreateEngineTypesPhase Phase)
-{
-	FUnrealStructDefinitionInfo::CreateUObjectEngineTypesInternal(Phase);
-
-	switch (Phase)
-	{
-	case ECreateEngineTypesPhase::Phase1:
-	{
-		if (ClassWithin == nullptr)
-		{
-			Throwf(TEXT("ClassWithin is not set.  This can be the result of the class's UCLASS not being parsed due to statement skipping preceding the UCLASS keyword."));
-		}
-		check(ClassWithin);
-		check(GetObjectSafe() == nullptr);
-
-		const FString& ClassName = GetNameCPP();
-		FString ClassNameStripped = GetClassNameWithPrefixRemoved(*ClassName);
-
-		// See if we have an existing engine object we should use
-		UClass* Class = FEngineAPI::FindFirstObject<UClass>(*ClassNameStripped, EFindFirstObjectOptions::EnsureIfAmbiguous);
-
-		if (Class == nullptr)
-		{
-			UPackage* Package = GetPackageDef().GetPackage();
-
-			// Create new class.
-			Class = new(EC_InternalUseOnlyConstructor, Package, *ClassNameStripped, RF_Public | RF_Standalone) UClass(FObjectInitializer(), nullptr);
-
-
-			Class->ClassFlags = ClassFlags;
-			Class->ClassCastFlags = ClassCastFlags;
-			Class->PropertiesSize = PropertiesSize;
-			Class->ClassConfigName = ClassConfigName;
-			Class->ClassWithin = ClassWithin->GetClassSafe();
-			Class->SetInternalFlags(GetInternalFlags());
-			Class->GetPackage()->GetMetaData()->SetObjectValues(Class, GetMetaDataMap());
-
-			// Setup the base class
-			if (FUnrealClassDefinitionInfo* SuperClassDef = GetSuperClass())
-			{
-				Class->SetSuperStruct(SuperClassDef->GetClass());
-			}
-
-			// Add the class flags from the interfaces
-			for (FUnrealStructDefinitionInfo::FBaseStructInfo& BaseClassInfo : GetBaseStructInfos())
-			{
-				if (FUnrealClassDefinitionInfo* BaseClassDef = UHTCast<FUnrealClassDefinitionInfo>(BaseClassInfo.Struct))
-				{
-					if (BaseClassDef->HasAnyClassFlags(CLASS_Interface))
-					{
-						Class->Interfaces.Emplace(BaseClassDef->GetClass(), 1, false);
-					}
-				}
-			}
-		}
-		else
-		{
-			bIsCompiledIn = true;
-		}
-		SetObject(Class);
-		break;
-	}
-
-	case ECreateEngineTypesPhase::Phase2:
-	{
-
-		// Only do this for UHT types
-		if (!bIsCompiledIn)
-		{
-			// Initialize the class object
-			UClass* Class = GetClass();
-
-			// Clear the property size
-			Class->PropertiesSize = 0;
-
-			// Make the visible to the package
-			Class->ClearFlags(RF_Transient);
-			check(Class->HasAnyFlags(RF_Public));
-			check(Class->HasAnyFlags(RF_Standalone));
-
-			// Finalize all of the children introduced in this class
-			for (TSharedRef<FUnrealFunctionDefinitionInfo> FunctionDef : GetFunctions())
-			{
-				UFunction* Function = FunctionDef->GetFunction();
-				Class->AddFunctionToFunctionMap(Function, Function->GetFName());
-			}
-		}
-		break;
-	}
-
-	case ECreateEngineTypesPhase::Phase3:
-	{
-		// Only do this for UHT types
-		if (!bIsCompiledIn)
-		{
-
-			// Initialize the class object
-			UClass* Class = GetClass();
-
-			if (!HasAnyClassFlags(CLASS_Native))
-			{
-				//Class->UnMark(EObjectMark(OBJECTMARK_TagImp | OBJECTMARK_TagExp));
-			}
-			else if (!IsNoExport() && !HasAnyClassFlags(CLASS_Intrinsic))
-			{
-				//Class->UnMark(OBJECTMARK_TagImp);
-				//Class->Mark(OBJECTMARK_TagExp);
-			}
-
-			// This needs to be done outside of parallel blocks because it will modify UClass memory.
-			// Later calls to SetUpUhtReplicationData inside parallel blocks should be fine, because
-			// they will see the memory has already been set up, and just return the parent pointer.
-			Class->SetUpUhtReplicationData();
-		}
-		break;
-	}
-	}
-}
-
 void FUnrealClassDefinitionInfo::PostParseFinalizeInternal(EPostParseFinalizePhase Phase)
 {
 	FUnrealStructDefinitionInfo::PostParseFinalizeInternal(Phase);
@@ -3029,71 +2581,6 @@ FUnrealFunctionDefinitionInfo* FUnrealFunctionDefinitionInfo::GetSuperFunction()
 	return UHTCast<FUnrealFunctionDefinitionInfo>(GetSuperStruct());
 }
 
-void FUnrealFunctionDefinitionInfo::CreateUObjectEngineTypesInternal(ECreateEngineTypesPhase Phase)
-{
-	// Invoke the base class creatation
-	FUnrealStructDefinitionInfo::CreateUObjectEngineTypesInternal(Phase);
-
-	switch (Phase)
-	{
-	case ECreateEngineTypesPhase::Phase1:
-		break;
-
-	case ECreateEngineTypesPhase::Phase2:
-	{
-		// We have to precreate the function prior to invoking the parent finalize
-		UFunction* Function = nullptr;
-		switch (FunctionType)
-		{
-		case EFunctionType::Function:
-			Function = new(EC_InternalUseOnlyConstructor, GetOuter()->GetObject(), *GetNameCPP(), RF_Public) UFunction(FObjectInitializer(), nullptr);
-			break;
-		case EFunctionType::Delegate:
-			Function = new(EC_InternalUseOnlyConstructor, GetOuter()->GetObject(), *GetNameCPP(), RF_Public) UDelegateFunction(FObjectInitializer(), nullptr);
-			break;
-		case EFunctionType::SparseDelegate:
-		{
-			USparseDelegateFunction* USPF = new(EC_InternalUseOnlyConstructor, GetOuter()->GetObject(), *GetNameCPP(), RF_Public) USparseDelegateFunction(FObjectInitializer(), nullptr);
-			USPF->OwningClassName = SparseOwningClassName;
-			USPF->DelegateName = SparseDelegateName;
-			Function = USPF;
-			break;
-		}
-		}
-		check(Function);
-
-		Function->ReturnValueOffset = MAX_uint16;
-		Function->FirstPropertyToInit = nullptr;
-		Function->FunctionFlags |= FunctionData.FunctionFlags;
-		Function->GetPackage()->GetMetaData()->SetObjectValues(Function, GetMetaDataMap());
-
-		SetObject(Function);
-
-		if (FUnrealStructDefinitionInfo* StructDef = UHTCast<FUnrealStructDefinitionInfo>(GetOuter()))
-		{
-			UStruct* Struct = StructDef->GetStruct();
-			Function->Next = Struct->Children;
-			Struct->Children = Function;
-		}
-		Function->NumParms = uint8(GetProperties().Num());
-		Function->Bind();
-		break;
-	}
-
-	case ECreateEngineTypesPhase::Phase3:
-	{
-		// The following code is only performed on functions in a class.  
-		if (UHTCast<FUnrealClassDefinitionInfo>(GetOuter()) != nullptr)
-		{
-			UFunction* Function = GetFunction();
-			Function->StaticLink(true);
-			Function->ParmsSize = 0;
-		}
-		break;
-	}
-	}
-}
-
 void FUnrealFunctionDefinitionInfo::PostParseFinalizeInternal(EPostParseFinalizePhase Phase)
 {
 	// Invoke the base class finalization
diff --git a/Engine/Source/Programs/UnrealHeaderTool/Private/UnrealTypeDefinitionInfo.h b/Engine/Source/Programs/UnrealHeaderTool/Private/UnrealTypeDefinitionInfo.h
--- a/Engine/Source/Programs/UnrealHeaderTool/Private/UnrealTypeDefinitionInfo.h
+++ b/Engine/Source/Programs/UnrealHeaderTool/Private/UnrealTypeDefinitionInfo.h
@@ -10,8 +10,6 @@
 #include "BaseParser.h"
 #include "ParserHelper.h"
 
-#define UHT_ENABLE_ENGINE_TYPE_CHECKS 0
-
 // Forward declarations.
 class FHeaderParser;
 class FScope;
@@ -228,9 +226,6 @@ protected:
 		checkf(false, TEXT("You can not perform remapping of meta data values on this type"));
 		return nullptr;
 	}
-#if UHT_ENABLE_ENGINE_TYPE_CHECKS
-	virtual void CheckFindMetaData(const FName& Key, const FString* ValuePtr) const = 0;
-#endif
 
 private:
 	FName GetMetaDataKey(const TCHAR* Key, int32 NameIndex, EFindName FindName) const;
@@ -416,11 +411,6 @@ public:
 	 */
 	virtual void OnPostParsePropertyFlagsChanged(FUnrealPropertyDefinitionInfo& PropertyDef) {}
 
-	/**
-	 * Create UObject engine types.
-	 */
-	void CreateUObjectEngineTypes(ECreateEngineTypesPhase Phase);
-
 	/**
 	 * Perform any post parsing finalization and validation
 	 */
@@ -605,20 +595,6 @@ protected:
 
 	virtual void ValidateMetaDataFormat(const FName InKey, ECheckedMetadataSpecifier InCheckedMetadataSpecifier, const FString& InValue) const;
 
-	/**
-	 * Create UObject engine types
-	 */
-	virtual void CreateUObjectEngineTypesInternal(ECreateEngineTypesPhase Phase) 
-	{
-		if (Phase == ECreateEngineTypesPhase::Phase1)
-		{
-			if (Outer != nullptr && Outer->AsPackage() == nullptr)
-			{
-				CreateUObjectEngineTypes(Phase);
-			}
-		}
-	}
-
 	/**
 	 * Perform any post parsing finalization and validation
 	 */
@@ -662,32 +638,6 @@ public:
 		return this;
 	}
 
-	/**
-	 * Return the Engine instance associated with the compiler instance
-	 */
-	FProperty* GetProperty() const
-	{
-		check(Property);
-		return Property;
-	}
-
-	/**
-	 * Return the Engine instance associated with the compiler instance, does not check the pointer, can be null
-	 */
-	FProperty* GetPropertySafe() const
-	{
-		return Property;
-	}
-
-	/**
-	 * Sets the engine type
-	 */
-	void SetProperty(FProperty* InProperty)
-	{
-		check(Property == nullptr || Property == InProperty);
-		Property = InProperty;
-	}
-
 	void ExportCppDeclaration(FOutputDevice& Out, EExportedDeclaration::Type DeclarationType, const TCHAR* ArrayDimOverride = NULL, uint32 AdditionalExportCPPFlags = 0, bool bSkipParameterName = false) const;
 
 	FString GetCPPMacroType(FString& ExtendedTypeText) const;
@@ -697,18 +647,12 @@ public:
 	 */
 	virtual FString GetName() const override
 	{
-#if UHT_ENABLE_ENGINE_TYPE_CHECKS
-		check(GetPropertySafe() == nullptr || GetPropertySafe()->GetName() == GetNameCPP());
-#endif
 		return GetNameCPP();
 	}
 
 	/** Returns the logical name of this object */
 	virtual FName GetFName() const override
 	{
-#if UHT_ENABLE_ENGINE_TYPE_CHECKS
-		check(GetPropertySafe() == nullptr || GetPropertySafe()->GetFName() == Name);
-#endif
 		return Name;
 	}
 
@@ -773,9 +717,6 @@ public:
 	FString GetNameWithDeprecated() const
 	{
 		FString OutName = HasAnyPropertyFlags(CPF_Deprecated) ? GetName() + TEXT("_DEPRECATED") : GetName();
-#if UHT_ENABLE_ENGINE_TYPE_CHECKS
-		check(GetPropertySafe() == nullptr || GetPropertySafe()->GetNameCPP() == OutName);
-#endif
 		return OutName;
 	}
 
@@ -794,33 +735,16 @@ public:
 	 */
 	EPropertyFlags GetPropertyFlags() const
 	{
-#if UHT_ENABLE_ENGINE_TYPE_CHECKS
-		check(GetPropertySafe() == nullptr || (GetPropertySafe()->GetPropertyFlags() & ~CPF_ComputedFlags) == PropertyBase.PropertyFlags);
-#endif
 		return PropertyBase.PropertyFlags;
 	}
 
 	void SetPropertyFlags(EPropertyFlags NewFlags)
 	{
-#if UHT_ENABLE_ENGINE_TYPE_CHECKS
-		check(GetPropertySafe() == nullptr || (GetPropertySafe()->GetPropertyFlags() & ~CPF_ComputedFlags) == PropertyBase.PropertyFlags);
-#endif
-		if (FProperty* LocalProperty = GetPropertySafe())
-		{
-			LocalProperty->PropertyFlags |= NewFlags;
-		}
 		PropertyBase.PropertyFlags |= NewFlags;
 	}
 
 	void ClearPropertyFlags(EPropertyFlags NewFlags)
 	{
-#if UHT_ENABLE_ENGINE_TYPE_CHECKS
-		check(GetPropertySafe() == nullptr || (GetPropertySafe()->GetPropertyFlags() & ~CPF_ComputedFlags) == PropertyBase.PropertyFlags);
-#endif
-		if (FProperty* LocalProperty = GetPropertySafe())
-		{
-			LocalProperty->PropertyFlags &= ~NewFlags;
-		}
 		PropertyBase.PropertyFlags &= ~NewFlags;
 	}
 
@@ -837,10 +761,6 @@ public:
 	 */
 	bool HasAnyPropertyFlags(uint64 FlagsToCheck) const
 	{
-#if UHT_ENABLE_ENGINE_TYPE_CHECKS
-		check((FlagsToCheck & CPF_ComputedFlags) == 0);
-		check(GetPropertySafe() == nullptr || (GetPropertySafe()->GetPropertyFlags() & ~CPF_ComputedFlags) == PropertyBase.PropertyFlags);
-#endif
 		return (PropertyBase.PropertyFlags & FlagsToCheck) != 0;
 	}
 
@@ -856,19 +776,11 @@ public:
 	 */
 	bool HasAllPropertyFlags(uint64 FlagsToCheck) const
 	{
-#if UHT_ENABLE_ENGINE_TYPE_CHECKS
-		check((FlagsToCheck & CPF_ComputedFlags) == 0);
-		check(GetPropertySafe() == nullptr || (GetPropertySafe()->GetPropertyFlags() & ~CPF_ComputedFlags) == PropertyBase.PropertyFlags);
-#endif
 		return (PropertyBase.PropertyFlags & FlagsToCheck) == FlagsToCheck;
 	}
 
 	bool HasSpecificPropertyFlags(uint64 FlagsToCheck, uint64 ExpectedFlags)
 	{
-#if UHT_ENABLE_ENGINE_TYPE_CHECKS
-		check((FlagsToCheck & CPF_ComputedFlags) == 0);
-		check(GetPropertySafe() == nullptr || (GetPropertySafe()->GetPropertyFlags() & ~CPF_ComputedFlags) == PropertyBase.PropertyFlags);
-#endif
 		return (PropertyBase.PropertyFlags & FlagsToCheck) == ExpectedFlags;
 	}
 
@@ -890,9 +802,6 @@ public:
 	bool IsEditorOnlyProperty() const
 	{
 		bool bResult = HasAnyPropertyFlags(CPF_DevelopmentAssets);
-#if UHT_ENABLE_ENGINE_TYPE_CHECKS
-		check(GetPropertySafe() == nullptr || bResult == GetPropertySafe()->IsEditorOnlyProperty());
-#endif
 		return bResult;
 	}
 
@@ -905,9 +814,6 @@ public:
 	bool ContainsInstancedObjectProperty() const
 	{
 		bool bResult = HasAnyPropertyFlags(CPF_ContainsInstancedReference | CPF_InstancedReference);
-#if UHT_ENABLE_ENGINE_TYPE_CHECKS
-		check(GetPropertySafe() == nullptr || bResult == GetPropertySafe()->ContainsInstancedObjectProperty());
-#endif
 		return bResult;
 	}
 
@@ -1208,9 +1114,6 @@ public:
 	 */
 	FName GetRepNotifyFunc() const
 	{
-#if UHT_ENABLE_ENGINE_TYPE_CHECKS
-		check(GetPropertySafe() == nullptr || PropertyBase.RepNotifyName == GetPropertySafe()->RepNotifyFunc);
-#endif
 		return PropertyBase.RepNotifyName;
 	}
 
@@ -1224,7 +1127,6 @@ public:
 	 */
 	virtual void AddMetaData(TMap<FName, FString>&& InMetaData) override
 	{
-		check(Property == nullptr);
 		GetMetaDataMap().Append(MoveTemp(InMetaData));
 	}
 
@@ -1245,18 +1147,6 @@ public:
 	virtual void ConcurrentPostParseFinalize() override;
 
 protected:
-#if UHT_ENABLE_ENGINE_TYPE_CHECKS
-	virtual void CheckFindMetaData(const FName& Key, const FString* ValuePtr) const
-	{
-		if (FProperty* LocalProperty = GetPropertySafe())
-		{
-			const FString* CheckPtr = LocalProperty->FindMetaData(Key);
-			check((CheckPtr == nullptr && ValuePtr == nullptr) || (CheckPtr != nullptr && ValuePtr != nullptr && *CheckPtr == *ValuePtr));
-		}
-	}
-#endif
-
-	virtual void CreateUObjectEngineTypesInternal(ECreateEngineTypesPhase Phase) override;
 
 	/**
 	 * Perform any post parsing finalization and validation
@@ -1270,7 +1160,6 @@ private:
 	FString TypePackageName;
 	TSharedPtr<FUnrealPropertyDefinitionInfo> KeyPropDef;
 	TSharedPtr<FUnrealPropertyDefinitionInfo> ValuePropDef;
-	FProperty* Property = nullptr;
 	int32 ParsePosition;
 	EVariableCategory VariableCategory = EVariableCategory::Member;
 	EAccessSpecifier AccessSpecifier = ACCESS_Public;
@@ -1293,23 +1182,6 @@ public:
 		return this;
 	}
 
-	/**
-	 * Return the Engine instance associated with the compiler instance
-	 */
-	UObject* GetObject() const
-	{
-		check(Object);
-		return Object;
-	}
-
-	/**
-	 * Return the Engine instance associated with the compiler instance without the check that it is set
-	 */
-	UObject* GetObjectSafe() const
-	{
-		return Object;
-	}
-
 	/**
 	 * Return the engine class name
 	 */
@@ -1319,38 +1191,18 @@ public:
 		return ClassName;
 	}
 
-	/**
-	 * Set the Engine instance associated with the compiler instance
-	 */
-	virtual void SetObject(UObject* InObject)
-	{
-		check(InObject != nullptr);
-		check(Object == nullptr);
-		Object = InObject;
-#if UHT_ENABLE_ENGINE_TYPE_CHECKS
-		check(Object->GetFName() == Name);
-		check(Object->GetInternalFlags() == InternalObjectFlags);
-#endif
-	}
-
 	/**
 	 * Returns the name of this object (with no path information)
 	 */
 	virtual FString GetName() const override
 	{
 		FString Result = GetFName().ToString();
-#if UHT_ENABLE_ENGINE_TYPE_CHECKS
-		check(GetObjectSafe() == nullptr || GetObjectSafe()->GetName() == Result);
-#endif
 		return Result;
 	}
 
 	/** Returns the logical name of this object */
 	virtual FName GetFName() const override
 	{
-#if UHT_ENABLE_ENGINE_TYPE_CHECKS
-		check(GetObjectSafe() == nullptr || GetObjectSafe()->GetFName() == Name);
-#endif
 		return Name;
 	}
 
@@ -1394,9 +1246,6 @@ public:
 	 */
 	EInternalObjectFlags GetInternalFlags() const
 	{
-#if UHT_ENABLE_ENGINE_TYPE_CHECKS
-		check(GetObjectSafe() == nullptr || GetObjectSafe()->GetInternalFlags() == InternalObjectFlags);
-#endif
 		return InternalObjectFlags;
 	}
 
@@ -1406,13 +1255,6 @@ public:
 	void SetInternalFlags(EInternalObjectFlags FlagsToSet)
 	{
 		InternalObjectFlags |= FlagsToSet;
-		if (UObject* Obj = GetObjectSafe())
-		{
-			Obj->SetInternalFlags(FlagsToSet);
-#if UHT_ENABLE_ENGINE_TYPE_CHECKS
-			check(Obj->GetInternalFlags() == InternalObjectFlags);
-#endif
-		}
 	}
 
 	/**
@@ -1421,13 +1263,6 @@ public:
 	void ClearInternalFlags(EInternalObjectFlags FlagsToClear)
 	{
 		InternalObjectFlags &= ~FlagsToClear;
-		if (UObject* Obj = GetObjectSafe())
-		{
-			Obj->ClearInternalFlags(FlagsToClear);
-#if UHT_ENABLE_ENGINE_TYPE_CHECKS
-			check(Obj->GetInternalFlags() == InternalObjectFlags);
-#endif
-		}
 	}
 
 	/**
@@ -1438,9 +1273,6 @@ public:
 	 */
 	bool HasAnyInternalFlags(EInternalObjectFlags FlagsToCheck) const
 	{
-#if UHT_ENABLE_ENGINE_TYPE_CHECKS
-		check(GetObjectSafe() == nullptr || GetObjectSafe()->GetInternalFlags() == InternalObjectFlags);
-#endif
 		return EnumHasAnyFlags(InternalObjectFlags, FlagsToCheck);
 	}
 
@@ -1489,40 +1321,8 @@ protected:
 		, Name(InName)
 	{ }
 
-#if UHT_ENABLE_ENGINE_TYPE_CHECKS
-	virtual void CheckFindMetaData(const FName& Key, const FString* ValuePtr) const override
-	{
-		if (UObject* Obj = GetObjectSafe())
-		{
-			TMap<FName, FString>* PackageMap = GetUObjectMetaDataMap();
-			const FString* CheckPtr = PackageMap ? PackageMap->Find(Key) : nullptr;
-			check((CheckPtr == nullptr && ValuePtr == nullptr) || (CheckPtr != nullptr && ValuePtr != nullptr && *CheckPtr == *ValuePtr));
-		}
-	}
-#endif
-
-	UMetaData* GetUObjectMetaData() const
-	{
-		if (UObject* Obj = GetObjectSafe())
-		{
-			UPackage* Package = Obj->GetPackage();
-			check(Package);
-			UMetaData* Metadata = Package->GetMetaData();
-			check(Metadata);
-			return Metadata;
-		}
-		return nullptr;
-	}
-
-	TMap<FName, FString>* GetUObjectMetaDataMap() const
-	{
-		UMetaData* Metadata = GetUObjectMetaData();
-		return Metadata ? Metadata->ObjectMetaDataMap.Find(GetObject()) : nullptr;
-	}
-
 private:
 	TMap<FName, FString> MetaData;
-	UObject* Object = nullptr;
 	FName Name = NAME_None;
 	EInternalObjectFlags InternalObjectFlags = EInternalObjectFlags::None;
 };
@@ -1542,22 +1342,6 @@ public:
 		return this;
 	}
 
-	/**
-	 * Return the Engine instance associated with the compiler instance
-	 */
-	UPackage* GetPackage() const
-	{
-		return static_cast<UPackage*>(GetObject());
-	}
-
-	/**
-	 * Return the Engine instance associated with the compiler instance
-	 */
-	UPackage* GetPackageSafe() const
-	{
-		return static_cast<UPackage*>(GetObjectSafe());
-	}
-
 	/**
 	 * Return the engine class name
 	 */
@@ -1616,6 +1400,14 @@ public:
 		return API;
 	}
 
+	/**
+	 * Get the short name of the package.
+	 */
+	const FString& GetShortName() const
+	{
+		return ShortName;
+	}
+
 	/**
 	 * Get the short name of the package uppercased.
 	 */
@@ -1624,6 +1416,11 @@ public:
 		return ShortUpperName;
 	}
 
+	uint32 GetPackageFlags() const
+	{
+		return PackageFlags;
+	}
+
 	/**
 	 * Add a unique cross module reference for this field
 	 */
@@ -1654,10 +1451,6 @@ public:
 	}
 
 protected:
-	/**
-	 * Create UObject engine types
-	 */
-	virtual void CreateUObjectEngineTypesInternal(ECreateEngineTypesPhase Phase) override;
 
 	/**
 	 * Perform any post parsing finalization and validation
@@ -1671,8 +1464,10 @@ private:
 	FString SingletonName;
 	FString SingletonNameChopped;
 	FString ExternDecl;
+	FString ShortName;
 	FString ShortUpperName;
 	FString API;
+	uint32 PackageFlags = 0;
 	bool bWriteClassesH = false;
 };
 
@@ -1699,22 +1494,6 @@ public:
 		return this;
 	}
 
-	/**
-	 * Return the Engine instance associated with the compiler instance
-	 */
-	UField* GetField() const
-	{
-		return static_cast<UField*>(GetObject());
-	}
-
-	/**
-	 * Return the Engine instance associated with the compiler instance without the check that it is set
-	 */
-	UField* GetFieldSafe() const
-	{
-		return static_cast<UField*>(GetObjectSafe());
-	}
-
 	/**
 	 * Return the engine class name
 	 */
@@ -1861,31 +1640,12 @@ public:
 		return ClassName;
 	}
 
-	/**
-	 * Return the Engine instance associated with the compiler instance
-	 */
-	UEnum* GetEnum() const
-	{
-		return static_cast<UEnum*>(GetObject());
-	}
-
-	/**
-	 * Return the Engine instance associated with the compiler instance but does not check to see if it is set
-	 */
-	UEnum* GetEnumSafe() const
-	{
-		return static_cast<UEnum*>(GetObjectSafe());
-	}
-
 	/**
 	 * Test if the enum contains and of the given flags
 	 */
 	bool HasAnyEnumFlags(EEnumFlags InFlags) const
 	{
 		bool bResults = EnumHasAnyFlags(EnumFlags, InFlags);
-#if UHT_ENABLE_ENGINE_TYPE_CHECKS
-		check(GetEnumSafe() == nullptr || GetEnumSafe()->HasAnyEnumFlags(InFlags) == bResults);
-#endif
 		return bResults;
 	}
 
@@ -1909,9 +1669,6 @@ public:
 	 */
 	int32 NumEnums() const
 	{
-#if UHT_ENABLE_ENGINE_TYPE_CHECKS
-		check(GetEnumSafe() == nullptr || GetEnumSafe()->NumEnums() == Names.Num());
-#endif
 		return Names.Num();
 	}
 
@@ -1922,9 +1679,6 @@ public:
 	int64 GetValueByIndex(int32 Index) const
 	{
 		check(Names.IsValidIndex(Index));
-#if UHT_ENABLE_ENGINE_TYPE_CHECKS
-		check(GetEnumSafe() == nullptr || GetEnumSafe()->GetValueByIndex(Index) == Names[Index].Value);
-#endif
 		return Names[Index].Value;
 	}
 
@@ -1944,9 +1698,6 @@ public:
 	 */
 	UEnum::ECppForm GetCppForm() const
 	{
-#if UHT_ENABLE_ENGINE_TYPE_CHECKS
-		check(GetEnumSafe() == nullptr || GetEnumSafe()->GetCppForm() == CppForm);
-#endif
 		return CppForm;
 	}
 
@@ -1955,9 +1706,6 @@ public:
 	 */
 	const FString& GetCppType() const
 	{
-#if UHT_ENABLE_ENGINE_TYPE_CHECKS
-		check(GetEnumSafe() == nullptr || GetEnumSafe()->CppType == CppType);
-#endif
 		return CppType;
 	}
 
@@ -1967,10 +1715,6 @@ public:
 	void SetCppType(FString&& InCppType)
 	{
 		CppType = MoveTemp(InCppType);
-		if (GetEnumSafe() != nullptr)
-		{
-			GetEnumSafe()->CppType = CppType;
-		}
 	}
 
 	/**
@@ -1994,13 +1738,6 @@ public:
 		Names = InNames;
 		check(CppForm == InCppForm);
 		EnumFlags = InFlags;
-#if UHT_ENABLE_ENGINE_TYPE_CHECKS
-		if (UEnum* Enum = GetEnumSafe())
-		{
-			Enum->SetEnums(InNames, InCppForm, InFlags, false /*bAddMaxKeyIfMissing*/);
-			check(Enum->GetCppForm() == InCppForm);
-		}
-#endif
 		return true;
 	}
 
@@ -2037,10 +1774,6 @@ public:
 	}
 
 protected:
-	/**
-	 * Create the shell of the UObject.  Only invoked if the object has not been set
-	 */
-	virtual void CreateUObjectEngineTypesInternal(ECreateEngineTypesPhase Phase) override;
 
 	virtual FString GetMetaDataIndexName(int32 NameIndex) const
 	{
@@ -2087,24 +1820,6 @@ public:
 
 	virtual TSharedRef<FScope> GetScope() override;
 
-	virtual void SetObject(UObject* InObject) override;
-
-	/**
-	 * Return the Engine instance associated with the compiler instance
-	 */
-	UStruct* GetStruct() const
-	{
-		return static_cast<UStruct*>(GetObject());
-	}
-
-	/**
-	 * Return the Engine instance associated with the compiler instance without the null check
-	 */
-	UStruct* GetStructSafe() const
-	{
-		return static_cast<UStruct*>(GetObjectSafe());
-	}
-
 	/**
 	 * @return true if this object is of the specified type.
 	 */
@@ -2125,13 +1840,6 @@ public:
 	virtual const TCHAR* GetPrefixCPP() const
 	{
 		const TCHAR* ReturnValue = TEXT("F");
-#if UHT_ENABLE_ENGINE_TYPE_CHECKS
-		if (UStruct* Struct = GetStructSafe())
-		{
-			const TCHAR* EngineValue = Struct->GetPrefixCPP();
-			check(FCString::Strcmp(EngineValue, ReturnValue) == 0);
-		}
-#endif
 		return ReturnValue;
 	}
 
@@ -2148,11 +1856,6 @@ public:
 		return FString::Printf(TEXT("%s%s"), (bForceInterface ? TEXT("I") : GetPrefixCPP()), *GetName());
 	}
 
-	/**
-	 * Create the underlying properties for the structure
-	 */
-	void CreatePropertyEngineTypes();
-
 	/**
 	 * Add a new property to the structure
 	 */
@@ -2329,11 +2032,6 @@ public:
 protected:
 	FUnrealStructDefinitionInfo(FUnrealSourceFile& InSourceFile, int32 InLineNumber, FString&& InNameCPP, FName InName, FUnrealObjectDefinitionInfo& InOuter);
 
-	/**
-	 * Create UObject engine types
-	 */
-	virtual void CreateUObjectEngineTypesInternal(ECreateEngineTypesPhase Phase) override;
-
 	/**
 	 * Perform any post parsing finalization and validation
 	 */
@@ -2389,22 +2087,6 @@ public:
 		return ClassName;
 	}
 
-	/**
-	 * Return the Engine instance associated with the compiler instance
-	 */
-	UScriptStruct* GetScriptStruct() const
-	{
-		return static_cast<UScriptStruct*>(GetObject());
-	}
-
-	/**
-	 * Return the Engine instance associated with the compiler instance without the null check
-	 */
-	UScriptStruct* GetScriptStructSafe() const
-	{
-		return static_cast<UScriptStruct*>(GetObjectSafe());
-	}
-
 	/**
 	 * Return the super script struct
 	 */
@@ -2415,9 +2097,6 @@ public:
 	 */
 	EStructFlags GetStructFlags() const
 	{
-#if UHT_ENABLE_ENGINE_TYPE_CHECKS
-		check(GetScriptStructSafe() == nullptr || (GetScriptStructSafe()->StructFlags & ~STRUCT_ComputedFlags) == StructFlags);
-#endif
 		return StructFlags;
 	}
 
@@ -2426,13 +2105,6 @@ public:
 	 */
 	void SetStructFlags(EStructFlags FlagsToSet)
 	{
-#if UHT_ENABLE_ENGINE_TYPE_CHECKS
-		if (UScriptStruct* Struct = GetScriptStructSafe())
-		{
-			check((Struct->StructFlags & ~STRUCT_ComputedFlags) == StructFlags);
-			Struct->StructFlags = EStructFlags(int32(Struct->StructFlags) | int32(FlagsToSet));
-		}
-#endif
 		StructFlags = EStructFlags(int32(StructFlags) | int32(FlagsToSet));
 	}
 
@@ -2441,14 +2113,7 @@ public:
 	 */
 	void ClearStructFlags(EStructFlags FlagsToClear)
 	{
-#if UHT_ENABLE_ENGINE_TYPE_CHECKS
-		check(GetScriptStructSafe() == nullptr || (GetScriptStructSafe()->StructFlags & ~STRUCT_ComputedFlags) == StructFlags);
-#endif
 		StructFlags = EStructFlags(int32(StructFlags) & int32(~FlagsToClear));
-		if (UScriptStruct* Struct = GetScriptStructSafe())
-		{
-			Struct->StructFlags = EStructFlags(int32(Struct->StructFlags) & int32(~FlagsToClear));
-		}
 	}
 
 	/**
@@ -2461,14 +2126,6 @@ public:
 	 */
 	bool HasAnyStructFlags(EStructFlags FlagsToCheck) const
 	{
-#if UHT_ENABLE_ENGINE_TYPE_CHECKS
-		if (GetScriptStructSafe())
-		{
-			int32 a = GetScriptStructSafe()->StructFlags & ~STRUCT_ComputedFlags;
-			check(StructFlags == a);
-		}
-		check(GetScriptStructSafe() == nullptr || (GetScriptStructSafe()->StructFlags & ~STRUCT_ComputedFlags) == StructFlags);
-#endif
 		return EnumHasAnyFlags(StructFlags, FlagsToCheck);
 	}
 
@@ -2480,9 +2137,6 @@ public:
 	 */
 	bool HasAllStructFlags(EStructFlags FlagsToCheck) const
 	{
-#if UHT_ENABLE_ENGINE_TYPE_CHECKS
-		check(GetScriptStructSafe() == nullptr || (GetScriptStructSafe()->StructFlags & ~STRUCT_ComputedFlags) == StructFlags);
-#endif
 		return EnumHasAllFlags(StructFlags, FlagsToCheck);
 	}
 
@@ -2495,19 +2149,32 @@ public:
 	 */
 	bool HasSpecificStructFlags(EStructFlags FlagsToCheck, EStructFlags ExpectedFlags) const
 	{
-#if UHT_ENABLE_ENGINE_TYPE_CHECKS
-		check(GetScriptStructSafe() == nullptr || (GetScriptStructSafe()->StructFlags & ~STRUCT_ComputedFlags) == StructFlags);
-#endif
 		return (StructFlags & FlagsToCheck) == ExpectedFlags;
 	}
 
+	void SetScriptStructExportFlags(EScriptStructExportFlags NewFlags)
+	{
+		ScriptStructExportFlags = EScriptStructExportFlags(int32(ScriptStructExportFlags) | int32(NewFlags));
+	}
+
+	void ClearScriptStructExportFlags(EScriptStructExportFlags ClearFlags)
+	{
+		ScriptStructExportFlags = EScriptStructExportFlags(int32(ScriptStructExportFlags) & int32(~ClearFlags));
+	}
+
+	bool HasAnyScriptStructExportFlags(EScriptStructExportFlags FlagsToCheck) const
+	{
+		return EnumHasAnyFlags(ScriptStructExportFlags, FlagsToCheck);
+	}
+
+
 	/**
 	 * If it is native, it is assumed to have defaults because it has a constructor
 	 * @return true if this struct has defaults
 	 */
 	bool HasDefaults() const
 	{
-		return UScriptStruct::FindDeferredCppStructOps(GetStructPathName()) != nullptr;
+		return HasAnyScriptStructExportFlags(STRUCTEXPORT_HasDefaults);
 	}
 
 
@@ -2529,10 +2196,6 @@ public:
 	virtual void OnPostParsePropertyFlagsChanged(FUnrealPropertyDefinitionInfo& PropertyDef) override;
 
 protected:
-	/**
-	 * Create UObject engine types
-	 */
-	virtual void CreateUObjectEngineTypesInternal(ECreateEngineTypesPhase Phase) override;
 
 	/**
 	 * Perform any post parsing finalization and validation
@@ -2545,6 +2208,7 @@ private:
 
 	int32 MacroDeclaredLineNumber = INDEX_NONE;
 	EStructFlags StructFlags = STRUCT_NoFlags;
+	EScriptStructExportFlags ScriptStructExportFlags = (EScriptStructExportFlags)0;
 };
 
 /**
@@ -2593,30 +2257,11 @@ public:
 		}
 	}
 
-	/**
-	 * Return the Engine instance associated with the compiler instance
-	 */
-	UFunction* GetFunction() const
-	{
-		return static_cast<UFunction*>(GetObject());
-	}
-
-	/**
-	 * Return the Engine instance associated with the compiler instance without the check to see if it is set
-	 */
-	UFunction* GetFunctionSafe() const
-	{
-		return static_cast<UFunction*>(GetObjectSafe());
-	}
-
 	/**
 	 * Return the function flags
 	 */
 	EFunctionFlags GetFunctionFlags() const
 	{
-#if UHT_ENABLE_ENGINE_TYPE_CHECKS
-		check(GetFunctionSafe() == nullptr || GetFunctionSafe()->FunctionFlags == FunctionData.FunctionFlags);
-#endif
 		return FunctionData.FunctionFlags;
 	}
 
@@ -2625,13 +2270,6 @@ public:
 	 */
 	void SetFunctionFlags(EFunctionFlags FlagsToSet)
 	{
-#if UHT_ENABLE_ENGINE_TYPE_CHECKS
-		check(GetFunctionSafe() == nullptr || GetFunctionSafe()->FunctionFlags == FunctionData.FunctionFlags);
-#endif
-		if (GetFunctionSafe())
-		{
-			GetFunctionSafe()->FunctionFlags |= FlagsToSet;
-		}
 		FunctionData.FunctionFlags |= FlagsToSet;
 	}
 
@@ -2645,9 +2283,6 @@ public:
 	 */
 	bool HasAnyFunctionFlags(EFunctionFlags FlagsToCheck) const
 	{
-#if UHT_ENABLE_ENGINE_TYPE_CHECKS
-		check(GetFunctionSafe() == nullptr || GetFunctionSafe()->FunctionFlags == FunctionData.FunctionFlags);
-#endif
 		return EnumHasAnyFlags(FunctionData.FunctionFlags, FlagsToCheck);
 	}
 
@@ -2659,9 +2294,6 @@ public:
 	 */
 	bool HasAllFunctionFlags(EFunctionFlags FlagsToCheck) const
 	{
-#if UHT_ENABLE_ENGINE_TYPE_CHECKS
-		check(GetFunctionSafe() == nullptr || GetFunctionSafe()->FunctionFlags == FunctionData.FunctionFlags);
-#endif
 		return EnumHasAllFlags(FunctionData.FunctionFlags, FlagsToCheck);
 	}
 
@@ -2674,9 +2306,6 @@ public:
 	 */
 	bool HasSpecificFunctionFlags(EFunctionFlags FlagsToCheck, EFunctionFlags ExpectedFlags) const
 	{
-#if UHT_ENABLE_ENGINE_TYPE_CHECKS
-		check(GetFunctionSafe() == nullptr || GetFunctionSafe()->FunctionFlags == FunctionData.FunctionFlags);
-#endif
 		EFunctionFlags Flags = FunctionData.FunctionFlags;
 		return (Flags & FlagsToCheck) == ExpectedFlags;
 	}
@@ -2756,10 +2385,6 @@ public:
 	}
 
 protected:
-	/**
-	 * Create UObject engine types
-	 */
-	virtual void CreateUObjectEngineTypesInternal(ECreateEngineTypesPhase Phase) override;
 
 	/**
 	 * Perform any post parsing finalization and validation
@@ -2834,31 +2459,6 @@ public:
 
 	virtual uint32 GetHash(FUnrealTypeDefinitionInfo& ReferencingDef, bool bIncludeNoExport = true) const override;
 
-	/**
-	 * Return the Engine instance associated with the compiler instance
-	 */
-	UClass* GetClass() const
-	{
-		return static_cast<UClass*>(GetObject());
-	}
-
-	/**
-	 * Return the Engine instance associated with the compiler instance without null check
-	 */
-	UClass* GetClassSafe() const
-	{
-		return static_cast<UClass*>(GetObjectSafe());
-	}
-
-	/**
-	 * Set the Engine instance associated with the compiler instance
-	 */
-	virtual void SetObject(UObject* InObject) override
-	{
-		FUnrealStructDefinitionInfo::SetObject(InObject);
-		CastChecked<UClass>(InObject)->ClassFlags |= ClassFlags;
-	}
-
 	FUnrealClassDefinitionInfo* GetSuperClass() const;
 
 	/**
@@ -2906,9 +2506,6 @@ public:
 	 */
 	EClassFlags GetClassFlags() const
 	{
-#if UHT_ENABLE_ENGINE_TYPE_CHECKS
-		check(GetClassSafe() == nullptr || (GetClassSafe()->ClassFlags & CLASS_SaveInCompiledInClasses) == (ClassFlags & CLASS_SaveInCompiledInClasses));
-#endif
 		return ClassFlags;
 	}
 
@@ -2917,14 +2514,7 @@ public:
 	 */
 	void SetClassFlags(EClassFlags FlagsToSet)
 	{
-#if UHT_ENABLE_ENGINE_TYPE_CHECKS
-		check(GetClassSafe() == nullptr || (GetClassSafe()->ClassFlags & CLASS_SaveInCompiledInClasses) == (ClassFlags & CLASS_SaveInCompiledInClasses));
-#endif
 		ClassFlags |= FlagsToSet;
-		if (UClass* Class = GetClassSafe())
-		{
-			Class->ClassFlags |= FlagsToSet;
-		}
 	}
 
 	/**
@@ -2932,14 +2522,7 @@ public:
 	 */
 	void ClearClassFlags(EClassFlags FlagsToClear)
 	{
-#if UHT_ENABLE_ENGINE_TYPE_CHECKS
-		check(GetClassSafe() == nullptr || (GetClassSafe()->ClassFlags & CLASS_SaveInCompiledInClasses) == (ClassFlags & CLASS_SaveInCompiledInClasses));
-#endif
 		ClassFlags &= ~FlagsToClear;
-		if (UClass* Class = GetClassSafe())
-		{
-			Class->ClassFlags &= ~FlagsToClear;
-		}
 	}
 
 	/**
@@ -2951,9 +2534,6 @@ public:
 	 */
 	bool HasAnyClassFlags(EClassFlags FlagsToCheck) const
 	{
-#if UHT_ENABLE_ENGINE_TYPE_CHECKS
-		check(GetClassSafe() == nullptr || (GetClassSafe()->ClassFlags & CLASS_SaveInCompiledInClasses) == (ClassFlags & CLASS_SaveInCompiledInClasses));
-#endif
 		return EnumHasAnyFlags(ClassFlags, FlagsToCheck);
 	}
 
@@ -2965,9 +2545,6 @@ public:
 	 */
 	bool HasAllClassFlags(EClassFlags FlagsToCheck) const
 	{
-#if UHT_ENABLE_ENGINE_TYPE_CHECKS
-		check(GetClassSafe() == nullptr || (GetClassSafe()->ClassFlags & CLASS_SaveInCompiledInClasses) == (ClassFlags & CLASS_SaveInCompiledInClasses));
-#endif
 		return EnumHasAllFlags(ClassFlags, FlagsToCheck);
 	}
 
@@ -3013,9 +2590,6 @@ public:
 	 */
 	EClassCastFlags GetClassCastFlags() const
 	{
-#if UHT_ENABLE_ENGINE_TYPE_CHECKS
-		check(GetClassSafe() == nullptr || GetClassSafe()->ClassCastFlags == ClassCastFlags);
-#endif
 		return ClassCastFlags;
 	}
 
@@ -3024,14 +2598,7 @@ public:
 	 */
 	void SetClassCastFlags(EClassCastFlags FlagsToSet)
 	{
-#if UHT_ENABLE_ENGINE_TYPE_CHECKS
-		check(GetClassSafe() == nullptr || GetClassSafe()->ClassCastFlags == ClassCastFlags);
-#endif
 		ClassCastFlags |= FlagsToSet;
-		if (UClass* Class = GetClassSafe())
-		{
-			Class->ClassCastFlags |= FlagsToSet;
-		}
 	}
 
 	/**
@@ -3115,10 +2682,6 @@ public:
 	void SetClassWithin(FUnrealClassDefinitionInfo* InClassWithin)
 	{
 		ClassWithin = InClassWithin;
-		if (UClass* Class = GetClassSafe())
-		{
-			Class->ClassWithin = ClassWithin->GetClass();
-		}
 	}
 
 	/**
@@ -3158,9 +2721,6 @@ public:
 	 */
 	FName GetClassConfigName() const
 	{
-#if UHT_ENABLE_ENGINE_TYPE_CHECKS
-		check(GetClassSafe() == nullptr || GetClassSafe()->ClassConfigName == ClassConfigName);
-#endif
 		return ClassConfigName;
 	}
 
@@ -3169,14 +2729,7 @@ public:
 	 */
 	void SetClassConfigName(FName InClassConfigName)
 	{
-#if UHT_ENABLE_ENGINE_TYPE_CHECKS
-		check(GetClassSafe() == nullptr || GetClassSafe()->ClassConfigName == ClassConfigName);
-#endif
 		ClassConfigName = InClassConfigName;
-		if (UClass* Class = GetClassSafe())
-		{
-			Class->ClassConfigName = InClassConfigName;
-		}
 	}
 
 	/**
@@ -3184,9 +2737,6 @@ public:
 	 */
 	int32 GetPropertiesSize() const
 	{
-#if UHT_ENABLE_ENGINE_TYPE_CHECKS
-		check(GetClassSafe() == nullptr || GetClassSafe()->PropertiesSize == PropertiesSize);
-#endif
 		return PropertiesSize;
 	}
 
@@ -3195,14 +2745,7 @@ public:
 	 */
 	void SetPropertiesSize(int32 InPropertiesSize)
 	{
-#if UHT_ENABLE_ENGINE_TYPE_CHECKS
-		check(GetClassSafe() == nullptr || GetClassSafe()->PropertiesSize == PropertiesSize);
-#endif
 		PropertiesSize = InPropertiesSize;
-		if (UClass* Class = GetClassSafe())
-		{
-			Class->PropertiesSize = InPropertiesSize;
-		}
 	}
 
 	/**
@@ -3479,10 +3022,6 @@ public:
 	FString GetNameWithPrefix(EEnforceInterfacePrefix EnforceInterfacePrefix = EEnforceInterfacePrefix::None) const;
 
 protected:
-	/**
-	 * Create UObject engine types
-	 */
-	virtual void CreateUObjectEngineTypesInternal(ECreateEngineTypesPhase Phase) override;
 
 	/**
 	 * Perform any post parsing finalization and validation
diff --git a/Engine/Source/Programs/UnrealHeaderTool/Public/IScriptGeneratorPluginInterface.h b/Engine/Source/Programs/UnrealHeaderTool/Public/IScriptGeneratorPluginInterface.h
--- a/Engine/Source/Programs/UnrealHeaderTool/Public/IScriptGeneratorPluginInterface.h
+++ b/Engine/Source/Programs/UnrealHeaderTool/Public/IScriptGeneratorPluginInterface.h
@@ -96,7 +96,8 @@ struct EPackageOverrideType
 /**
  * The public interface to script generator plugins.
  */
-class IScriptGeneratorPluginInterface : public IModuleInterface, public IModularFeature
+class UE_DEPRECATED(5.2, "C++ UHT is no longer supported.  Please port any script generators to use C# UHT")
+IScriptGeneratorPluginInterface : public IModuleInterface, public IModularFeature
 {
 public:
 
diff --git a/Engine/Source/Programs/UnrealHeaderTool/UnrealHeaderTool.Target.cs b/Engine/Source/Programs/UnrealHeaderTool/UnrealHeaderTool.Target.cs
--- a/Engine/Source/Programs/UnrealHeaderTool/UnrealHeaderTool.Target.cs
+++ b/Engine/Source/Programs/UnrealHeaderTool/UnrealHeaderTool.Target.cs
@@ -34,7 +34,7 @@ public class UnrealHeaderToolTarget : TargetRules
 		// This app is a console application, not a Windows app (sets entry point to main(), instead of WinMain())
 		bIsBuildingConsoleApplication = true;
 
-        GlobalDefinitions.Add("HACK_HEADER_GENERATOR=1");
+		GlobalDefinitions.Add("USE_COMPILED_IN_NATIVES=0");
         GlobalDefinitions.Add("FNAME_WRITE_PROTECT_PAGES=0");
         GlobalDefinitions.Add("USE_LOCALIZED_PACKAGE_CACHE=0");
 		GlobalDefinitions.Add("STATS=0");
