Removed FNAME_WRITE_PROTECT_PAGES=0 from target definitions

#rnx

diff --git a/Engine/Source/Programs/UnrealHeaderTool/UnrealHeaderTool.Target.cs b/Engine/Source/Programs/UnrealHeaderTool/UnrealHeaderTool.Target.cs
--- a/Engine/Source/Programs/UnrealHeaderTool/UnrealHeaderTool.Target.cs
+++ b/Engine/Source/Programs/UnrealHeaderTool/UnrealHeaderTool.Target.cs
@@ -35,7 +35,6 @@ public class UnrealHeaderToolTarget : TargetRules
 		bIsBuildingConsoleApplication = true;
 
 		GlobalDefinitions.Add("USE_COMPILED_IN_NATIVES=0");
-        GlobalDefinitions.Add("FNAME_WRITE_PROTECT_PAGES=0");
         GlobalDefinitions.Add("USE_LOCALIZED_PACKAGE_CACHE=0");
 		GlobalDefinitions.Add("STATS=0");
 
