Updating C++ UHT with C# UHT changes to generated include files when it only contains enums.

#rnx

diff --git a/Engine/Source/Programs/UnrealHeaderTool/Private/CodeGenerator.cpp b/Engine/Source/Programs/UnrealHeaderTool/Private/CodeGenerator.cpp
--- a/Engine/Source/Programs/UnrealHeaderTool/Private/CodeGenerator.cpp
+++ b/Engine/Source/Programs/UnrealHeaderTool/Private/CodeGenerator.cpp
@@ -3263,8 +3263,17 @@ bool FNativeClassHeaderGenerator::WriteHeader(FGeneratedCPP& FileInfo, const FSt
 	FUHTStringBuilder GeneratedHeaderTextWithCopyright;
 	GeneratedHeaderTextWithCopyright.Log(HeaderCopyright);
 	GeneratedHeaderTextWithCopyright.Logf(TEXT("// IWYU pragma: private, include \"%s\"\r\n"), *FileInfo.SourceFile.GetIncludePath());
-	GeneratedHeaderTextWithCopyright.Log(TEXT("#include \"UObject/ObjectMacros.h\"\r\n"));
-	GeneratedHeaderTextWithCopyright.Log(TEXT("#include \"UObject/ScriptMacros.h\"\r\n"));
+	if (FileInfo.SourceFile.GetDefinedTypeCount() != 0 && FileInfo.SourceFile.GetDefinedEnumsCount() == FileInfo.SourceFile.GetDefinedTypeCount())
+	{
+		GeneratedHeaderTextWithCopyright.Log(TEXT("#include \"Templates/IsUEnumClass.h\"\r\n"));
+		GeneratedHeaderTextWithCopyright.Log(TEXT("#include \"UObject/ObjectMacros.h\"\r\n"));
+		GeneratedHeaderTextWithCopyright.Log(TEXT("#include \"UObject/ReflectedTypeAccessors.h\"\r\n"));
+	}
+	else
+	{
+		GeneratedHeaderTextWithCopyright.Log(TEXT("#include \"UObject/ObjectMacros.h\"\r\n"));
+		GeneratedHeaderTextWithCopyright.Log(TEXT("#include \"UObject/ScriptMacros.h\"\r\n"));
+	}
 
 	for (const FString& AdditionalHeader : InAdditionalHeaders)
 	{
