Fixed issue in C++ UHT with the C# UHT enum change.

#rnx

diff --git a/Engine/Source/Programs/UnrealHeaderTool/Private/HeaderParser.cpp b/Engine/Source/Programs/UnrealHeaderTool/Private/HeaderParser.cpp
--- a/Engine/Source/Programs/UnrealHeaderTool/Private/HeaderParser.cpp
+++ b/Engine/Source/Programs/UnrealHeaderTool/Private/HeaderParser.cpp
@@ -1166,7 +1166,9 @@ FUnrealEnumDefinitionInfo& FHeaderParser::CompileEnum()
 	if (CppForm == UEnum::ECppForm::EnumClass)
 	{
 		UnderlyingType = ParseUnderlyingEnumType();
-		if (UnderlyingType != EUnderlyingEnumType::uint8 && EnumMetaData.Contains(FHeaderParserNames::NAME_BlueprintType))
+		if (UnderlyingType != EUnderlyingEnumType::Unspecified &&
+			UnderlyingType != EUnderlyingEnumType::uint8 && 
+			EnumMetaData.Contains(FHeaderParserNames::NAME_BlueprintType))
 		{
 			Throwf(TEXT("Invalid BlueprintType enum base - currently only uint8 supported"));
 		}
