Modified UHT to not emit unsued #defines.

diff --git a/Engine/Source/Programs/Shared/EpicGames.UHT/Exporters/CodeGen/UhtHeaderCodeGeneratorHFile.cs b/Engine/Source/Programs/Shared/EpicGames.UHT/Exporters/CodeGen/UhtHeaderCodeGeneratorHFile.cs
--- a/Engine/Source/Programs/Shared/EpicGames.UHT/Exporters/CodeGen/UhtHeaderCodeGeneratorHFile.cs
+++ b/Engine/Source/Programs/Shared/EpicGames.UHT/Exporters/CodeGen/UhtHeaderCodeGeneratorHFile.cs
@@ -423,7 +423,15 @@ namespace EpicGames.UHT.Exporters.CodeGen
 		private StringBuilder AppendClass(StringBuilder builder, UhtClass classObj)
 		{
 			string api = classObj.ClassFlags.HasAnyFlags(EClassFlags.MinimalAPI) ? PackageApi : "NO_API ";
+			bool usesLegacy = classObj.ClassExportFlags.HasAnyFlags(UhtClassExportFlags.UsesGeneratedBodyLegacy);
 			UhtClass? nativeInterface = ObjectInfos[classObj.ObjectTypeIndex].NativeInterface;
+			if (classObj.ClassFlags.HasAnyFlags(EClassFlags.Interface) && nativeInterface == null)
+			{
+				throw new UhtIceException("Interfaces must have an associated native interface");
+			}
+
+			// With interfaces, there are cases where we use the native interface class export flags to check for legacy macro
+			bool alternateUsesLegacy = (classObj.ClassFlags.HasAnyFlags(EClassFlags.Interface) ? nativeInterface! : classObj).ClassExportFlags.HasAnyFlags(UhtClassExportFlags.UsesGeneratedBodyLegacy);
 
 			// Write the spare declarations
 			AppendSparseDeclarations(builder, classObj);
@@ -436,10 +444,10 @@ namespace EpicGames.UHT.Exporters.CodeGen
 			bool hasEditorRpc = reversedFunctions.Any(x => IsRpcFunction(x, true));
 
 			// Output the RPC methods
-			AppendRpcFunctions(builder, classObj, reversedFunctions, false);
+			AppendRpcFunctions(builder, classObj, alternateUsesLegacy, reversedFunctions, false);
 			if (hasEditorRpc)
 			{
-				AppendRpcFunctions(builder, classObj, reversedFunctions, true);
+				AppendRpcFunctions(builder, classObj, alternateUsesLegacy, reversedFunctions, true);
 			}
 
 			// Output property accessors
@@ -448,9 +456,7 @@ namespace EpicGames.UHT.Exporters.CodeGen
 			// Collect the callback function and sort by name to make the order stable
 			List<UhtFunction> callbackFunctions = new(classObj.Functions.Where(x => x.FunctionFlags.HasAnyFlags(EFunctionFlags.Event) && x.SuperFunction == null));
 			callbackFunctions.Sort((x, y) => StringComparerUE.OrdinalIgnoreCase.Compare(x.EngineName, y.EngineName));
-
-			// This is now done in.gen.cpp - Generate the callback parameter structures
-			//AppendCallbackParametersDecls(builder, classObj, callbackFunctions);
+			bool hasCallbacks = callbackFunctions.Count > 0;
 
 			// Generate the RPC wrappers for the callbacks
 			AppendCallbackRpcWrapperDecls(builder, classObj, callbackFunctions);
@@ -464,61 +470,95 @@ namespace EpicGames.UHT.Exporters.CodeGen
 
 			if (classObj.ClassFlags.HasAnyFlags(EClassFlags.Interface))
 			{
-				AppendStandardConstructors(builder, classObj, api);
-				AppendEnhancedConstructors(builder, classObj, api);
+				if (usesLegacy)
+				{
+					AppendStandardConstructors(builder, classObj, api);
+				}
+				else
+				{
+					AppendEnhancedConstructors(builder, classObj, api);
+				}
 
 				using (UhtMacroCreator macro = new(builder, this, classObj, GeneratedUInterfaceBodyMacroSuffix))
 				{
 					AppendCommonGeneratedBody(builder, classObj, api);
 				}
 
-				using (UhtMacroCreator macro = new(builder, this, classObj, GeneratedBodyLegacyMacroSuffix))
+				if (usesLegacy)
 				{
-					builder.Append('\t');
-					AppendGeneratedMacroDeprecationWarning(builder, "GENERATED_UINTERFACE_BODY");
-					builder.Append('\t').Append(DisableDeprecationWarnings).Append(" \\\r\n");
-					builder.Append('\t').AppendMacroName(this, classObj, GeneratedUInterfaceBodyMacroSuffix).Append(" \\\r\n");
-					builder.Append('\t').AppendMacroName(this, classObj, StandardConstructorsMacroSuffix).Append(" \\\r\n");
-					builder.Append('\t').Append(EnableDeprecationWarnings).Append(" \\\r\n");
+					using (UhtMacroCreator macro = new(builder, this, classObj, GeneratedBodyLegacyMacroSuffix))
+					{
+						builder.Append('\t');
+						AppendGeneratedMacroDeprecationWarning(builder, "GENERATED_UINTERFACE_BODY");
+						builder.Append('\t').Append(DisableDeprecationWarnings).Append(" \\\r\n");
+						builder.Append('\t').AppendMacroName(this, classObj, GeneratedUInterfaceBodyMacroSuffix).Append(" \\\r\n");
+						builder.Append('\t').AppendMacroName(this, classObj, StandardConstructorsMacroSuffix).Append(" \\\r\n");
+						builder.Append('\t').Append(EnableDeprecationWarnings).Append(" \\\r\n");
+					}
 				}
-
-				using (UhtMacroCreator macro = new(builder, this, classObj, GeneratedBodyMacroSuffix))
+				else
 				{
-					builder.Append('\t').Append(DisableDeprecationWarnings).Append(" \\\r\n");
-					builder.Append('\t').AppendMacroName(this, classObj, GeneratedUInterfaceBodyMacroSuffix).Append(" \\\r\n");
-					builder.Append('\t').AppendMacroName(this, classObj, EnchancedConstructorsMacroSuffix).Append(" \\\r\n");
-					AppendAccessSpecifier(builder, classObj);
-					builder.Append(" \\\r\n");
-					builder.Append('\t').Append(EnableDeprecationWarnings).Append(" \\\r\n");
+					using (UhtMacroCreator macro = new(builder, this, classObj, GeneratedBodyMacroSuffix))
+					{
+						builder.Append('\t').Append(DisableDeprecationWarnings).Append(" \\\r\n");
+						builder.Append('\t').AppendMacroName(this, classObj, GeneratedUInterfaceBodyMacroSuffix).Append(" \\\r\n");
+						builder.Append('\t').AppendMacroName(this, classObj, EnchancedConstructorsMacroSuffix).Append(" \\\r\n");
+						AppendAccessSpecifier(builder, classObj);
+						builder.Append(" \\\r\n");
+						builder.Append('\t').Append(EnableDeprecationWarnings).Append(" \\\r\n");
+					}
 				}
 
-				using (UhtMacroCreator macro = new(builder, this, classObj, InClassIInterfaceNoPureDeclsMacroSuffix))
+				if (alternateUsesLegacy)
 				{
-					AppendInClassIInterface(builder, classObj, callbackFunctions, api);
+					using (UhtMacroCreator macro = new(builder, this, classObj, InClassIInterfaceMacroSuffix))
+					{
+						AppendInClassIInterface(builder, classObj, callbackFunctions, api);
+					}
 				}
-
-				using (UhtMacroCreator macro = new(builder, this, classObj, InClassIInterfaceMacroSuffix))
+				else
 				{
-					AppendInClassIInterface(builder, classObj, callbackFunctions, api);
+					using (UhtMacroCreator macro = new(builder, this, classObj, InClassIInterfaceNoPureDeclsMacroSuffix))
+					{
+						AppendInClassIInterface(builder, classObj, callbackFunctions, api);
+					}
 				}
+
+				AppendProlog(builder, classObj);
+				AppendGeneratedBodyMacroBlock(builder, classObj, nativeInterface!, alternateUsesLegacy, hasEditorRpc, hasCallbacks, null);
 			}
 			else
 			{
-				using (UhtMacroCreator macro = new(builder, this, classObj, InClassNoPureDeclsMacroSuffix))
+				if (usesLegacy)
 				{
-					AppendClassGeneratedBody(builder, classObj, api);
+					using (UhtMacroCreator macro = new(builder, this, classObj, InClassMacroSuffix))
+					{
+						AppendClassGeneratedBody(builder, classObj, api);
+					}
+					AppendStandardConstructors(builder, classObj, api);
 				}
-
-				using (UhtMacroCreator macro = new(builder, this, classObj, InClassMacroSuffix))
+				else
 				{
-					AppendClassGeneratedBody(builder, classObj, api);
+					using (UhtMacroCreator macro = new(builder, this, classObj, InClassNoPureDeclsMacroSuffix))
+					{
+						AppendClassGeneratedBody(builder, classObj, api);
+					}
+					AppendEnhancedConstructors(builder, classObj, api);
 				}
 
-				AppendStandardConstructors(builder, classObj, api);
-				AppendEnhancedConstructors(builder, classObj, api);
 				AppendFieldNotify(builder, classObj);
+				AppendProlog(builder, classObj);
+				AppendGeneratedBodyMacroBlock(builder, classObj, classObj, usesLegacy, hasEditorRpc, hasCallbacks, usesLegacy ? "GENERATED_UCLASS_BODY" : null);
 			}
 
+			// Forward declare the StaticClass specialization in the header
+			builder.Append("template<> ").Append(PackageApi).Append("UClass* StaticClass<class ").Append(classObj.SourceName).Append(">();\r\n");
+			builder.Append("\r\n");
+			return builder;
+		}
+
+		private StringBuilder AppendProlog(StringBuilder builder, UhtClass classObj/*, List<UhtFunction> classbackFunctions*/)
+		{
 			using (UhtMacroCreator macro = new(builder, this, classObj.PrologLineNumber, PrologMacroSuffix))
 			{
 				//if (callbackFunctions.Count > 0)
@@ -526,25 +566,6 @@ namespace EpicGames.UHT.Exporters.CodeGen
 				//	builder.Append('\t').AppendMacroName(this, classObj, EventParamsMacroSuffix).Append(" \\\r\n");
 				//}
 			}
-
-			bool hasCallbacks = callbackFunctions.Count > 0;
-			if (classObj.ClassFlags.HasAnyFlags(EClassFlags.Interface))
-			{
-				if (nativeInterface != null)
-				{
-					AppendGeneratedBodyMacroBlock(builder, classObj, nativeInterface, true, hasEditorRpc, hasCallbacks, null);
-					AppendGeneratedBodyMacroBlock(builder, classObj, nativeInterface, false, hasEditorRpc, hasCallbacks, null);
-				}
-			}
-			else
-			{
-				AppendGeneratedBodyMacroBlock(builder, classObj, classObj, true, hasEditorRpc, hasCallbacks, "GENERATED_UCLASS_BODY");
-				AppendGeneratedBodyMacroBlock(builder, classObj, classObj, false, hasEditorRpc, hasCallbacks, null);
-			}
-
-			// Forward declare the StaticClass specialization in the header
-			builder.Append("template<> ").Append(PackageApi).Append("UClass* StaticClass<class ").Append(classObj.SourceName).Append(">();\r\n");
-			builder.Append("\r\n");
 			return builder;
 		}
 
@@ -689,34 +710,44 @@ namespace EpicGames.UHT.Exporters.CodeGen
 			return builder;
 		}
 
-		private StringBuilder AppendRpcFunctions(StringBuilder builder, UhtClass classObj, List<UhtFunction> reversedFunctions, bool editorOnly)
+		private StringBuilder AppendRpcFunctions(StringBuilder builder, UhtClass classObj, bool usesLegacy, List<UhtFunction> reversedFunctions, bool editorOnly)
 		{
 			builder.AppendBeginEditorOnlyGuard(editorOnly);
 
-			using (UhtMacroCreator macro = new(builder, this, classObj, editorOnly ? EditorOnlyRpcWrappersMacroSuffix : RpcWrappersMacroSuffix))
+			if (usesLegacy)
 			{
-				AppendAutogeneratedBlueprintFunctionDeclarations(builder, classObj, reversedFunctions, editorOnly);
-				AppendRpcWrappers(builder, reversedFunctions, editorOnly);
+				using (UhtMacroCreator macro = new(builder, this, classObj, editorOnly ? EditorOnlyRpcWrappersMacroSuffix : RpcWrappersMacroSuffix))
+				{
+					AppendAutogeneratedBlueprintFunctionDeclarations(builder, classObj, reversedFunctions, editorOnly);
+					AppendRpcWrappers(builder, reversedFunctions, editorOnly);
+				}
 			}
-
-			using (UhtMacroCreator macro = new(builder, this, classObj, editorOnly ? EditorOnlyRpcWrappersNoPureDeclsMacroSuffix : RpcWrappersNoPureDeclsMacroSuffix))
+			else
 			{
-				if (classObj.GeneratedCodeVersion <= EGeneratedCodeVersion.V1)
+				using (UhtMacroCreator macro = new(builder, this, classObj, editorOnly ? EditorOnlyRpcWrappersNoPureDeclsMacroSuffix : RpcWrappersNoPureDeclsMacroSuffix))
 				{
-					AppendAutogeneratedBlueprintFunctionDeclarationsOnlyNotDeclared(builder, classObj, reversedFunctions, editorOnly);
+					if (classObj.GeneratedCodeVersion <= EGeneratedCodeVersion.V1)
+					{
+						AppendAutogeneratedBlueprintFunctionDeclarationsOnlyNotDeclared(builder, classObj, reversedFunctions, editorOnly);
+					}
+					AppendRpcWrappers(builder, reversedFunctions, editorOnly);
 				}
-				AppendRpcWrappers(builder, reversedFunctions, editorOnly);
 			}
 
 			if (editorOnly)
 			{
 				builder.Append("#else\r\n");
-				using (UhtMacroCreator macro = new(builder, this, classObj, editorOnly ? EditorOnlyRpcWrappersMacroSuffix : RpcWrappersMacroSuffix))
+				if (usesLegacy)
 				{
+					using (UhtMacroCreator macro = new(builder, this, classObj, editorOnly ? EditorOnlyRpcWrappersMacroSuffix : RpcWrappersMacroSuffix))
+					{
+					}
 				}
-
-				using (UhtMacroCreator macro = new(builder, this, classObj, editorOnly ? EditorOnlyRpcWrappersNoPureDeclsMacroSuffix : RpcWrappersNoPureDeclsMacroSuffix))
+				else
 				{
+					using (UhtMacroCreator macro = new(builder, this, classObj, editorOnly ? EditorOnlyRpcWrappersNoPureDeclsMacroSuffix : RpcWrappersNoPureDeclsMacroSuffix))
+					{
+					}
 				}
 				builder.AppendEndEditorOnlyGuard();
 			}
diff --git a/Engine/Source/Programs/UnrealHeaderTool/Private/CodeGenerator.cpp b/Engine/Source/Programs/UnrealHeaderTool/Private/CodeGenerator.cpp
--- a/Engine/Source/Programs/UnrealHeaderTool/Private/CodeGenerator.cpp
+++ b/Engine/Source/Programs/UnrealHeaderTool/Private/CodeGenerator.cpp
@@ -2881,31 +2881,36 @@ void FNativeClassHeaderGenerator::ExportClassFromSourceFileInner(
 
 			const TCHAR* Offset = TEXT("\t");
 
-			OutGeneratedHeaderText.Log(
-				Macroize(
-					*SourceFile.GetGeneratedBodyMacroName(ClassGeneratedBodyLine, true),
-					FString::Printf(TEXT("\t%s\t%s\t%s()" LINE_TERMINATOR_ANSI "%s\t%s")
-						, *DeprecationWarning
-						, DisableDeprecationWarnings
-						, *InterfaceMacroName
-						, *StandardUObjectConstructorsMacroCall
-						, EnableDeprecationWarnings
+			if (ClassDef.UsesGeneratedBodyLegacy())
+			{
+				OutGeneratedHeaderText.Log(
+					Macroize(
+						*SourceFile.GetGeneratedBodyMacroName(ClassGeneratedBodyLine, true),
+						FString::Printf(TEXT("\t%s\t%s\t%s()" LINE_TERMINATOR_ANSI "%s\t%s")
+							, *DeprecationWarning
+							, DisableDeprecationWarnings
+							, *InterfaceMacroName
+							, *StandardUObjectConstructorsMacroCall
+							, EnableDeprecationWarnings
+						)
 					)
-				)
-			);
-
-			OutGeneratedHeaderText.Log(
-				Macroize(
-					*SourceFile.GetGeneratedBodyMacroName(ClassGeneratedBodyLine),
-					FString::Printf(TEXT("\t%s\t%s()" LINE_TERMINATOR_ANSI "%s%s\t%s")
-						, DisableDeprecationWarnings
-						, *InterfaceMacroName
-						, *EnhancedUObjectConstructorsMacroCall
-						, *GetPreservedAccessSpecifierString(ClassDef)
-						, EnableDeprecationWarnings
+				);
+			}
+			else
+			{
+				OutGeneratedHeaderText.Log(
+					Macroize(
+						*SourceFile.GetGeneratedBodyMacroName(ClassGeneratedBodyLine),
+						FString::Printf(TEXT("\t%s\t%s()" LINE_TERMINATOR_ANSI "%s%s\t%s")
+							, DisableDeprecationWarnings
+							, *InterfaceMacroName
+							, *EnhancedUObjectConstructorsMacroCall
+							, *GetPreservedAccessSpecifierString(ClassDef)
+							, EnableDeprecationWarnings
+						)
 					)
-				)
-			);
+				);
+			}
 
 			// =============================================
 			// Export the pure interface version of the class
@@ -2937,13 +2942,18 @@ void FNativeClassHeaderGenerator::ExportClassFromSourceFileInner(
 				WriteReplicatedMacroData(*ClassCPPName, *APIArg, ClassDef, InterfaceBoilerplate, SourceFile, OutFlags);
 			}
 
-			FString NoPureDeclsMacroName = SourceFile.GetGeneratedMacroName(ClassDef.GetGeneratedBodyLine(), TEXT("_INCLASS_IINTERFACE_NO_PURE_DECLS"));
-			WriteMacro(OutGeneratedHeaderText, NoPureDeclsMacroName, InterfaceBoilerplate);
-			ClassNoPureDeclsMacroCalls.Logf(TEXT("\t%s\r\n"), *NoPureDeclsMacroName);
-
-			FString MacroName = SourceFile.GetGeneratedMacroName(ClassDef.GetGeneratedBodyLine(), TEXT("_INCLASS_IINTERFACE"));
-			WriteMacro(OutGeneratedHeaderText, MacroName, InterfaceBoilerplate);
-			ClassMacroCalls.Logf(TEXT("\t%s\r\n"), *MacroName);
+			if (ClassDef.NativeInterfaceUsesGeneratedBodyLegacy())
+			{
+				FString MacroName = SourceFile.GetGeneratedMacroName(ClassDef.GetGeneratedBodyLine(), TEXT("_INCLASS_IINTERFACE"));
+				WriteMacro(OutGeneratedHeaderText, MacroName, InterfaceBoilerplate);
+				ClassMacroCalls.Logf(TEXT("\t%s\r\n"), *MacroName);
+			}
+			else
+			{
+				FString NoPureDeclsMacroName = SourceFile.GetGeneratedMacroName(ClassDef.GetGeneratedBodyLine(), TEXT("_INCLASS_IINTERFACE_NO_PURE_DECLS"));
+				WriteMacro(OutGeneratedHeaderText, NoPureDeclsMacroName, InterfaceBoilerplate);
+				ClassNoPureDeclsMacroCalls.Logf(TEXT("\t%s\r\n"), *NoPureDeclsMacroName);
+			}
 		}
 		else
 		{
@@ -2969,13 +2979,18 @@ void FNativeClassHeaderGenerator::ExportClassFromSourceFileInner(
 			}
 
 			{
-				FString NoPureDeclsMacroName = SourceFile.GetGeneratedMacroName(ClassDef.GetGeneratedBodyLine(), TEXT("_INCLASS_NO_PURE_DECLS"));
-				WriteMacro(OutGeneratedHeaderText, NoPureDeclsMacroName, Boilerplate);
-				ClassNoPureDeclsMacroCalls.Logf(TEXT("\t%s\r\n"), *NoPureDeclsMacroName);
-
-				FString MacroName = SourceFile.GetGeneratedMacroName(ClassDef.GetGeneratedBodyLine(), TEXT("_INCLASS"));
-				WriteMacro(OutGeneratedHeaderText, MacroName, Boilerplate);
-				ClassMacroCalls.Logf(TEXT("\t%s\r\n"), *MacroName);
+				if (ClassDef.UsesGeneratedBodyLegacy())
+				{
+					FString MacroName = SourceFile.GetGeneratedMacroName(ClassDef.GetGeneratedBodyLine(), TEXT("_INCLASS"));
+					WriteMacro(OutGeneratedHeaderText, MacroName, Boilerplate);
+					ClassMacroCalls.Logf(TEXT("\t%s\r\n"), *MacroName);
+				}
+				else
+				{
+					FString NoPureDeclsMacroName = SourceFile.GetGeneratedMacroName(ClassDef.GetGeneratedBodyLine(), TEXT("_INCLASS_NO_PURE_DECLS"));
+					WriteMacro(OutGeneratedHeaderText, NoPureDeclsMacroName, Boilerplate);
+					ClassNoPureDeclsMacroCalls.Logf(TEXT("\t%s\r\n"), *NoPureDeclsMacroName);
+				}
 
 				ExportConstructorsMacros(OutGeneratedHeaderText, OutCpp, StandardUObjectConstructorsMacroCall, EnhancedUObjectConstructorsMacroCall, SourceFile.GetGeneratedMacroName(ClassDef.GetGeneratedBodyLine()), ClassDef, *APIArg);
 				ExportFieldNotify(OutGeneratedHeaderText, OutCpp, StandardUObjectConstructorsMacroCall, EnhancedUObjectConstructorsMacroCall, SourceFile.GetGeneratedMacroName(ClassDef.GetGeneratedBodyLine()), ClassDef);
@@ -3019,8 +3034,14 @@ void FNativeClassHeaderGenerator::ExportClassFromSourceFileInner(
 		FString WrappedLegacyGeneratedBody = FString::Printf(TEXT("%s%s%s%s%s%s"), *DeprecationWarning, DisableDeprecationWarnings, Public, *LegacyGeneratedBody, Public, EnableDeprecationWarnings);
 		FString WrappedGeneratedBody = FString::Printf(TEXT("%s%s%s%s%s"), DisableDeprecationWarnings, Public, *GeneratedBody, *GetPreservedAccessSpecifierString(ClassDef), EnableDeprecationWarnings);
 
-		OutGeneratedHeaderText.Log(Macroize(*SourceFile.GetGeneratedBodyMacroName(GeneratedBodyLine, true), MoveTemp(WrappedLegacyGeneratedBody)));
-		OutGeneratedHeaderText.Log(Macroize(*SourceFile.GetGeneratedBodyMacroName(GeneratedBodyLine, false), MoveTemp(WrappedGeneratedBody)));
+		if (ClassDef.IsInterface() ? ClassDef.NativeInterfaceUsesGeneratedBodyLegacy() : ClassDef.UsesGeneratedBodyLegacy())
+		{
+			OutGeneratedHeaderText.Log(Macroize(*SourceFile.GetGeneratedBodyMacroName(GeneratedBodyLine, true), MoveTemp(WrappedLegacyGeneratedBody)));
+		}
+		else
+		{
+			OutGeneratedHeaderText.Log(Macroize(*SourceFile.GetGeneratedBodyMacroName(GeneratedBodyLine, false), MoveTemp(WrappedGeneratedBody)));
+		}
 	}
 
 	// Forward declare the StaticClass specialisation in the header
@@ -3251,8 +3272,14 @@ void FNativeClassHeaderGenerator::ExportConstructorsMacros(FOutputDevice& OutGen
 		OutGeneratedCppText.Logf(TEXT("\t%s::~%s() {}" LINE_TERMINATOR_ANSI), *ClassCPPName, *ClassCPPName);
 	}
 
-	OutGeneratedHeaderText.Log(Macroize(*StdMacroName, *StdMacro));
-	OutGeneratedHeaderText.Log(Macroize(*EnhMacroName, *EnhMacro));
+	if (ClassDef.UsesGeneratedBodyLegacy())
+	{
+		OutGeneratedHeaderText.Log(Macroize(*StdMacroName, *StdMacro));
+	}
+	else
+	{
+		OutGeneratedHeaderText.Log(Macroize(*EnhMacroName, *EnhMacro));
+	}
 
 	StandardUObjectConstructorsMacroCall.Logf(TEXT("\t%s\r\n"), *StdMacroName);
 	EnhancedUObjectConstructorsMacroCall.Logf(TEXT("\t%s\r\n"), *EnhMacroName);
@@ -5095,9 +5122,10 @@ void FNativeClassHeaderGenerator::ExportNativeFunctions(FOutputDevice& OutGenera
 		OutNoPureDeclsMacroCalls.Logf(TEXT("\t%s\r\n"), *MacroName);
 	}
 
+	bool bRpcWrappersUsesLegacy = ClassDef.HasAnyClassFlags(CLASS_Interface) ? ClassDef.NativeInterfaceUsesGeneratedBodyLegacy() : ClassDef.UsesGeneratedBodyLegacy();
+
 	// Write runtime wrappers
 	{
-		FString MacroName = SourceFile.GetGeneratedMacroName(ClassDef.GetGeneratedBodyLine(), TEXT("_RPC_WRAPPERS"));
 
 		// WriteMacro has an assumption about what will be at the end of this block that is no longer true due to splitting the
 		// definition and implementation, so add on a line terminator to satisfy it
@@ -5106,21 +5134,27 @@ void FNativeClassHeaderGenerator::ExportNativeFunctions(FOutputDevice& OutGenera
 			RuntimeStringBuilders.RPCWrappers += LINE_TERMINATOR;
 		}
 
-		WriteMacro(OutGeneratedHeaderText, MacroName, RuntimeStringBuilders.AutogeneratedBlueprintFunctionDeclarations + RuntimeStringBuilders.RPCWrappers);
-		OutMacroCalls.Logf(TEXT("\t%s\r\n"), *MacroName);
-
-		// Put static checks before RPCWrappers to get proper messages from static asserts before compiler errors.
-		FString NoPureDeclsMacroName = SourceFile.GetGeneratedMacroName(ClassDef.GetGeneratedBodyLine(), TEXT("_RPC_WRAPPERS_NO_PURE_DECLS"));
-		if (ClassDef.GetGeneratedCodeVersion() > EGeneratedCodeVersion::V1)
+		if (bRpcWrappersUsesLegacy)
 		{
-			WriteMacro(OutGeneratedHeaderText, NoPureDeclsMacroName, RuntimeStringBuilders.RPCWrappers);
+			FString MacroName = SourceFile.GetGeneratedMacroName(ClassDef.GetGeneratedBodyLine(), TEXT("_RPC_WRAPPERS"));
+			WriteMacro(OutGeneratedHeaderText, MacroName, RuntimeStringBuilders.AutogeneratedBlueprintFunctionDeclarations + RuntimeStringBuilders.RPCWrappers);
+			OutMacroCalls.Logf(TEXT("\t%s\r\n"), *MacroName);
 		}
 		else
 		{
-			WriteMacro(OutGeneratedHeaderText, NoPureDeclsMacroName, RuntimeStringBuilders.AutogeneratedBlueprintFunctionDeclarationsOnlyNotDeclared + RuntimeStringBuilders.RPCWrappers);
-		}
+			// Put static checks before RPCWrappers to get proper messages from static asserts before compiler errors.
+			FString NoPureDeclsMacroName = SourceFile.GetGeneratedMacroName(ClassDef.GetGeneratedBodyLine(), TEXT("_RPC_WRAPPERS_NO_PURE_DECLS"));
+			if (ClassDef.GetGeneratedCodeVersion() > EGeneratedCodeVersion::V1)
+			{
+				WriteMacro(OutGeneratedHeaderText, NoPureDeclsMacroName, RuntimeStringBuilders.RPCWrappers);
+			}
+			else
+			{
+				WriteMacro(OutGeneratedHeaderText, NoPureDeclsMacroName, RuntimeStringBuilders.AutogeneratedBlueprintFunctionDeclarationsOnlyNotDeclared + RuntimeStringBuilders.RPCWrappers);
+			}
 
-		OutNoPureDeclsMacroCalls.Logf(TEXT("\t%s\r\n"), *NoPureDeclsMacroName);
+			OutNoPureDeclsMacroCalls.Logf(TEXT("\t%s\r\n"), *NoPureDeclsMacroName);
+		}
 
 		OutGeneratedCPPText.Log(RuntimeStringBuilders.RPCImplementations);
 	}
@@ -5130,8 +5164,6 @@ void FNativeClassHeaderGenerator::ExportNativeFunctions(FOutputDevice& OutGenera
 	{
 		OutGeneratedHeaderText.Log( BeginEditorOnlyGuard );
 
-		FString MacroName = SourceFile.GetGeneratedMacroName(ClassDef.GetGeneratedBodyLine(), TEXT("_EDITOR_ONLY_RPC_WRAPPERS"));
-
 		// WriteMacro has an assumption about what will be at the end of this block that is no longer true due to splitting the
 		// definition and implementation, so add on a line terminator to satisfy it
 		if (EditorStringBuilders.RPCWrappers.Len() > 0)
@@ -5139,26 +5171,39 @@ void FNativeClassHeaderGenerator::ExportNativeFunctions(FOutputDevice& OutGenera
 			EditorStringBuilders.RPCWrappers += LINE_TERMINATOR;
 		}
 
-		WriteMacro(OutGeneratedHeaderText, MacroName, EditorStringBuilders.AutogeneratedBlueprintFunctionDeclarations + EditorStringBuilders.RPCWrappers);
-		OutMacroCalls.Logf(TEXT("\t%s\r\n"), *MacroName);
-
-		// Put static checks before RPCWrappers to get proper messages from static asserts before compiler errors.
+		FString MacroName = SourceFile.GetGeneratedMacroName(ClassDef.GetGeneratedBodyLine(), TEXT("_EDITOR_ONLY_RPC_WRAPPERS"));
 		FString NoPureDeclsMacroName = SourceFile.GetGeneratedMacroName(ClassDef.GetGeneratedBodyLine(), TEXT("_EDITOR_ONLY_RPC_WRAPPERS_NO_PURE_DECLS"));
-		if (ClassDef.GetGeneratedCodeVersion() > EGeneratedCodeVersion::V1)
+		if (bRpcWrappersUsesLegacy)
 		{
-			WriteMacro(OutGeneratedHeaderText, NoPureDeclsMacroName, EditorStringBuilders.RPCWrappers);
+
+			WriteMacro(OutGeneratedHeaderText, MacroName, EditorStringBuilders.AutogeneratedBlueprintFunctionDeclarations + EditorStringBuilders.RPCWrappers);
+			OutMacroCalls.Logf(TEXT("\t%s\r\n"), *MacroName);
 		}
 		else
 		{
-			WriteMacro(OutGeneratedHeaderText, NoPureDeclsMacroName, EditorStringBuilders.AutogeneratedBlueprintFunctionDeclarationsOnlyNotDeclared + EditorStringBuilders.RPCWrappers);
+			// Put static checks before RPCWrappers to get proper messages from static asserts before compiler errors.
+			if (ClassDef.GetGeneratedCodeVersion() > EGeneratedCodeVersion::V1)
+			{
+				WriteMacro(OutGeneratedHeaderText, NoPureDeclsMacroName, EditorStringBuilders.RPCWrappers);
+			}
+			else
+			{
+				WriteMacro(OutGeneratedHeaderText, NoPureDeclsMacroName, EditorStringBuilders.AutogeneratedBlueprintFunctionDeclarationsOnlyNotDeclared + EditorStringBuilders.RPCWrappers);
+			}
 		}
 
 		// write out an else preprocessor block for when not compiling for the editor.  The generated macros should be empty then since the functions are compiled out
 		{
 			OutGeneratedHeaderText.Log(TEXT("#else\r\n"));
 
-			WriteMacro(OutGeneratedHeaderText, MacroName, TEXT(""));
-			WriteMacro(OutGeneratedHeaderText, NoPureDeclsMacroName, TEXT(""));
+			if (bRpcWrappersUsesLegacy)
+			{
+				WriteMacro(OutGeneratedHeaderText, MacroName, TEXT(""));
+			}
+			else
+			{
+				WriteMacro(OutGeneratedHeaderText, NoPureDeclsMacroName, TEXT(""));
+			}
 
 			OutGeneratedHeaderText.Log(EndEditorOnlyGuard);
 		}
diff --git a/Engine/Source/Programs/UnrealHeaderTool/Private/HeaderParser.cpp b/Engine/Source/Programs/UnrealHeaderTool/Private/HeaderParser.cpp
--- a/Engine/Source/Programs/UnrealHeaderTool/Private/HeaderParser.cpp
+++ b/Engine/Source/Programs/UnrealHeaderTool/Private/HeaderParser.cpp
@@ -5115,7 +5115,7 @@ bool FHeaderParser::CompileDeclaration(TArray<FUnrealFunctionDefinitionInfo*>& D
 		if (Token.IsIdentifier(TEXT("GENERATED_IINTERFACE_BODY"), ESearchCase::CaseSensitive))
 		{
 			// Due to the way IInterface and UInterface share the same class instance, we don't set the legacy body
-			//ClassDef.MarkUsesGeneratedBodyLegacy();
+			ClassDef.MarkNativeInterfaceUsesGeneratedBodyLegacy();
 			CurrentAccessSpecifier = ACCESS_Public;
 		}
 
diff --git a/Engine/Source/Programs/UnrealHeaderTool/Private/UnrealTypeDefinitionInfo.h b/Engine/Source/Programs/UnrealHeaderTool/Private/UnrealTypeDefinitionInfo.h
--- a/Engine/Source/Programs/UnrealHeaderTool/Private/UnrealTypeDefinitionInfo.h
+++ b/Engine/Source/Programs/UnrealHeaderTool/Private/UnrealTypeDefinitionInfo.h
@@ -2982,6 +2982,22 @@ public:
 		bUsesGeneratedBodyLegacy = true;
 	}
 
+	/**
+	 * Return true if the class uses legacy generated body
+	 */
+	bool NativeInterfaceUsesGeneratedBodyLegacy() const
+	{
+		return bNativeInterfaceUsesGeneratedBodyLegacy;
+	}
+
+	/**
+	 * Mark the class as using legacy generated body
+	 */
+	void MarkNativeInterfaceUsesGeneratedBodyLegacy()
+	{
+		bNativeInterfaceUsesGeneratedBodyLegacy = true;
+	}
+
 	/**
 	 * Get the generated body access specifier
 	 */
@@ -3112,6 +3128,9 @@ private:
 	/** True if the class uses legacy generated body */
 	bool bUsesGeneratedBodyLegacy = false;
 
+	/** True if the class uses legacy generated body */
+	bool bNativeInterfaceUsesGeneratedBodyLegacy = false;
+
 	/** True if the class is an interface */
 	bool bIsInterface = false;
 
