Added ability in UHT for arrays to be sized using static_cast.

#9789

diff --git a/Engine/Source/Programs/Shared/EpicGames.UHT/Types/UhtProperty.cs b/Engine/Source/Programs/Shared/EpicGames.UHT/Types/UhtProperty.cs
--- a/Engine/Source/Programs/Shared/EpicGames.UHT/Types/UhtProperty.cs
+++ b/Engine/Source/Programs/Shared/EpicGames.UHT/Types/UhtProperty.cs
@@ -846,6 +846,17 @@ namespace EpicGames.UHT.Types
 			"(signed)",
 			"(unsigned int)",
 			"(signed int)",
+			"static_cast<uint32>",
+			"static_cast<int32>",
+			"static_cast<uint16>",
+			"static_cast<int16>",
+			"static_cast<uint8>",
+			"static_cast<int8>",
+			"static_cast<int>",
+			"static_cast<unsigned>",
+			"static_cast<signed>",
+			"static_cast<unsigned int>",
+			"static_cast<signed int>",
 		};
 
 		/// <summary>
diff --git a/Engine/Source/Programs/UnrealHeaderTool/Private/PropertyTypes.cpp b/Engine/Source/Programs/UnrealHeaderTool/Private/PropertyTypes.cpp
--- a/Engine/Source/Programs/UnrealHeaderTool/Private/PropertyTypes.cpp
+++ b/Engine/Source/Programs/UnrealHeaderTool/Private/PropertyTypes.cpp
@@ -2259,7 +2259,18 @@ struct FPropertyTypeTraitsStaticArray : public FPropertyTypeTraitsBase
 					TEXT("(unsigned)"),
 					TEXT("(signed)"),
 					TEXT("(unsigned int)"),
-					TEXT("(signed int)")
+					TEXT("(signed int)"),
+					TEXT("static_cast<uint32>"),
+					TEXT("static_cast<int32>"),
+					TEXT("static_cast<uint16>"),
+					TEXT("static_cast<int16>"),
+					TEXT("static_cast<uint8>"),
+					TEXT("static_cast<int8>"),
+					TEXT("static_cast<int>"),
+					TEXT("static_cast<unsigned>"),
+					TEXT("static_cast<signed>"),
+					TEXT("static_cast<unsigned int>"),
+					TEXT("static_cast<signed int>")
 				};
 				
 				// Remove any irrelevant whitespace
