Add support for parsing nested templates in getter/setter methods in C++ UHT.

diff --git a/Engine/Source/Programs/UnrealHeaderTool/Private/HeaderParser.cpp b/Engine/Source/Programs/UnrealHeaderTool/Private/HeaderParser.cpp
--- a/Engine/Source/Programs/UnrealHeaderTool/Private/HeaderParser.cpp
+++ b/Engine/Source/Programs/UnrealHeaderTool/Private/HeaderParser.cpp
@@ -5720,9 +5720,26 @@ static bool ParseAccessorType(FString& OutAccessorType, FTokenReplay& Tokens)
 			{
 				TemplateNestCount--;
 			}
+			else if (Token.IsSymbol(TEXT(">>"), ESearchCase::CaseSensitive))
+			{
+				if (TemplateNestCount < 2)
+				{
+					return false;
+				}
+				TemplateNestCount -= 2;
+			}
+			else if (Token.IsSymbol(TEXT(">>>"), ESearchCase::CaseSensitive))
+			{
+				if (TemplateNestCount < 3)
+				{
+					return false;
+				}
+				TemplateNestCount -= 3;
+			}
 		}
 		Tokens.GetToken(Token);
 	}
+
 	// Skip '&', 'const' and retain '*'
 	do
 	{
