Updating C++ UHT to match recent changes to C# UHT

#rnx

diff --git a/Engine/Source/Programs/UnrealHeaderTool/Private/CodeGenerator.cpp b/Engine/Source/Programs/UnrealHeaderTool/Private/CodeGenerator.cpp
--- a/Engine/Source/Programs/UnrealHeaderTool/Private/CodeGenerator.cpp
+++ b/Engine/Source/Programs/UnrealHeaderTool/Private/CodeGenerator.cpp
@@ -118,6 +118,8 @@
 
 namespace
 {
+	static const FName NAME_Comment(TEXT("Comment"));
+	static const FName NAME_ToolTip(TEXT("ToolTip"));
 	static const FName NAME_SerializeToFArchive("SerializeToFArchive");
 	static const FName NAME_SerializeToFStructuredArchive("SerializeToFStructuredArchive");
 	static const FName NAME_ObjectInitializerConstructorDeclared("ObjectInitializerConstructorDeclared");
@@ -640,19 +642,29 @@
 		Out.Log (TEXT("#if WITH_METADATA\r\n"));
 		Out.Logf(TEXT("%sconst UECodeGen_Private::FMetaDataPairParam %s[] = {\r\n"), Spaces, MetaDataBlockName);
 
+		bool bIsPartOfEngine = TypeDef.GetUnrealSourceFile().GetPackageDef().IsPartOfEngine();
 		for (const KVPType& KVP : KVPs)
 		{
+			bool bRestricted = !bIsPartOfEngine && (KVP.Key == NAME_Comment || KVP.Key == NAME_ToolTip);
+			if (bRestricted)
+			{
+				Out.Log(TEXT("#if !UE_BUILD_SHIPPING\r\n"));
+			}
 			Out.Logf(TEXT("%s\t{ %s, %s },\r\n"), Spaces, *CreateUTF8LiteralString(KVP.Key.ToString()), *CreateUTF8LiteralString(*KVP.Value));
+			if (bRestricted)
+			{
+				Out.Log(TEXT("#endif\r\n"));
+			}
 		}
 
 		Out.Logf(TEXT("%s};\r\n"), Spaces);
 		Out.Log (TEXT("#endif\r\n"));
 
-		Result = FString::Printf(TEXT("METADATA_PARAMS(%s, UE_ARRAY_COUNT(%s))"), MetaDataBlockName, MetaDataBlockName);
+		Result = FString::Printf(TEXT("METADATA_PARAMS(UE_ARRAY_COUNT(%s), %s)"), MetaDataBlockName, MetaDataBlockName);
 	}
 	else
 	{
-		Result = TEXT("METADATA_PARAMS(nullptr, 0)");
+		Result = TEXT("METADATA_PARAMS(0, nullptr)");
 	}
 
 	return Result;
@@ -823,9 +835,9 @@
 			*PropNotifyFunc,
 			PropFlags,
 			FPropertyObjectFlags,
-			*ArrayDim,
 			*SetterFunc,
 			*GetterFunc,
+			*ArrayDim,
 			OffsetStr,
 			*GetSingletonNameFuncAddr(PropertyBase.EnumDef, OutReferenceGatherers.UniqueCrossModuleReferences),
 			*MetaDataParams,
@@ -855,9 +867,9 @@
 			*PropNotifyFunc,
 			PropFlags,
 			FPropertyObjectFlags,
-			*ArrayDim,
 			*SetterFunc,
 			*GetterFunc,
+			*ArrayDim,
 			OffsetStr,
 			*MetaDataParams,
 			*PropTag
@@ -879,9 +891,9 @@
 			*PropNotifyFunc,
 			PropFlags,
 			FPropertyObjectFlags,
-			*ArrayDim,
 			*SetterFunc,
 			*GetterFunc,
+			*ArrayDim,
 			OffsetStr,
 			*MetaDataParams,
 			*PropTag
@@ -906,9 +918,9 @@
 			*PropNotifyFunc,
 			PropFlags,
 			FPropertyObjectFlags,
-			*ArrayDim,
 			*SetterFunc,
 			*GetterFunc,
+			*ArrayDim,
 			OffsetStr,
 			*MetaDataParams,
 			*PropTag
@@ -930,9 +942,9 @@
 			*PropNotifyFunc,
 			PropFlags,
 			FPropertyObjectFlags,
-			*ArrayDim,
 			*SetterFunc,
 			*GetterFunc,
+			*ArrayDim,
 			OffsetStr,
 			*MetaDataParams,
 			*PropTag
@@ -954,9 +966,9 @@
 			*PropNotifyFunc,
 			PropFlags,
 			FPropertyObjectFlags,
-			*ArrayDim,
 			*SetterFunc,
 			*GetterFunc,
+			*ArrayDim,
 			OffsetStr,
 			*MetaDataParams,
 			*PropTag
@@ -981,9 +993,9 @@
 			*PropNotifyFunc,
 			PropFlags,
 			FPropertyObjectFlags,
-			*ArrayDim,
 			*SetterFunc,
 			*GetterFunc,
+			*ArrayDim,
 			OffsetStr,
 			*MetaDataParams,
 			*PropTag
@@ -1005,9 +1017,9 @@
 			*PropNotifyFunc,
 			PropFlags,
 			FPropertyObjectFlags,
-			*ArrayDim,
 			*SetterFunc,
 			*GetterFunc,
+			*ArrayDim,
 			OffsetStr,
 			*MetaDataParams,
 			*PropTag
@@ -1029,9 +1041,9 @@
 			*PropNotifyFunc,
 			PropFlags,
 			FPropertyObjectFlags,
-			*ArrayDim,
 			*SetterFunc,
 			*GetterFunc,
+			*ArrayDim,
 			OffsetStr,
 			*MetaDataParams,
 			*PropTag
@@ -1053,9 +1065,9 @@
 			*PropNotifyFunc,
 			PropFlags,
 			FPropertyObjectFlags,
-			*ArrayDim,
 			*SetterFunc,
 			*GetterFunc,
+			*ArrayDim,
 			OffsetStr,
 			*MetaDataParams,
 			*PropTag
@@ -1077,9 +1089,9 @@
 			*PropNotifyFunc,
 			PropFlags,
 			FPropertyObjectFlags,
-			*ArrayDim,
 			*SetterFunc,
 			*GetterFunc,
+			*ArrayDim,
 			OffsetStr,
 			*MetaDataParams,
 			*PropTag
@@ -1128,9 +1140,9 @@
 			PropFlags,
 			bIsNativeBool ? TEXT("| UECodeGen_Private::EPropertyGenFlags::NativeBool") : TEXT(""),
 			FPropertyObjectFlags,
-			*ArrayDim,
 			*SetterFunc,
 			*GetterFunc,
+			*ArrayDim,
 			*PropertyDef.GetCPPType(nullptr, 0),
 			*OuterSize,
 			*Setter,
@@ -1156,9 +1168,9 @@
 				*PropNotifyFunc,
 				PropFlags,
 				FPropertyObjectFlags,
-				*ArrayDim,
 				*SetterFunc,
 				*GetterFunc,
+				*ArrayDim,
 				OffsetStr,
 				*GetSingletonNameFuncAddr(PropertyBase.ClassDef, OutReferenceGatherers.UniqueCrossModuleReferences, false),
 				*GetSingletonNameFuncAddr(PropertyBase.MetaClassDef, OutReferenceGatherers.UniqueCrossModuleReferences, false),				
@@ -1178,9 +1190,9 @@
 				*PropNotifyFunc,
 				PropFlags,
 				FPropertyObjectFlags,
-				*ArrayDim,
 				*SetterFunc,
 				*GetterFunc,
+				*ArrayDim,
 				OffsetStr,
 				*GetSingletonNameFuncAddr(PropertyBase.ClassDef, OutReferenceGatherers.UniqueCrossModuleReferences, false),
 				*MetaDataParams,
@@ -1206,9 +1218,9 @@
 				*PropNotifyFunc,
 				PropFlags,
 				FPropertyObjectFlags,
-				*ArrayDim,
 				*SetterFunc,
 				*GetterFunc,
+				*ArrayDim,
 				OffsetStr,
 				*GetSingletonNameFuncAddr(PropertyBase.ClassDef, OutReferenceGatherers.UniqueCrossModuleReferences, false),
 				*GetSingletonNameFuncAddr(PropertyBase.MetaClassDef, OutReferenceGatherers.UniqueCrossModuleReferences, false),
@@ -1228,9 +1240,9 @@
 				*PropNotifyFunc,
 				PropFlags,
 				FPropertyObjectFlags,
-				*ArrayDim,
 				*SetterFunc,
 				*GetterFunc,
+				*ArrayDim,
 				OffsetStr,
 				*GetSingletonNameFuncAddr(PropertyBase.ClassDef, OutReferenceGatherers.UniqueCrossModuleReferences, false),
 				*MetaDataParams,
@@ -1256,9 +1268,9 @@
 				*PropNotifyFunc,
 				PropFlags,
 				FPropertyObjectFlags,
-				*ArrayDim,
 				*SetterFunc,
 				*GetterFunc,
+				*ArrayDim,
 				OffsetStr,
 				*GetSingletonNameFuncAddr(PropertyBase.MetaClassDef, OutReferenceGatherers.UniqueCrossModuleReferences, false),
 				*MetaDataParams,
@@ -1277,9 +1289,9 @@
 				*PropNotifyFunc,
 				PropFlags,
 				FPropertyObjectFlags,
-				*ArrayDim,
 				*SetterFunc,
 				*GetterFunc,
+				*ArrayDim,
 				OffsetStr,
 				*GetSingletonNameFuncAddr(PropertyBase.ClassDef, OutReferenceGatherers.UniqueCrossModuleReferences, false),
 				*MetaDataParams,
@@ -1303,9 +1315,9 @@
 			*PropNotifyFunc,
 			PropFlags,
 			FPropertyObjectFlags,
-			*ArrayDim,
 			*SetterFunc,
 			*GetterFunc,
+			*ArrayDim,
 			OffsetStr,
 			*GetSingletonNameFuncAddr(PropertyBase.ClassDef, OutReferenceGatherers.UniqueCrossModuleReferences, false),
 			*MetaDataParams,
@@ -1328,9 +1340,9 @@
 			*PropNotifyFunc,
 			PropFlags,
 			FPropertyObjectFlags,
-			*ArrayDim,
 			*SetterFunc,
 			*GetterFunc,
+			*ArrayDim,
 			OffsetStr,
 			*GetSingletonNameFuncAddr(PropertyBase.ClassDef, OutReferenceGatherers.UniqueCrossModuleReferences, false),
 			*MetaDataParams,
@@ -1353,9 +1365,9 @@
 			*PropNotifyFunc,
 			PropFlags,
 			FPropertyObjectFlags,
-			*ArrayDim,
 			*SetterFunc,
 			*GetterFunc,
+			*ArrayDim,
 			OffsetStr,
 			*GetSingletonNameFuncAddr(PropertyBase.ClassDef, OutReferenceGatherers.UniqueCrossModuleReferences, false),
 			*MetaDataParams,
@@ -1378,9 +1390,9 @@
 			*PropNotifyFunc,
 			PropFlags,
 			FPropertyObjectFlags,
-			*ArrayDim,
 			*SetterFunc,
 			*GetterFunc,
+			*ArrayDim,
 			OffsetStr,
 			*MetaDataParams,
 			*PropTag
@@ -1402,9 +1414,9 @@
 			*PropNotifyFunc,
 			PropFlags,
 			FPropertyObjectFlags,
-			*ArrayDim,
 			*SetterFunc,
 			*GetterFunc,
+			*ArrayDim,
 			OffsetStr,
 			*MetaDataParams,
 			*PropTag
@@ -1432,9 +1444,9 @@
 			*PropNotifyFunc,
 			PropFlags,
 			FPropertyObjectFlags,
-			*ArrayDim,
 			*SetterFunc,
 			*GetterFunc,
+			*ArrayDim,
 			OffsetStr,
 			PropertyDef.GetAllocatorType() == EAllocatorType::MemoryImage ? TEXT("EArrayPropertyFlags::UsesMemoryImageAllocator") : TEXT("EArrayPropertyFlags::None"),
 			*MetaDataParams,
@@ -1466,9 +1478,9 @@
 			*PropNotifyFunc,
 			PropFlags,
 			FPropertyObjectFlags,
-			*ArrayDim,
 			*SetterFunc,
 			*GetterFunc,
+			*ArrayDim,
 			OffsetStr,
 			PropertyDef.GetAllocatorType() == EAllocatorType::MemoryImage ? TEXT("EMapPropertyFlags::UsesMemoryImageAllocator") : TEXT("EMapPropertyFlags::None"),
 			*MetaDataParams,
@@ -1503,9 +1515,9 @@
 			*PropNotifyFunc,
 			PropFlags,
 			FPropertyObjectFlags,
-			*ArrayDim,
 			*SetterFunc,
 			*GetterFunc,
+			*ArrayDim,
 			OffsetStr,
 			*MetaDataParams,
 			*PropTag
@@ -1527,9 +1539,9 @@
 			*PropNotifyFunc,
 			PropFlags,
 			FPropertyObjectFlags,
-			*ArrayDim,
 			*SetterFunc,
 			*GetterFunc,
+			*ArrayDim,
 			OffsetStr,
 			*GetSingletonNameFuncAddr(PropertyBase.ScriptStructDef, OutReferenceGatherers.UniqueCrossModuleReferences),
 			*MetaDataParams,
@@ -1552,9 +1564,9 @@
 			*PropNotifyFunc,
 			PropFlags,
 			FPropertyObjectFlags,
-			*ArrayDim,
 			*SetterFunc,
 			*GetterFunc,
+			*ArrayDim,
 			OffsetStr,
 			*GetSingletonNameFuncAddr(PropertyBase.FunctionDef, OutReferenceGatherers.UniqueCrossModuleReferences),
 			*MetaDataParams,
@@ -1580,9 +1592,9 @@
 			PropFlags,
 			!bIsSparse ? TEXT("Inline") : TEXT("Sparse"),
 			FPropertyObjectFlags,
-			*ArrayDim,
 			*SetterFunc,
 			*GetterFunc,
+			*ArrayDim,
 			OffsetStr,
 			*GetSingletonNameFuncAddr(PropertyBase.FunctionDef, OutReferenceGatherers.UniqueCrossModuleReferences),
 			*MetaDataParams,
@@ -1605,9 +1617,9 @@
 			*PropNotifyFunc,
 			PropFlags,
 			FPropertyObjectFlags,
-			*ArrayDim,
 			*SetterFunc,
 			*GetterFunc,
+			*ArrayDim,
 			OffsetStr,
 			*MetaDataParams,
 			*PropTag
@@ -1639,9 +1651,9 @@
 				*PropNotifyFunc,
 				PropFlags,
 				FPropertyObjectFlags,
-				*ArrayDim,
 				*SetterFunc,
 				*GetterFunc,
+				*ArrayDim,
 				OffsetStr,
 				*GetSingletonNameFuncAddr(PropertyBase.EnumDef, OutReferenceGatherers.UniqueCrossModuleReferences),
 				*MetaDataParams,
@@ -1665,9 +1677,9 @@
 			*PropNotifyFunc,
 			PropFlags,
 			FPropertyObjectFlags,
-			*ArrayDim,
 			*SetterFunc,
 			*GetterFunc,
+			*ArrayDim,
 			OffsetStr,
 			*FString::Printf(TEXT("&F%s::StaticClass"), *PropertyBase.FieldClassName.ToString()),
 			*MetaDataParams,
@@ -1720,6 +1732,35 @@
 	return Result;
 }
 
+void FNativeClassHeaderGenerator::ExportPropertiesParamsCountStaticAssert(FOutputDevice& Out, FUnrealStructDefinitionInfo& StructDef, FString& StaticsName)
+{
+	if (StructDef.GetProperties().Num() == 0)
+	{
+		return;
+	}
+
+	bool bHasAllEditorOnlyDataProperties = true;
+	for (TSharedRef<FUnrealPropertyDefinitionInfo> PropertyDef : StructDef.GetProperties())
+	{
+		if (!IsEditorOnlyDataProperty(*PropertyDef))
+		{
+			bHasAllEditorOnlyDataProperties = false;
+			break;
+		}
+	}
+
+	if (!bHasAllEditorOnlyDataProperties)
+	{
+		Out.Logf(TEXT("\tstatic_assert(UE_ARRAY_COUNT(%s::PropPointers) < 2048);\r\n"), *StaticsName);
+	}
+	else
+	{
+		Out.Log(TEXT("#if WITH_EDITORONLY_DATA\r\n"));
+		Out.Logf(TEXT("\tstatic_assert(UE_ARRAY_COUNT(%s::PropPointers) < 2048);\r\n"), *StaticsName);
+		Out.Log(TEXT("#endif\r\n"));
+	}
+}
+
 TTuple<FString, FString> FNativeClassHeaderGenerator::OutputProperties(FOutputDevice& DeclOut, FOutputDevice& Out, FReferenceGatherers& OutReferenceGatherers, const TCHAR* Scope, FUnrealStructDefinitionInfo& StructDef, const TCHAR* DeclSpaces, const TCHAR* Spaces) const
 {
 	if (StructDef.GetProperties().Num() == 0)
@@ -2085,6 +2126,8 @@
 			StaticDefinitions.Log (*Singletons);
 			StaticDefinitions.Logf(TEXT("\t};\r\n"));
 
+			StaticDefinitions.Logf(TEXT("\tstatic_assert(UE_ARRAY_COUNT(%s::DependentSingletons) < 16);\r\n"), *StaticsStructName);
+
 			SingletonsArray = TEXT("DependentSingletons");
 			SingletonsCount = TEXT("UE_ARRAY_COUNT(DependentSingletons)");
 		}
@@ -2128,6 +2171,7 @@
 			}
 
 			StaticDefinitions.Log(TEXT("\t};\r\n"));
+			StaticDefinitions.Logf(TEXT("\tstatic_assert(UE_ARRAY_COUNT(%s::FuncInfo) < 2048);\r\n"), *StaticsStructName);
 			StaticDefinitions.Log(END_WRAP_EDITOR_ONLY(bAllEditorOnlyFunctions));
 
 			if (bAllEditorOnlyFunctions)
@@ -2197,6 +2241,7 @@
 				);
 			}
 			StaticDefinitions.Log(TEXT("\t\t};\r\n"));
+			StaticDefinitions.Logf(TEXT("\tstatic_assert(UE_ARRAY_COUNT(%s::InterfaceParams) < 64);\r\n"), *StaticsStructName);
 
 			InterfaceArray = TEXT("InterfaceParams");
 			InterfaceCount = TEXT("UE_ARRAY_COUNT(InterfaceParams)");
@@ -2238,6 +2283,8 @@
 		StaticDefinitions.Logf(TEXT("\t\t%s\r\n"), *MetaDataParams);
 		StaticDefinitions.Log (TEXT("\t};\r\n"));
 
+		ExportPropertiesParamsCountStaticAssert(StaticDefinitions, ClassDef, StaticsStructName);
+
 		GeneratedClassRegisterFunctionText.Logf(TEXT("\t};\r\n"));
 		GeneratedClassRegisterFunctionText.Log(*StaticDefinitions);
 
@@ -2461,9 +2508,9 @@
 		*CreateUTF8LiteralString(FunctionDef.GetName()),
 		(bIsSparse ? *CreateUTF8LiteralString(FunctionDef.GetSparseOwningClassName().ToString()) : TEXT("nullptr")),
 		(bIsSparse ? *CreateUTF8LiteralString(FunctionDef.GetSparseDelegateName().ToString()) : TEXT("nullptr")),
-		*StructureSize,
 		*PropertyRange.Get<0>(),
 		*PropertyRange.Get<1>(),
+		*StructureSize,
 		UFunctionObjectFlags,
 		(uint32)FunctionDef.GetFunctionFlags(),
 		bIsNet ? FunctionData.RPCId : 0,
@@ -2471,6 +2518,13 @@
 		*MetaDataParams
 	);
 
+	ExportPropertiesParamsCountStaticAssert(StaticDefinitions, FunctionDef, StaticsStructName);
+
+	if (StructureSize != TEXT("0"))
+	{
+		StaticDefinitions.Logf(TEXT("\tstatic_assert(%s < MAX_uint16);\r\n"), *StructureSize);
+	}
+
 	CurrentFunctionText.Log(TEXT("\t};\r\n"));
 	CurrentFunctionText.Log(*StaticDefinitions);
 
@@ -3749,6 +3803,8 @@
 	for (FUnrealScriptStructDefinitionInfo* NoExportStructDef : NoExportStructs)
 	{
 		ExportMirrorsForNoexportStruct(GeneratedStructRegisterFunctionText, *NoExportStructDef, /*Indent=*/ 2);
+		GeneratedStructRegisterFunctionText.Appendf(TEXT("\t\tstatic_assert(sizeof(%s) < MAX_uint16);\r\n"), *NoExportStructDef->GetNameCPP());
+		GeneratedStructRegisterFunctionText.Appendf(TEXT("\t\tstatic_assert(alignof(%s) < MAX_uint8);\r\n"), *NoExportStructDef->GetNameCPP());
 	}
 
 	if (BaseStructDef)
@@ -3788,15 +3844,17 @@
 	StaticDefinitions.Logf(TEXT("\t\t%s,\r\n"), *GetSingletonNameFuncAddr(BaseStructDef, OutReferenceGatherers.UniqueCrossModuleReferences));
 	StaticDefinitions.Logf(TEXT("\t\t%s,\r\n"), *NewStructOps);
 	StaticDefinitions.Logf(TEXT("\t\t%s,\r\n"), *CreateUTF8LiteralString(ActualStructName));
-	StaticDefinitions.Logf(TEXT("\t\tsizeof(%s),\r\n"), *StructNameCPP);
-	StaticDefinitions.Logf(TEXT("\t\talignof(%s),\r\n"), *StructNameCPP);
 	StaticDefinitions.Logf(TEXT("\t\t%s,\r\n"), *PropertyRange.Get<0>());
 	StaticDefinitions.Logf(TEXT("\t\t%s,\r\n"), *PropertyRange.Get<1>());
+	StaticDefinitions.Logf(TEXT("\t\tsizeof(%s),\r\n"), *StructNameCPP);
+	StaticDefinitions.Logf(TEXT("\t\talignof(%s),\r\n"), *StructNameCPP);
 	StaticDefinitions.Logf(TEXT("\t\t%s,\r\n"), TEXT("RF_Public|RF_Transient|RF_MarkAsNative"));
 	StaticDefinitions.Logf(TEXT("\t\tEStructFlags(0x%08X),\r\n"), (uint32)UncomputedFlags);
 	StaticDefinitions.Logf(TEXT("\t\t%s\r\n"), *MetaDataParams);
 	StaticDefinitions.Log (TEXT("\t};\r\n"));
 
+	ExportPropertiesParamsCountStaticAssert(StaticDefinitions, ScriptStructDef, StaticsStructName);
+		
 	GeneratedStructRegisterFunctionText.Log (TEXT("\t};\r\n"));
 
 	GeneratedStructRegisterFunctionText.Log(StaticDefinitions);
@@ -4022,8 +4080,8 @@
 		StaticDefinitions.Logf(TEXT("\t\t%s,\r\n"), *CreateUTF8LiteralString(EnumNameCpp));
 		StaticDefinitions.Logf(TEXT("\t\t%s,\r\n"), *CreateUTF8LiteralString(EnumDef.GetCppType()));
 		StaticDefinitions.Logf(TEXT("\t\t%s::Enumerators,\r\n"), *StaticsStructName);
-		StaticDefinitions.Logf(TEXT("\t\tUE_ARRAY_COUNT(%s::Enumerators),\r\n"), *StaticsStructName);
 		StaticDefinitions.Logf(TEXT("\t\t%s,\r\n"), UEnumObjectFlags);
+		StaticDefinitions.Logf(TEXT("\t\tUE_ARRAY_COUNT(%s::Enumerators),\r\n"), *StaticsStructName);
 		StaticDefinitions.Logf(TEXT("\t\t%s,\r\n"), EnumFlags);
 		StaticDefinitions.Logf(TEXT("\t\t(uint8)%s,\r\n"), EnumFormStr);
 		StaticDefinitions.Logf(TEXT("\t\t%s\r\n"), *MetaDataParams);
diff --git a/Engine/Source/Programs/UnrealHeaderTool/Private/HeaderParser.cpp b/Engine/Source/Programs/UnrealHeaderTool/Private/HeaderParser.cpp
--- a/Engine/Source/Programs/UnrealHeaderTool/Private/HeaderParser.cpp
+++ b/Engine/Source/Programs/UnrealHeaderTool/Private/HeaderParser.cpp
@@ -8464,37 +8464,7 @@ FHeaderParser::FHeaderParser(FUnrealPackageDefinitionInfo& InPackageDef, FUnreal
 	, Filename(IFileManager::Get().ConvertToAbsolutePathForExternalAppForRead(*InSourceFile.GetFilename()))
 	, PackageDef(InPackageDef)
 {
-
-	const FManifestModule& Module = PackageDef.GetModule();
-
-	// Determine if the current module is part of the engine or a game (we are more strict about things for Engine modules)
-	switch (Module.ModuleType)
-	{
-	case EBuildModuleType::Program:
-		{
-			const FString AbsoluteEngineDir = FPaths::ConvertRelativePathToFull(FPaths::EngineDir());
-			const FString ModuleDir = FPaths::ConvertRelativePathToFull(Module.BaseDirectory);
-			bIsCurrentModulePartOfEngine = ModuleDir.StartsWith(AbsoluteEngineDir);
-		}
-		break;
-	case EBuildModuleType::EngineRuntime:
-	case EBuildModuleType::EngineUncooked:
-	case EBuildModuleType::EngineDeveloper:
-	case EBuildModuleType::EngineEditor:
-	case EBuildModuleType::EngineThirdParty:
-		bIsCurrentModulePartOfEngine = true;
-		break;
-	case EBuildModuleType::GameRuntime:
-	case EBuildModuleType::GameUncooked:
-	case EBuildModuleType::GameDeveloper:
-	case EBuildModuleType::GameEditor:
-	case EBuildModuleType::GameThirdParty:
-		bIsCurrentModulePartOfEngine = false;
-		break;
-	default:
-		bIsCurrentModulePartOfEngine = true;
-		check(false);
-	}
+	bIsCurrentModulePartOfEngine = PackageDef.IsPartOfEngine();
 }
 
 // Throws if a specifier value wasn't provided
@@ -9067,10 +9037,10 @@ void FHeaderParser::SimplifiedClassParse(FUnrealSourceFile& SourceFile, const TC
 			{
 				if (UInterfaceMacroDecl == FCString::Strspn(Str, TEXT("\t ")) + Str)
 				{
-					if (UInterfaceMacroDecl[10] != TEXT('('))
-					{
-						FUHTMessage(SourceFile, CurrentLine).Throwf(TEXT("Missing open parenthesis after UINTERFACE"));
-					}
+					//if (UInterfaceMacroDecl[10] != TEXT('('))
+					//{
+					//	FUHTMessage(SourceFile, CurrentLine).Throwf(TEXT("Missing open parenthesis after UINTERFACE"));
+					//}
 
 					TSharedRef<FUnrealTypeDefinitionInfo> ClassDecl = Parser.ParseClassDeclaration(StartOfLine + (UInterfaceMacroDecl - Str), CurrentLine, true, TEXT("UINTERFACE"));
 					bGeneratedIncludeRequired |= !ClassDecl->AsClassChecked().IsNoExport();
@@ -9082,10 +9052,10 @@ void FHeaderParser::SimplifiedClassParse(FUnrealSourceFile& SourceFile, const TC
 			{
 				if (UClassMacroDecl == FCString::Strspn(Str, TEXT("\t ")) + Str)
 				{
-					if (UClassMacroDecl[6] != TEXT('('))
-					{
-						FUHTMessage(SourceFile, CurrentLine).Throwf(TEXT("Missing open parenthesis after UCLASS"));
-					}
+					//if (UClassMacroDecl[6] != TEXT('('))
+					//{
+					//	FUHTMessage(SourceFile, CurrentLine).Throwf(TEXT("Missing open parenthesis after UCLASS"));
+					//}
 
 					TSharedRef<FUnrealTypeDefinitionInfo> ClassDecl = Parser.ParseClassDeclaration(StartOfLine + (UClassMacroDecl - Str), CurrentLine, false, TEXT("UCLASS"));
 					bGeneratedIncludeRequired |= !ClassDecl->AsClassChecked().IsNoExport();
@@ -9097,10 +9067,10 @@ void FHeaderParser::SimplifiedClassParse(FUnrealSourceFile& SourceFile, const TC
 			{
 				if (UEnumMacroDecl == FCString::Strspn(Str, TEXT("\t ")) + Str)
 				{
-					if (UEnumMacroDecl[5] != TEXT('('))
-					{
-						FUHTMessage(SourceFile, CurrentLine).Throwf(TEXT("Missing open parenthesis after UENUM"));
-					}
+					//if (UEnumMacroDecl[5] != TEXT('('))
+					//{
+					//	FUHTMessage(SourceFile, CurrentLine).Throwf(TEXT("Missing open parenthesis after UENUM"));
+					//}
 
 					TSharedRef<FUnrealTypeDefinitionInfo> EnumDecl = Parser.ParseEnumDeclaration(StartOfLine + (UEnumMacroDecl - Str), CurrentLine);
 					SourceFile.AddDefinedEnum(MoveTemp(EnumDecl));
@@ -9111,10 +9081,10 @@ void FHeaderParser::SimplifiedClassParse(FUnrealSourceFile& SourceFile, const TC
 			{
 				if (UDelegateMacroDecl == FCString::Strspn(Str, TEXT("\t ")) + Str)
 				{
-					if (UDelegateMacroDecl[9] != TEXT('('))
-					{
-						FUHTMessage(SourceFile, CurrentLine).Throwf(TEXT("Missing open parenthesis after UDELEGATE"));
-					}
+					//if (UDelegateMacroDecl[9] != TEXT('('))
+					//{
+					//	FUHTMessage(SourceFile, CurrentLine).Throwf(TEXT("Missing open parenthesis after UDELEGATE"));
+					//}
 					// We don't preparse the delegates, but we need to make sure the include is there
 					bGeneratedIncludeRequired = true;
 				}
@@ -9124,10 +9094,10 @@ void FHeaderParser::SimplifiedClassParse(FUnrealSourceFile& SourceFile, const TC
 			{
 				if (UStructMacroDecl == FCString::Strspn(Str, TEXT("\t ")) + Str)
 				{
-					if (UStructMacroDecl[7] != TEXT('('))
-					{
-						FUHTMessage(SourceFile, CurrentLine).Throwf(TEXT("Missing open parenthesis after USTRUCT"));
-					}
+					//if (UStructMacroDecl[7] != TEXT('('))
+					//{
+					//	FUHTMessage(SourceFile, CurrentLine).Throwf(TEXT("Missing open parenthesis after USTRUCT"));
+					//}
 
 					TSharedRef<FUnrealTypeDefinitionInfo> StructDecl = Parser.ParseStructDeclaration(StartOfLine + (UStructMacroDecl - Str), CurrentLine);
 					bGeneratedIncludeRequired |= !StructDecl->AsScriptStructChecked().HasAnyStructFlags(STRUCT_NoExport);
diff --git a/Engine/Source/Programs/UnrealHeaderTool/Private/NativeClassExporter.h b/Engine/Source/Programs/UnrealHeaderTool/Private/NativeClassExporter.h
--- a/Engine/Source/Programs/UnrealHeaderTool/Private/NativeClassExporter.h
+++ b/Engine/Source/Programs/UnrealHeaderTool/Private/NativeClassExporter.h
@@ -686,6 +686,11 @@ private:
 	*/
 	static FString GetFunctionParameterString(FUnrealFunctionDefinitionInfo& FunctionDef, FReferenceGatherers& OutReferenceGatherers);
 
+	/**
+	 * Exports the static assert for the property count
+	 */
+	static void ExportPropertiesParamsCountStaticAssert(FOutputDevice& Out, FUnrealStructDefinitionInfo& StructDef, FString& StaticsName);
+
 public:
 
 	// @todo: BP2CPP_remove
diff --git a/Engine/Source/Programs/UnrealHeaderTool/Private/UnrealTypeDefinitionInfo.cpp b/Engine/Source/Programs/UnrealHeaderTool/Private/UnrealTypeDefinitionInfo.cpp
--- a/Engine/Source/Programs/UnrealHeaderTool/Private/UnrealTypeDefinitionInfo.cpp
+++ b/Engine/Source/Programs/UnrealHeaderTool/Private/UnrealTypeDefinitionInfo.cpp
@@ -5,6 +5,7 @@
 #include "ClassMaps.h"
 #include "EngineAPI.h"
 #include "HeaderParser.h"
+#include "Manifest.h"
 #include "NativeClassExporter.h"
 #include "PropertyTypes.h"
 #include "Scope.h"
@@ -1081,6 +1082,34 @@ FUnrealPackageDefinitionInfo::FUnrealPackageDefinitionInfo(const FManifestModule
 	, API(FString::Printf(TEXT("%s_API "), *ShortUpperName))
 	, PackageFlags(InPackage->GetPackageFlags())
 {
+	// Determine if the current module is part of the engine or a game (we are more strict about things for Engine modules)
+	switch (Module.ModuleType)
+	{
+	case EBuildModuleType::Program:
+	{
+		const FString AbsoluteEngineDir = FPaths::ConvertRelativePathToFull(FPaths::EngineDir());
+		const FString ModuleDir = FPaths::ConvertRelativePathToFull(Module.BaseDirectory);
+		bIsPartOfEngine = ModuleDir.StartsWith(AbsoluteEngineDir);
+	}
+	break;
+	case EBuildModuleType::EngineRuntime:
+	case EBuildModuleType::EngineUncooked:
+	case EBuildModuleType::EngineDeveloper:
+	case EBuildModuleType::EngineEditor:
+	case EBuildModuleType::EngineThirdParty:
+		bIsPartOfEngine = true;
+		break;
+	case EBuildModuleType::GameRuntime:
+	case EBuildModuleType::GameUncooked:
+	case EBuildModuleType::GameDeveloper:
+	case EBuildModuleType::GameEditor:
+	case EBuildModuleType::GameThirdParty:
+		bIsPartOfEngine = false;
+		break;
+	default:
+		bIsPartOfEngine = true;
+		check(false);
+	}
 }
 
 void FUnrealPackageDefinitionInfo::PostParseFinalizeInternal(EPostParseFinalizePhase Phase)
diff --git a/Engine/Source/Programs/UnrealHeaderTool/Private/UnrealTypeDefinitionInfo.h b/Engine/Source/Programs/UnrealHeaderTool/Private/UnrealTypeDefinitionInfo.h
--- a/Engine/Source/Programs/UnrealHeaderTool/Private/UnrealTypeDefinitionInfo.h
+++ b/Engine/Source/Programs/UnrealHeaderTool/Private/UnrealTypeDefinitionInfo.h
@@ -1384,6 +1384,14 @@ public:
 		return bWriteClassesH;
 	}
 
+	/**
+	 * If true, this package is part of the engine
+	 */
+	bool IsPartOfEngine() const
+	{
+		return bIsPartOfEngine;
+	}
+
 	/**
 	 * Set the flag indicating that the classes H file should be generated.
 	 */
@@ -1469,6 +1477,7 @@ private:
 	FString API;
 	uint32 PackageFlags = 0;
 	bool bWriteClassesH = false;
+	bool bIsPartOfEngine = false;
 };
 
 /**
