Refactor C++ UHT first stage comment stripping to avoid a bug not found in C# UHT.

#rnx

diff --git a/Engine/Source/Programs/UnrealHeaderTool/Private/HeaderParser.cpp b/Engine/Source/Programs/UnrealHeaderTool/Private/HeaderParser.cpp
--- a/Engine/Source/Programs/UnrealHeaderTool/Private/HeaderParser.cpp
+++ b/Engine/Source/Programs/UnrealHeaderTool/Private/HeaderParser.cpp
@@ -8893,13 +8893,12 @@ void FHeaderParser::SimplifiedClassParse(FUnrealSourceFile& SourceFile, const TC
 		CurrentLine++;
 
 		const TCHAR* Str = *StrLine;
-		bool bProcess = !bInComment;	// for skipping nested multi-line comments
 
 		int32 BraceCount = 0;
-		if( bProcess && FParse::Command(&Str,TEXT("#if")) )
+		if( !bInComment && FParse::Command(&Str,TEXT("#if")) )
 		{
 		}
-		else if ( bProcess && FParse::Command(&Str,TEXT("#include")) )
+		else if (!bInComment && FParse::Command(&Str,TEXT("#include")) )
 		{
 			// Handle #include directives as if they were 'dependson' keywords.
 			const FString& DependsOnHeaderName = Str;
@@ -8922,114 +8921,88 @@ void FHeaderParser::SimplifiedClassParse(FUnrealSourceFile& SourceFile, const TC
 				}
 			}
 		}
-		else if ( bProcess && FParse::Command(&Str,TEXT("#else")) )
+		else if ( !bInComment && FParse::Command(&Str,TEXT("#else")) )
 		{
 		}
-		else if ( bProcess && FParse::Command(&Str,TEXT("#elif")) )
+		else if ( !bInComment && FParse::Command(&Str,TEXT("#elif")) )
 		{
 		}
-		else if ( bProcess && FParse::Command(&Str,TEXT("#endif")) )
+		else if ( !bInComment && FParse::Command(&Str,TEXT("#endif")) )
 		{
 		}
 		else
 		{
-			int32 Pos = INDEX_NONE;
-			int32 EndPos = INDEX_NONE;
-			int32 StrBegin = INDEX_NONE;
-			int32 StrEnd = INDEX_NONE;
-				
-			bool bEscaped = false;
-			for ( int32 CharPos = 0; CharPos < StrLine.Len(); CharPos++ )
+			TCHAR* RawStart = StrLine.GetCharArray().GetData();
+			TCHAR* RawIn = RawStart;
+			TCHAR* RawInEnd = RawStart + StrLine.Len();
+			TCHAR* RawOut = RawStart;
+
+			while (RawIn < RawInEnd)
 			{
-				if ( bEscaped )
-				{
-					bEscaped = false;
-				}
-				else if ( StrLine[CharPos] == TEXT('\\') )
-				{
-					bEscaped = true;
-				}
-				else if ( StrLine[CharPos] == TEXT('\"') )
+				TCHAR C = *RawIn++;
+				if (bInComment)
 				{
-					if ( StrBegin == INDEX_NONE )
+					if (C == '*' && RawIn < RawInEnd && *RawIn == '/')
 					{
-						StrBegin = CharPos;
-					}
-					else
-					{
-						StrEnd = CharPos;
-						break;
+						++RawIn;
+						bInComment = false;
 					}
 				}
-			}
-
-			// Find the first '/' and check for '//' or '/*' or '*/'
-			if (StrLine.FindChar(TEXT('/'), Pos))
-			{
-				if (Pos >= 0)
+				else
 				{
-					// Stub out the comments, ignoring anything inside literal strings.
-					Pos = StrLine.Find(TEXT("//"), ESearchCase::CaseSensitive, ESearchDir::FromStart, Pos);
-
-					// Check if first slash is end of multiline comment and adjust position if necessary.
-					if (Pos > 0 && StrLine[Pos - 1] == TEXT('*'))
-					{
-						++Pos;
-					}
-
-					if (Pos >= 0)
+					switch (C)
 					{
-						if (StrBegin == INDEX_NONE || Pos < StrBegin || Pos > StrEnd)
-						{
-							StrLine.LeftInline(Pos, false);
-						}
-
-						if (StrLine.IsEmpty())
+					case '\"':
+						*RawOut++ = C;
+						while (RawIn < RawInEnd)
 						{
-							continue;
+							C = *RawIn++;
+							*RawOut++ = C;
+							if (C == '\\')
+							{
+								if (RawIn < RawInEnd)
+								{
+									*RawOut++ = *RawIn++;
+								}
+							}
+							else if (C == '\"')
+							{
+								break;
+							}
 						}
-					}
+						break;
 
-					// look for a / * ... * / block, ignoring anything inside literal strings
-					Pos = StrLine.Find(TEXT("/*"), ESearchCase::CaseSensitive, ESearchDir::FromStart, Pos);
-					EndPos = StrLine.Find(TEXT("*/"), ESearchCase::CaseSensitive, ESearchDir::FromStart, FMath::Max(0, Pos - 1));
-					if (Pos >= 0)
-					{
-						if (StrBegin == INDEX_NONE || Pos < StrBegin || Pos > StrEnd)
+					case '/':
+						if (RawIn < RawInEnd)
 						{
-							if (EndPos != INDEX_NONE && (EndPos < StrBegin || EndPos > StrEnd))
+							if (*RawIn == '*')
 							{
-								StrLine = StrLine.Left(Pos) + StrLine.Mid(EndPos + 2);
-								EndPos = INDEX_NONE;
-								bInComment = false;
+								++RawIn;
+								bInComment = true;
+								continue;
 							}
-							else
+							else if (*RawIn == '/')
 							{
-								StrLine.LeftInline(Pos, false);
-								bInComment = true;
+								RawIn = RawInEnd;
+								break;
 							}
 						}
-						bProcess = !bInComment;
-					}
+						break;
 
-					if (EndPos >= 0)
-					{
-						if (StrBegin == INDEX_NONE || EndPos < StrBegin || EndPos > StrEnd)
-						{
-							StrLine.MidInline(EndPos + 2, MAX_int32, false);
-							bInComment = false;
-						}
-						bProcess = !bInComment;
+					default:
+						*RawOut++ = C;
+						break;
 					}
 				}
 			}
 
+			StrLine.LeftInline(static_cast<int32>(RawOut - RawStart));
 			StrLine.TrimStartInline();
-			if (!bProcess || StrLine.IsEmpty())
+			if (StrLine.IsEmpty())
 			{
 				continue;
 			}
-
+			
 			Str = *StrLine;
 
 			// Get class or interface name
