Fixed C++ UHT to be able to detect getter/setter functions that have been declared prior to the property.
Fixed C++ UHT to ignore possible interface classes that might share a name with a UCLASS.

#rnx

diff --git a/Engine/Source/Programs/UnrealHeaderTool/Private/HeaderParser.cpp b/Engine/Source/Programs/UnrealHeaderTool/Private/HeaderParser.cpp
--- a/Engine/Source/Programs/UnrealHeaderTool/Private/HeaderParser.cpp
+++ b/Engine/Source/Programs/UnrealHeaderTool/Private/HeaderParser.cpp
@@ -2861,6 +2861,21 @@ void FHeaderParser::VerifyNotifyValueChangedFunction(const FUnrealFunctionDefini
 
 void FHeaderParser::VerifyPropertyMarkups(FUnrealClassDefinitionInfo& TargetClassDef)
 {
+	// Iterate through the definitions looking for getter/setter function
+	for (const FDeclaration& Declaration : TargetClassDef.GetDeclarations())
+	{
+		if (Declaration.bIsUFunction)
+		{
+			continue;
+		}
+		if (CheckForPropertySetterFunction(TargetClassDef, Declaration))
+		{
+		}
+		else if (CheckForPropertyGetterFunction(TargetClassDef, Declaration))
+		{
+		}
+	}
+
 	// Iterate over all properties, looking for those flagged as CPF_RepNotify
 	for (TSharedRef<FUnrealPropertyDefinitionInfo> PropertyDef : TargetClassDef.GetProperties())
 	{
@@ -5328,7 +5343,7 @@ bool FHeaderParser::CompileDeclaration(TArray<FUnrealFunctionDefinitionInfo*>& D
 		return true;
 	}
 
-	FRecordTokens RecordTokens(*this, bEncounteredNewStyleClass_UnmatchedBrackets && IsInAClass() ? &GetCurrentClassDef() : nullptr, &Token);
+	FRecordTokens RecordTokens(*this, bEncounteredNewStyleClass_UnmatchedBrackets && IsInAClass() ? &GetCurrentClassDef() : nullptr, false, &Token);
 	bool Result = SkipDeclaration(Token);
 	if (RecordTokens.Stop())
 	{
@@ -5345,12 +5360,6 @@ bool FHeaderParser::CompileDeclaration(TArray<FUnrealFunctionDefinitionInfo*>& D
 			if (CheckForSerialize(StructDef, Declaration))
 			{
 			}
-			else if (CheckForPropertySetterFunction(StructDef, Declaration))
-			{
-			}
-			else if (CheckForPropertyGetterFunction(StructDef, Declaration))
-			{
-			}
 		}
 	}
 	return Result;
@@ -7354,7 +7363,7 @@ FUnrealFunctionDefinitionInfo& FHeaderParser::CompileFunctionDeclaration()
 	}
 
 	// Record the tokens so we can detect this function as a declaration later (i.e. RPC)
-	FRecordTokens RecordTokens(*this, &GetCurrentClassDef(), nullptr);
+	FRecordTokens RecordTokens(*this, &GetCurrentClassDef(), true, nullptr);
 
 	bool bSawVirtual = false;
 
@@ -10078,9 +10087,10 @@ FUnrealFunctionDefinitionInfo& FHeaderParser::CreateFunction(const TCHAR* FuncNa
 	return FuncDef;
 }
 
-FRecordTokens::FRecordTokens(FHeaderParser& InParser, FUnrealStructDefinitionInfo* InStructDef, FToken* InToken)
+FRecordTokens::FRecordTokens(FHeaderParser& InParser, FUnrealStructDefinitionInfo* InStructDef, bool bInIsUFunction, FToken* InToken)
 	: Parser(InParser)
 	, StructDef(InStructDef)
+	, bIsUFunction(bInIsUFunction)
 	, CurrentCompilerDirective(Parser.GetCurrentCompilerDirective())
 {
 	if (StructDef != nullptr)
@@ -10103,7 +10113,7 @@ bool FRecordTokens::Stop()
 	if (StructDef != nullptr)
 	{
 		Parser.bRecordTokens = false;
-		StructDef->AddDeclaration(CurrentCompilerDirective, MoveTemp(Parser.RecordedTokens));
+		StructDef->AddDeclaration(CurrentCompilerDirective, MoveTemp(Parser.RecordedTokens), bIsUFunction);
 		StructDef = nullptr;
 		return true;
 	}
diff --git a/Engine/Source/Programs/UnrealHeaderTool/Private/HeaderParser.h b/Engine/Source/Programs/UnrealHeaderTool/Private/HeaderParser.h
--- a/Engine/Source/Programs/UnrealHeaderTool/Private/HeaderParser.h
+++ b/Engine/Source/Programs/UnrealHeaderTool/Private/HeaderParser.h
@@ -706,12 +706,13 @@ public:
 class FRecordTokens
 {
 public:
-	explicit FRecordTokens(FHeaderParser& InParser, FUnrealStructDefinitionInfo* InStructDef, FToken* InToken);
+	explicit FRecordTokens(FHeaderParser& InParser, FUnrealStructDefinitionInfo* InStructDef, bool bInIsUFunction, FToken* InToken);
 	~FRecordTokens();
 	bool Stop();
 
 private:
 	FHeaderParser& Parser;
 	FUnrealStructDefinitionInfo* StructDef;
+	bool bIsUFunction;
 	uint32 CurrentCompilerDirective = 0;
 };
\ No newline at end of file
diff --git a/Engine/Source/Programs/UnrealHeaderTool/Private/UnrealTypeDefinitionInfo.h b/Engine/Source/Programs/UnrealHeaderTool/Private/UnrealTypeDefinitionInfo.h
--- a/Engine/Source/Programs/UnrealHeaderTool/Private/UnrealTypeDefinitionInfo.h
+++ b/Engine/Source/Programs/UnrealHeaderTool/Private/UnrealTypeDefinitionInfo.h
@@ -95,6 +95,7 @@ struct FDeclaration
 {
 	uint32 CurrentCompilerDirective;
 	TArray<FToken> Tokens;
+	bool bIsUFunction;
 };
 
 /**
@@ -2020,9 +2021,9 @@ public:
 	/**
 	 * Add a declaration to the struct
 	 */
-	void AddDeclaration(uint32 CurrentCompilerDirective, TArray<FToken>&& Tokens)
+	void AddDeclaration(uint32 CurrentCompilerDirective, TArray<FToken>&& Tokens, bool bIsUFunction)
 	{
-		Declarations.Emplace(FDeclaration{ CurrentCompilerDirective, MoveTemp(Tokens) });
+		Declarations.Emplace(FDeclaration{ CurrentCompilerDirective, MoveTemp(Tokens), bIsUFunction });
 	}
 
 	/** 
