Fixing issue where differences in spacing would cause get/set method to not match.

diff --git a/Engine/Source/Programs/UnrealHeaderTool/Private/HeaderParser.cpp b/Engine/Source/Programs/UnrealHeaderTool/Private/HeaderParser.cpp
--- a/Engine/Source/Programs/UnrealHeaderTool/Private/HeaderParser.cpp
+++ b/Engine/Source/Programs/UnrealHeaderTool/Private/HeaderParser.cpp
@@ -5987,7 +5987,13 @@ bool FHeaderParser::CheckForPropertyGetterFunction(FUnrealStructDefinitionInfo&
 	{
 		PropertyType += TEXT("*");
 	}
-	if (GetterType.Compare(PropertyType, ESearchCase::CaseSensitive) != 0)
+
+	FString GetterTypeTest = GetterType;
+	GetterTypeTest.RemoveSpacesInline();
+	FString PropertyTypeTest = PropertyType;
+	PropertyTypeTest.RemoveSpacesInline();
+
+	if (GetterTypeTest.Compare(PropertyTypeTest, ESearchCase::CaseSensitive) != 0)
 	{
 		if (bExplicitGetter)
 		{
