Fix non-unity
#rnx

diff --git a/Engine/Source/Programs/UnrealHeaderTool/Private/PropertyTypes.cpp b/Engine/Source/Programs/UnrealHeaderTool/Private/PropertyTypes.cpp
--- a/Engine/Source/Programs/UnrealHeaderTool/Private/PropertyTypes.cpp
+++ b/Engine/Source/Programs/UnrealHeaderTool/Private/PropertyTypes.cpp
@@ -9,6 +9,7 @@
 #include "UnrealTypeDefinitionInfo.h"
 #include "Misc/DefaultValueHelper.h"
 #include "UObject/ObjectMacros.h"
+#include <functional>
 
 // Globals for common class definitions
 extern FUnrealClassDefinitionInfo* GUObjectDef;
