Updating C++ UHT to match recent changes to C# UHT

diff --git a/Engine/Source/Programs/UnrealHeaderTool/Private/CodeGenerator.cpp b/Engine/Source/Programs/UnrealHeaderTool/Private/CodeGenerator.cpp
--- a/Engine/Source/Programs/UnrealHeaderTool/Private/CodeGenerator.cpp
+++ b/Engine/Source/Programs/UnrealHeaderTool/Private/CodeGenerator.cpp
@@ -4999,6 +4999,39 @@ struct FNativeFunctionStringBuilder
 	FUHTStringBuilder AccessorWrapperImplementations;
 };
 
+TArray<const FUnrealScriptStructDefinitionInfo*> GetSparseDataStructsToExport(FUnrealClassDefinitionInfo& ClassDef)
+{
+	TArray<const FUnrealScriptStructDefinitionInfo*> out;
+
+	TArray<FString> SparseClassDataTypes;
+	ClassDef.GetSparseClassDataTypes(SparseClassDataTypes);
+
+	TSet<FString> SuperSparseClassDataTypes;
+	if (FUnrealClassDefinitionInfo* SuperClass = ClassDef.GetSuperClass())
+	{
+		TArray<FString> TempTypes;
+		SuperClass->GetSparseClassDataTypes(TempTypes);
+		for (const FString& Type : TempTypes)
+		{
+			SuperSparseClassDataTypes.Add(Type);
+		}
+	}
+
+	for (const FString& SparseClass : SparseClassDataTypes)
+	{
+		if (FUnrealScriptStructDefinitionInfo* SparseScriptStructDef = GTypeDefinitionInfoMap.FindByName<FUnrealScriptStructDefinitionInfo>(*SparseClass))
+		{
+			while (SparseScriptStructDef != nullptr && !SuperSparseClassDataTypes.Contains(SparseScriptStructDef->GetName()))
+			{
+				out.Add(SparseScriptStructDef);
+				SparseScriptStructDef = SparseScriptStructDef->GetSuperScriptStruct();
+			}
+		}
+	}
+
+	return out;
+}
+
 void FNativeClassHeaderGenerator::ExportNativeFunctions(FOutputDevice& OutGeneratedHeaderText, FOutputDevice& OutGeneratedCPPText, FOutputDevice& OutMacroCalls, FOutputDevice& OutNoPureDeclsMacroCalls, FReferenceGatherers& OutReferenceGatherers, const FUnrealSourceFile& SourceFile, FUnrealClassDefinitionInfo& ClassDef) const
 {
 	FNativeFunctionStringBuilder RuntimeStringBuilders;
@@ -5007,27 +5040,17 @@ void FNativeClassHeaderGenerator::ExportNativeFunctions(FOutputDevice& OutGenera
 	const FString ClassCPPName = ClassDef.GetAlternateNameCPP(ClassDef.HasAnyClassFlags(CLASS_Interface));
 
 	// gather static class data
-	TArray<FString> SparseClassDataTypes;
-	ClassDef.GetSparseClassDataTypes(SparseClassDataTypes);
+	TArray<const FUnrealScriptStructDefinitionInfo*> SparseClassDataTypes = GetSparseDataStructsToExport(ClassDef);
 	FString FullClassName = ClassDef.GetNameWithPrefix();
-	for (const FString& SparseClassDataString : SparseClassDataTypes)
+	for (const FUnrealScriptStructDefinitionInfo* SparseClassData : SparseClassDataTypes)
 	{
-		TArray<FString> SuperSparseClassDataTypes;
-		if (FUnrealClassDefinitionInfo* SuperClass = ClassDef.GetSuperClass())
-		{
-			SuperClass->GetSparseClassDataTypes(SuperSparseClassDataTypes);
-		}
-
-		if (SparseClassDataTypes != SuperSparseClassDataTypes)
-		{
-			FString Empty;
-			const FString& API = ClassDef.HasAnyClassFlags(CLASS_MinimalAPI) ? GetAPIString() : Empty;
+		FString Empty;
+		const FString& API = ClassDef.HasAnyClassFlags(CLASS_MinimalAPI) ? GetAPIString() : Empty;
 
-			RuntimeStringBuilders.AutogeneratedStaticData.Logf(TEXT("%sF%s* Get%s();\r\n"), *API, *SparseClassDataString, *SparseClassDataString);
-			RuntimeStringBuilders.AutogeneratedStaticData.Logf(TEXT("%sF%s* Get%s() const;\r\n"), *API, *SparseClassDataString, *SparseClassDataString);
-			RuntimeStringBuilders.AutogeneratedStaticData.Logf(TEXT("%sconst F%s* Get%s(EGetSparseClassDataMethod GetMethod) const;\r\n"), *API, *SparseClassDataString, *SparseClassDataString);
-			RuntimeStringBuilders.AutogeneratedStaticData.Logf(TEXT("%sstatic UScriptStruct* StaticGet%sScriptStruct();\r\n"), *API, *SparseClassDataString);
-		}
+		const FString SparseClassDataString = SparseClassData->GetName();
+		RuntimeStringBuilders.AutogeneratedStaticData.Logf(TEXT("%sF%s* Get%s() const;\r\n"), *API, *SparseClassDataString, *SparseClassDataString);
+		RuntimeStringBuilders.AutogeneratedStaticData.Logf(TEXT("%sconst F%s* Get%s(EGetSparseClassDataMethod GetMethod) const;\r\n"), *API, *SparseClassDataString, *SparseClassDataString);
+		RuntimeStringBuilders.AutogeneratedStaticData.Logf(TEXT("%sstatic UScriptStruct* StaticGet%sScriptStruct();\r\n"), *API, *SparseClassDataString);
 
 		FUnrealScriptStructDefinitionInfo* SparseClassDataStructDef = GTypeDefinitionInfoMap.FindByName<FUnrealScriptStructDefinitionInfo>(*SparseClassDataString);
 		while (SparseClassDataStructDef != nullptr)
@@ -5049,18 +5072,6 @@ void FNativeClassHeaderGenerator::ExportNativeFunctions(FOutputDevice& OutGenera
 
 				if (!PropertyDef->HasMetaData(NAME_NoGetter))
 				{
-					if (PropertyDef->HasMetaData(NAME_GetByRef))
-					{
-						RuntimeStringBuilders.AutogeneratedStaticDataFuncs.Logf(TEXT("const %s& Get%s()\r\n"), *VarType, *CleanVarName);
-					}
-					else
-					{
-						RuntimeStringBuilders.AutogeneratedStaticDataFuncs.Logf(TEXT("%s Get%s()\r\n"), *VarType, *CleanVarName);
-					}
-					RuntimeStringBuilders.AutogeneratedStaticDataFuncs.Logf(TEXT("{\r\n"));
-					RuntimeStringBuilders.AutogeneratedStaticDataFuncs.Logf(TEXT("\treturn Get%s()->%s;\r\n"), *SparseClassDataString, *VarName);
-					RuntimeStringBuilders.AutogeneratedStaticDataFuncs.Logf(TEXT("}\r\n"));
-
 					if (PropertyDef->HasMetaData(NAME_GetByRef))
 					{
 						RuntimeStringBuilders.AutogeneratedStaticDataFuncs.Logf(TEXT("const %s& Get%s() const\r\n"), *VarType, *CleanVarName);
@@ -5070,7 +5081,7 @@ void FNativeClassHeaderGenerator::ExportNativeFunctions(FOutputDevice& OutGenera
 						RuntimeStringBuilders.AutogeneratedStaticDataFuncs.Logf(TEXT("%s Get%s() const\r\n"), *VarType, *CleanVarName);
 					}
 					RuntimeStringBuilders.AutogeneratedStaticDataFuncs.Logf(TEXT("{\r\n"));
-					RuntimeStringBuilders.AutogeneratedStaticDataFuncs.Logf(TEXT("\treturn Get%s()->%s;\r\n"), *SparseClassDataString, *VarName);
+					RuntimeStringBuilders.AutogeneratedStaticDataFuncs.Logf(TEXT("\treturn Get%s(EGetSparseClassDataMethod::ArchetypeIfNull)->%s;\r\n"), *SparseClassDataString, *VarName);
 					RuntimeStringBuilders.AutogeneratedStaticDataFuncs.Logf(TEXT("}\r\n"));
 				}
 			}
@@ -5359,39 +5370,24 @@ void FNativeClassHeaderGenerator::ExportNativeFunctions(FOutputDevice& OutGenera
 		OutGeneratedCPPText.Log(RuntimeStringBuilders.AccessorWrapperImplementations);
 	}
 
-	if (!SparseClassDataTypes.IsEmpty())
+	for (const FUnrealScriptStructDefinitionInfo* SparseClassData : SparseClassDataTypes)
 	{
-		TArray<FString> SuperSparseClassDataTypes;
-		if (FUnrealClassDefinitionInfo* SuperClass = ClassDef.GetSuperClass())
-		{
-			SuperClass->GetSparseClassDataTypes(SuperSparseClassDataTypes);
-		}
+		const FString SparseClassDataString = SparseClassData->GetName();
 
-		if (SuperSparseClassDataTypes != SparseClassDataTypes)
-		{
-			for (const FString& SparseClassDataString : SparseClassDataTypes)
-			{
-				OutGeneratedCPPText.Logf(TEXT("F%s* %s::Get%s() \r\n"), *SparseClassDataString, *ClassCPPName, *SparseClassDataString);
-				OutGeneratedCPPText.Log(TEXT("{ \r\n"));
-				OutGeneratedCPPText.Logf(TEXT("\treturn static_cast<F%s*>(GetClass()->GetOrCreateSparseClassData()); \r\n"), *SparseClassDataString);
-				OutGeneratedCPPText.Log(TEXT("} \r\n"));
-
-				OutGeneratedCPPText.Logf(TEXT("F%s* %s::Get%s() const \r\n"), *SparseClassDataString, *ClassCPPName, *SparseClassDataString);
-				OutGeneratedCPPText.Log(TEXT("{ \r\n"));
-				OutGeneratedCPPText.Logf(TEXT("\treturn static_cast<F%s*>(GetClass()->GetOrCreateSparseClassData()); \r\n"), *SparseClassDataString);
-				OutGeneratedCPPText.Log(TEXT("} \r\n"));
+		OutGeneratedCPPText.Logf(TEXT("F%s* %s::Get%s() const \r\n"), *SparseClassDataString, *ClassCPPName, *SparseClassDataString);
+		OutGeneratedCPPText.Log(TEXT("{ \r\n"));
+		OutGeneratedCPPText.Logf(TEXT("\treturn static_cast<F%s*>(GetClass()->GetOrCreateSparseClassData()); \r\n"), *SparseClassDataString);
+		OutGeneratedCPPText.Log(TEXT("} \r\n"));
 
-				OutGeneratedCPPText.Logf(TEXT("const F%s* %s::Get%s(EGetSparseClassDataMethod GetMethod) const \r\n"), *SparseClassDataString, *ClassCPPName, *SparseClassDataString);
-				OutGeneratedCPPText.Log(TEXT("{ \r\n"));
-				OutGeneratedCPPText.Logf(TEXT("\treturn static_cast<const F%s*>(GetClass()->GetSparseClassData(GetMethod)); \r\n"), *SparseClassDataString);
-				OutGeneratedCPPText.Log(TEXT("} \r\n"));
+		OutGeneratedCPPText.Logf(TEXT("const F%s* %s::Get%s(EGetSparseClassDataMethod GetMethod) const \r\n"), *SparseClassDataString, *ClassCPPName, *SparseClassDataString);
+		OutGeneratedCPPText.Log(TEXT("{ \r\n"));
+		OutGeneratedCPPText.Logf(TEXT("\treturn static_cast<const F%s*>(GetClass()->GetSparseClassData(GetMethod)); \r\n"), *SparseClassDataString);
+		OutGeneratedCPPText.Log(TEXT("} \r\n"));
 
-				OutGeneratedCPPText.Logf(TEXT("UScriptStruct* %s::StaticGet%sScriptStruct()\r\n"), *ClassCPPName, *SparseClassDataString);
-				OutGeneratedCPPText.Log(TEXT("{ \r\n"));
-				OutGeneratedCPPText.Logf(TEXT("\treturn F%s::StaticStruct(); \r\n"), *SparseClassDataString);
-				OutGeneratedCPPText.Log(TEXT("} \r\n"));
-			}
-		}
+		OutGeneratedCPPText.Logf(TEXT("UScriptStruct* %s::StaticGet%sScriptStruct()\r\n"), *ClassCPPName, *SparseClassDataString);
+		OutGeneratedCPPText.Log(TEXT("{ \r\n"));
+		OutGeneratedCPPText.Logf(TEXT("\treturn F%s::StaticStruct(); \r\n"), *SparseClassDataString);
+		OutGeneratedCPPText.Log(TEXT("} \r\n"));
 	}
 }
 
