Enhance C#/C++ UHT to be able to parse more C++ attributes.

diff --git a/Engine/Source/Programs/Shared/EpicGames.UHT/Parsers/UhtClassParser.cs b/Engine/Source/Programs/Shared/EpicGames.UHT/Parsers/UhtClassParser.cs
--- a/Engine/Source/Programs/Shared/EpicGames.UHT/Parsers/UhtClassParser.cs
+++ b/Engine/Source/Programs/Shared/EpicGames.UHT/Parsers/UhtClassParser.cs
@@ -76,7 +76,7 @@ namespace EpicGames.UHT.Parsers
 					// Consume the class specifier
 					topScope.TokenReader.Require("class");
 
-					topScope.TokenReader.SkipAlignasAndDeprecatedMacroIfNecessary();
+					topScope.TokenReader.OptionalAttributes(true);
 
 					// Read the class name and possible API macro name
 					topScope.TokenReader.TryOptionalAPIMacro(out UhtToken apiMacroToken);
diff --git a/Engine/Source/Programs/Shared/EpicGames.UHT/Parsers/UhtEnumParser.cs b/Engine/Source/Programs/Shared/EpicGames.UHT/Parsers/UhtEnumParser.cs
--- a/Engine/Source/Programs/Shared/EpicGames.UHT/Parsers/UhtEnumParser.cs
+++ b/Engine/Source/Programs/Shared/EpicGames.UHT/Parsers/UhtEnumParser.cs
@@ -41,19 +41,20 @@ namespace EpicGames.UHT.Parsers
 					specifiers.ParseSpecifiers();
 
 					// Read the name and the CPP type
-					switch (topScope.TokenReader.TryOptional(new string[] { "namespace", "enum" }))
+					if (topScope.TokenReader.TryOptional("namespace"))
 					{
-						case 0: // namespace
-							enumObject.CppForm = UhtEnumCppForm.Namespaced;
-							topScope.TokenReader.SkipDeprecatedMacroIfNecessary();
-							break;
-						case 1: // enum
-							enumObject.CppForm = topScope.TokenReader.TryOptional(new string[] { "class", "struct" }) >= 0 ? UhtEnumCppForm.EnumClass : UhtEnumCppForm.Regular;
-							topScope.TokenReader.SkipAlignasAndDeprecatedMacroIfNecessary();
-							break;
-						default:
-							throw new UhtTokenException(topScope.TokenReader, topScope.TokenReader.PeekToken(), null);
+						enumObject.CppForm = UhtEnumCppForm.Namespaced;
+					}
+					else if (topScope.TokenReader.TryOptional("enum"))
+					{
+						enumObject.CppForm = topScope.TokenReader.TryOptional("class") || topScope.TokenReader.TryOptional("struct") ? UhtEnumCppForm.EnumClass : UhtEnumCppForm.Regular;
 					}
+					else
+					{
+						throw new UhtTokenException(topScope.TokenReader, topScope.TokenReader.PeekToken(), null);
+					}
+
+					topScope.TokenReader.OptionalAttributes(false);
 
 					UhtToken enumToken = topScope.TokenReader.GetIdentifier("enumeration name");
 
@@ -104,7 +105,7 @@ namespace EpicGames.UHT.Parsers
 						case UhtEnumCppForm.Namespaced:
 							// Now handle the inner true enum portion
 							topScope.TokenReader.Require("enum");
-							topScope.TokenReader.SkipAlignasAndDeprecatedMacroIfNecessary();
+							topScope.TokenReader.OptionalAttributes(true);
 
 							UhtToken innerEnumToken = topScope.TokenReader.GetIdentifier("enumeration type name");
 
@@ -183,7 +184,7 @@ namespace EpicGames.UHT.Parsers
 						topScope.AddFormattedCommentsAsTooltipMetaData(enumIndex);
 
 						// Skip any deprecation
-						topScope.TokenReader.SkipDeprecatedMacroIfNecessary();
+						topScope.TokenReader.OptionalAttributes(false);
 
 						// Try to read an optional explicit enum value specification
 						if (topScope.TokenReader.TryOptional('='))
diff --git a/Engine/Source/Programs/Shared/EpicGames.UHT/Parsers/UhtScriptStructParser.cs b/Engine/Source/Programs/Shared/EpicGames.UHT/Parsers/UhtScriptStructParser.cs
--- a/Engine/Source/Programs/Shared/EpicGames.UHT/Parsers/UhtScriptStructParser.cs
+++ b/Engine/Source/Programs/Shared/EpicGames.UHT/Parsers/UhtScriptStructParser.cs
@@ -81,7 +81,7 @@ namespace EpicGames.UHT.Parsers
 					// Consume the struct specifier
 					topScope.TokenReader.Require("struct");
 
-					topScope.TokenReader.SkipAlignasAndDeprecatedMacroIfNecessary();
+					topScope.TokenReader.OptionalAttributes(true);
 
 					// Read the struct name and possible API macro name
 					topScope.TokenReader.TryOptionalAPIMacro(out UhtToken apiMacroToken);
diff --git a/Engine/Source/Programs/Shared/EpicGames.UHT/Tokenizer/UhtTokenBufferReader.cs b/Engine/Source/Programs/Shared/EpicGames.UHT/Tokenizer/UhtTokenBufferReader.cs
--- a/Engine/Source/Programs/Shared/EpicGames.UHT/Tokenizer/UhtTokenBufferReader.cs
+++ b/Engine/Source/Programs/Shared/EpicGames.UHT/Tokenizer/UhtTokenBufferReader.cs
@@ -680,7 +680,8 @@ Restart:
 							|| (c == '=' && d == '=')
 							|| (c == '*' && d == '*')
 							|| (c == '~' && d == '=')
-							|| (c == ':' && d == ':'))
+							|| (c == ':' && d == ':')
+							|| (c == '[' && d == '['))
 						{
 							++_inputPos;
 						}
diff --git a/Engine/Source/Programs/Shared/EpicGames.UHT/Tokenizer/UhtTokenReaderGeneralExtensions.cs b/Engine/Source/Programs/Shared/EpicGames.UHT/Tokenizer/UhtTokenReaderGeneralExtensions.cs
--- a/Engine/Source/Programs/Shared/EpicGames.UHT/Tokenizer/UhtTokenReaderGeneralExtensions.cs
+++ b/Engine/Source/Programs/Shared/EpicGames.UHT/Tokenizer/UhtTokenReaderGeneralExtensions.cs
@@ -1,6 +1,7 @@
 // Copyright Epic Games, Inc. All Rights Reserved.
 
 using System;
+using System.Text;
 using EpicGames.Core;
 using EpicGames.UHT.Utils;
 
@@ -59,26 +60,6 @@ namespace EpicGames.UHT.Tokenizer
 			return token.IsIdentifier(text) || token.IsSymbol(text);
 		}
 
-		/// <summary>
-		/// Test to see if the next token is one of the given strings.
-		/// </summary>
-		/// <param name="tokenReader">Token reader</param>
-		/// <param name="text">List of keywords to test</param>
-		/// <returns>Index of matched string or -1 if nothing matched</returns>
-		public static int TryOptional(this IUhtTokenReader tokenReader, string[] text)
-		{
-			ref UhtToken token = ref tokenReader.PeekToken();
-			for (int index = 0, endIndex = text.Length; index < endIndex; ++index)
-			{
-				if (token.IsIdentifier(text[index]) || token.IsSymbol(text[index]))
-				{
-					tokenReader.ConsumeToken();
-					return index;
-				}
-			}
-			return -1;
-		}
-
 		/// <summary>
 		/// Try to parse the given text
 		/// </summary>
@@ -231,6 +212,131 @@ namespace EpicGames.UHT.Tokenizer
 			return tokenReader;
 		}
 
+		/// <summary>
+		/// Parse attributes and optionally alignment specifier 
+		/// </summary>
+		/// <param name="tokenReader">Token reader</param>
+		/// <param name="enableAlignAs">If true, also parse alignment specifiers</param>
+		/// <param name="attributeAction">If specified, action to be invoked for every attribute found</param>
+		/// <returns>Token reader</returns>
+		public static IUhtTokenReader OptionalAttributes(this IUhtTokenReader tokenReader, bool enableAlignAs, Action<string>? attributeAction = null)
+		{
+			for (; ; )
+			{
+				UhtToken token = tokenReader.PeekToken();
+				if (token.IsIdentifier("DEPRECATED") || token.IsIdentifier("UE_DEPRECATED"))
+				{
+					tokenReader.ConsumeToken();
+					tokenReader
+						.Require('(', "deprecation macro")
+						.RequireConstFloat("version in deprecation macro")
+						.Require(',', "deprecation macro")
+						.RequireConstString("message in deprecation macro")
+						.Require(')', "deprecation macro");
+					attributeAction?.Invoke("deprecated");
+				}
+				else if (token.IsIdentifier("UE_NODISCARD") || token.IsIdentifier("UE_NODISCARD_CTOR"))
+				{
+					tokenReader.ConsumeToken();
+					attributeAction?.Invoke("nodiscard");
+				}
+				else if (token.IsIdentifier("UE_NORETURN"))
+				{
+					tokenReader.ConsumeToken();
+					attributeAction?.Invoke("noreturn");
+				}
+				else if (token.IsIdentifier("UE_NO_UNIQUE_ADDRESS"))
+				{
+					tokenReader.ConsumeToken();
+					attributeAction?.Invoke("no_unique_address");
+				}
+				else if (token.IsSymbol("[["))
+				{
+					tokenReader.ConsumeToken();
+
+					StringBuilder identifierBuilder = StringBuilderCache.Small.Borrow();
+					try
+					{
+						if (tokenReader.TryOptional("using"))
+						{
+							tokenReader
+								.RequireCppIdentifier(UhtCppIdentifierOptions.None, x => x.Join(identifierBuilder, "::"))
+								.Require(':');
+						}
+						if (identifierBuilder.Length > 0)
+						{
+							identifierBuilder.Append("::");
+						}
+						int identifierBuilderLength = identifierBuilder.Length;
+
+						// Attributes do allow patterns such as [[,]].
+						for (; ; )
+						{
+							token = tokenReader.GetToken();
+							if (token.IsIdentifier())
+							{
+								tokenReader.RequireCppIdentifier(ref token, UhtCppIdentifierOptions.None, x => x.Join(identifierBuilder, "::"));
+
+								if (tokenReader.TryOptional('('))
+								{
+									int nestingCount = 1;
+									while (nestingCount > 0)
+									{
+										if (tokenReader.IsEOF)
+										{
+											throw new UhtTokenException(tokenReader, tokenReader.GetToken(), ")");
+										}
+										else if (tokenReader.TryOptional(')'))
+										{
+											nestingCount--;
+										}
+										else if (tokenReader.TryOptional('('))
+										{
+											nestingCount++;
+										}
+										tokenReader.ConsumeToken();
+									}
+								}
+
+								attributeAction?.Invoke(identifierBuilder.ToString());
+								identifierBuilder.Length = identifierBuilderLength;
+							}
+							else if (token.IsSymbol(','))
+							{
+							}
+							else if (token.IsSymbol(']'))
+							{
+								// This isn't technically correct since it could be "] ]"
+								tokenReader.Require(']');
+								break;
+							}
+							else
+							{
+								throw new UhtTokenException(tokenReader, token, "]], ',', or identifier");
+							}
+						}
+					}
+					finally
+					{
+						StringBuilderCache.Small.Return(identifierBuilder);
+					}
+				}
+				else if (enableAlignAs && token.IsIdentifier("alignas"))
+				{
+					tokenReader.ConsumeToken();
+					tokenReader
+						.Require('(', "alignment specifier")
+						.RequireConstInt("alignment specifier")
+						.Require(')', "alignment specifier");
+				}
+				else
+				{
+					break;
+				}
+			}
+			return tokenReader;
+		}
+
 		/// <summary>
 		/// Require the given text
 		/// </summary>
diff --git a/Engine/Source/Programs/Shared/EpicGames.UHT/Tokenizer/UhtTokenReaderIdentifierExtensions.cs b/Engine/Source/Programs/Shared/EpicGames.UHT/Tokenizer/UhtTokenReaderIdentifierExtensions.cs
--- a/Engine/Source/Programs/Shared/EpicGames.UHT/Tokenizer/UhtTokenReaderIdentifierExtensions.cs
+++ b/Engine/Source/Programs/Shared/EpicGames.UHT/Tokenizer/UhtTokenReaderIdentifierExtensions.cs
@@ -195,127 +195,6 @@ namespace EpicGames.UHT.Tokenizer
 			throw new UhtTokenException(tokenReader, token, exceptionContext ?? "an identifier");
 		}
 
-		/// <summary>
-		/// Parse an optional cpp identifier
-		/// </summary>
-		/// <param name="tokenReader">Token reader</param>
-		/// <param name="identifier">Enumeration of the identifier tokens</param>
-		/// <param name="options">Identifier options</param>
-		/// <returns>True if an identifier is parsed</returns>
-		public static bool TryOptionalCppIdentifier(this IUhtTokenReader tokenReader, [NotNullWhen(true)] out IEnumerable<UhtToken>? identifier, UhtCppIdentifierOptions options = UhtCppIdentifierOptions.None)
-		{
-			identifier = null;
-			List<UhtToken> localIdentifier = new();
-
-			using UhtTokenSaveState savedState = new(tokenReader);
-
-			if (tokenReader.TryOptionalIdentifier(out UhtToken token))
-			{
-				localIdentifier.Add(token);
-			}
-			else
-			{
-				return false;
-			}
-
-			if (options.HasAnyFlags(UhtCppIdentifierOptions.AllowTemplates))
-			{
-				while (true)
-				{
-					if (tokenReader.TryPeekOptional('<'))
-					{
-						localIdentifier.Add(tokenReader.GetToken());
-						int nestedScopes = 1;
-						while (nestedScopes > 0)
-						{
-							token = tokenReader.GetToken();
-							if (token.TokenType.IsEndType())
-							{
-								return false;
-							}
-							localIdentifier.Add(token);
-							if (token.IsSymbol('<'))
-							{
-								++nestedScopes;
-							}
-							else if (token.IsSymbol('>'))
-							{
-								--nestedScopes;
-							}
-						}
-					}
-
-					if (!tokenReader.TryPeekOptional("::"))
-					{
-						break;
-					}
-					localIdentifier.Add(tokenReader.GetToken());
-					if (!tokenReader.TryOptionalIdentifier(out token))
-					{
-						localIdentifier.Add(token);
-					}
-					else
-					{
-						return false;
-					}
-				}
-			}
-			else
-			{
-				while (tokenReader.TryOptional("::"))
-				{
-					if (tokenReader.TryOptionalIdentifier(out token))
-					{
-						localIdentifier.Add(token);
-					}
-					else
-					{
-						return false;
-					}
-				}
-			}
-			identifier = localIdentifier;
-			savedState.AbandonState();
-			return true;
-		}
-
-		/// <summary>
-		/// Parse a required cpp identifier
-		/// </summary>
-		/// <param name="tokenReader">Token reader</param>
-		/// <param name="options">Parsing options</param>
-		/// <returns>Token reader</returns>
-		public static IUhtTokenReader RequireCppIdentifier(this IUhtTokenReader tokenReader, UhtCppIdentifierOptions options = UhtCppIdentifierOptions.None)
-		{
-			tokenReader.GetCppIdentifier(options);
-			return tokenReader;
-		}
-
-		/// <summary>
-		/// Parse a required cpp identifier
-		/// </summary>
-		/// <param name="tokenReader">Token reader</param>
-		/// <param name="initialIdentifier">Initial token of the identifier</param>
-		/// <param name="options">Parsing options</param>
-		/// <returns>Token reader</returns>
-		public static IUhtTokenReader RequireCppIdentifier(this IUhtTokenReader tokenReader, ref UhtToken initialIdentifier, UhtCppIdentifierOptions options = UhtCppIdentifierOptions.None)
-		{
-			tokenReader.GetCppIdentifier(ref initialIdentifier, options);
-			return tokenReader;
-		}
-
-		/// <summary>
-		/// Parse a required cpp identifier
-		/// </summary>
-		/// <param name="tokenReader">Token reader</param>
-		/// <param name="tokenListDelegate">Invoked when identifier is parsed</param>
-		/// <returns>Token reader</returns>
-		public static IUhtTokenReader RequireCppIdentifier(this IUhtTokenReader tokenReader, UhtTokenListDelegate tokenListDelegate)
-		{
-			tokenReader.RequireCppIdentifier(UhtCppIdentifierOptions.None, tokenListDelegate);
-			return tokenReader;
-		}
-
 		/// <summary>
 		/// Parse a required cpp identifier
 		/// </summary>
@@ -331,19 +210,6 @@ namespace EpicGames.UHT.Tokenizer
 			return tokenReader;
 		}
 
-		/// <summary>
-		/// Parse a required cpp identifier
-		/// </summary>
-		/// <param name="tokenReader">Token reader</param>
-		/// <param name="initialIdentifier">Initial token of the identifier</param>
-		/// <param name="tokenListDelegate">Invoked when identifier is parsed</param>
-		/// <returns>Token reader</returns>
-		public static IUhtTokenReader RequireCppIdentifier(this IUhtTokenReader tokenReader, ref UhtToken initialIdentifier, UhtTokenListDelegate tokenListDelegate)
-		{
-			tokenReader.RequireCppIdentifier(ref initialIdentifier, UhtCppIdentifierOptions.None, tokenListDelegate);
-			return tokenReader;
-		}
-
 		/// <summary>
 		/// Parse a required cpp identifier
 		/// </summary>
diff --git a/Engine/Source/Programs/Shared/EpicGames.UHT/Tokenizer/UhtTokenReaderSkipExtensions.cs b/Engine/Source/Programs/Shared/EpicGames.UHT/Tokenizer/UhtTokenReaderSkipExtensions.cs
--- a/Engine/Source/Programs/Shared/EpicGames.UHT/Tokenizer/UhtTokenReaderSkipExtensions.cs
+++ b/Engine/Source/Programs/Shared/EpicGames.UHT/Tokenizer/UhtTokenReaderSkipExtensions.cs
@@ -22,58 +22,6 @@ namespace EpicGames.UHT.Tokenizer
 			return tokenReader;
 		}
 
-		/// <summary>
-		/// Skip an alignas expression
-		/// </summary>
-		/// <param name="tokenReader">Token reader</param>
-		/// <returns>Token reader</returns>
-		public static IUhtTokenReader SkipAlignasIfNecessary(this IUhtTokenReader tokenReader)
-		{
-			const string Identifier = "alignas";
-			if (tokenReader.TryOptional(Identifier))
-			{
-				tokenReader
-					.Require('(', Identifier)
-					.RequireConstInt(Identifier)
-					.Require(')', Identifier);
-			}
-			return tokenReader;
-		}
-
-		/// <summary>
-		/// Skip deprecation macro
-		/// </summary>
-		/// <param name="tokenReader">Token reader</param>
-		/// <returns>Token reader</returns>
-		public static IUhtTokenReader SkipDeprecatedMacroIfNecessary(this IUhtTokenReader tokenReader)
-		{
-			if (tokenReader.TryOptional(new string[] { "DEPRECATED", "UE_DEPRECATED" }) >= 0)
-			{
-				tokenReader
-					.Require('(', "deprecation macro")
-					.RequireConstFloat("version in deprecation macro")
-					.Require(',', "deprecation macro")
-					.RequireConstString("message in deprecation macro")
-					.Require(')', "deprecation macro");
-			}
-			return tokenReader;
-		}
-
-		/// <summary>
-		/// Skip alignas and/or deprecation macros
-		/// </summary>
-		/// <param name="tokenReader">Token reader</param>
-		/// <returns>Token reader</returns>
-		public static IUhtTokenReader SkipAlignasAndDeprecatedMacroIfNecessary(this IUhtTokenReader tokenReader)
-		{
-			// alignas() can come before or after the deprecation macro.
-			// We can't have both, but the compiler will catch that anyway.
-			return tokenReader
-				.SkipAlignasIfNecessary()
-				.SkipDeprecatedMacroIfNecessary()
-				.SkipAlignasIfNecessary();
-		}
-
 		/// <summary>
 		/// Skip any block of tokens wrapped by the given token symbols
 		/// </summary>
diff --git a/Engine/Source/Programs/UnrealHeaderTool/Private/BaseParser.cpp b/Engine/Source/Programs/UnrealHeaderTool/Private/BaseParser.cpp
--- a/Engine/Source/Programs/UnrealHeaderTool/Private/BaseParser.cpp
+++ b/Engine/Source/Programs/UnrealHeaderTool/Private/BaseParser.cpp
@@ -601,6 +601,7 @@ bool FBaseParser::GetToken(FToken& Token, bool bNoConsts/*=false*/, ESymbolParse
 		||	PAIR('*','*')
 		||	PAIR('~','=')
 		||	PAIR(':',':')
+		||	PAIR('[','[')
 		)
 		{
 			if (c == TEXT('>') && d == TEXT('>') &&  *Pos == TEXT('>'))
diff --git a/Engine/Source/Programs/UnrealHeaderTool/Private/HeaderParser.cpp b/Engine/Source/Programs/UnrealHeaderTool/Private/HeaderParser.cpp
--- a/Engine/Source/Programs/UnrealHeaderTool/Private/HeaderParser.cpp
+++ b/Engine/Source/Programs/UnrealHeaderTool/Private/HeaderParser.cpp
@@ -1,6 +1,5 @@
 // Copyright Epic Games, Inc. All Rights Reserved.
 
-
 #include "HeaderParser.h"
 #include "UnrealHeaderTool.h"
 #include "PropertyTypes.h"
@@ -656,57 +655,150 @@ namespace
 		return nullptr;
 	}
 
-	void SkipAlignasIfNecessary(FBaseParser& Parser)
+	void ParseCppIdentifier(FBaseParser& Parser, FUHTStringBuilder& Out, FToken& InitialIdentifier)
 	{
-		if (Parser.MatchIdentifier(TEXT("alignas"), ESearchCase::CaseSensitive))
+		Out.Append(InitialIdentifier.Value);
+		for (;;)
 		{
-			Parser.RequireSymbol(TEXT('('), TEXT("'alignas'"));
-			Parser.RequireAnyConstInt(TEXT("'alignas'"));
-			Parser.RequireSymbol(TEXT(')'), TEXT("'alignas'"));
+			FToken Token;
+			if (!Parser.GetToken(Token))
+			{
+				break;
+			}
+			if (!Token.IsSymbol(TEXT("::"), ESearchCase::CaseSensitive))
+			{
+				Parser.UngetToken(Token);
+				break;
+			}
+
+			if (!Parser.GetToken(Token) || !Token.IsIdentifier())
+			{
+				Parser.Throwf(TEXT("Expected identifier"));
+			}
+			Out.Append(TEXT("::"));
+			Out.Append(Token.Value);
 		}
 	}
 
-	void SkipDeprecatedMacroIfNecessary(FBaseParser& Parser)
+	void ParseAttributes(FBaseParser& Parser, bool bAllowAlignAs)
 	{
-		FToken MacroToken;
-		if (!Parser.GetToken(MacroToken))
+		for (;;)
 		{
-			return;
-		}
+			FToken Token;
+			if (!Parser.GetToken(Token))
+			{
+				return;
+			}
 
-		if (!MacroToken.IsIdentifier() || 
-			(!MacroToken.IsValue(TEXT("DEPRECATED"), ESearchCase::CaseSensitive) && !MacroToken.IsValue(TEXT("UE_DEPRECATED"), ESearchCase::CaseSensitive)))
-		{
-			Parser.UngetToken(MacroToken);
-			return;
-		}
+			if (Token.IsIdentifier(TEXT("DEPRECATED"), ESearchCase::CaseSensitive) ||
+				Token.IsIdentifier(TEXT("UE_DEPRECATED"), ESearchCase::CaseSensitive))
+			{
+				auto ErrorMessageGetter = [&Token]() { return FString::Printf(TEXT("%s macro"), *Token.GetTokenValue()); };
 
-		auto ErrorMessageGetter = [&MacroToken]() { return FString::Printf(TEXT("%s macro"), *MacroToken.GetTokenValue()); };
+				Parser.RequireSymbol(TEXT('('), ErrorMessageGetter);
 
-		Parser.RequireSymbol(TEXT('('), ErrorMessageGetter);
+				if (Parser.GetToken(Token) && !Token.IsConstFloat())
+				{
+					Parser.Throwf(TEXT("Expected engine version in %s macro"), *Token.GetTokenValue());
+				}
 
-		FToken Token;
-		if (Parser.GetToken(Token) && !Token.IsConstFloat())
-		{
-			Parser.Throwf(TEXT("Expected engine version in %s macro"), *MacroToken.GetTokenValue());
-		}
+				Parser.RequireSymbol(TEXT(','), ErrorMessageGetter);
+				if (Parser.GetToken(Token) && !Token.IsConstString())
+				{
+					Parser.Throwf(TEXT("Expected deprecation message in %s macro"), *Token.GetTokenValue());
+				}
 
-		Parser.RequireSymbol(TEXT(','), ErrorMessageGetter);
-		if (Parser.GetToken(Token) && !Token.IsConstString())
-		{
-			Parser.Throwf(TEXT("Expected deprecation message in %s macro"), *MacroToken.GetTokenValue());
-		}
+				Parser.RequireSymbol(TEXT(')'), ErrorMessageGetter);
+			}
+			else if (Token.IsIdentifier(TEXT("UE_NODISCARD"), ESearchCase::CaseSensitive) ||
+				Token.IsIdentifier(TEXT("UE_NODISCARD_CTOR"), ESearchCase::CaseSensitive))
+			{
+			}
+			else if (Token.IsIdentifier(TEXT("UE_NORETURN"), ESearchCase::CaseSensitive))
+			{
+			}
+			else if (Token.IsIdentifier(TEXT("UE_NO_UNIQUE_ADDRESS"), ESearchCase::CaseSensitive))
+			{
+			}
+			else if (Token.IsSymbol(TEXT("[["), ESearchCase::CaseSensitive))
+			{
+				FUHTStringBuilder IdentifierBuilder;
+				if (Parser.MatchIdentifier(TEXT("using"), ESearchCase::CaseSensitive))
+				{
+					if (!Parser.GetToken(Token) || !Token.IsIdentifier())
+					{
+						Parser.Throwf(TEXT("Expected identifier"));
+					}
+					ParseCppIdentifier(Parser, IdentifierBuilder, Token);
+					Parser.RequireSymbol(':', TEXT("end of using statement in attribute"), ESymbolParseOption::Normal);
+				}
+				if (IdentifierBuilder.Len() > 0)
+				{
+					IdentifierBuilder.Append(TEXT("::"));
+				}
+				int identifierBuilderLength = IdentifierBuilder.Len();
 
-		Parser.RequireSymbol(TEXT(')'), ErrorMessageGetter);
-	}
+				// Attributes do allow patterns such as [[,]].
+				for (;;)
+				{
+					if (!Parser.GetToken(Token))
+					{
+						Parser.Throwf(TEXT("Expected attribute"));
+					}
+					if (Token.IsIdentifier())
+					{
+						ParseCppIdentifier(Parser, IdentifierBuilder, Token);
 
-	void SkipAlignasAndDeprecatedMacroIfNecessary(FBaseParser& Parser)
-	{
-		// alignas() can come before or after the deprecation macro.
-		// We can't have both, but the compiler will catch that anyway.
-		SkipAlignasIfNecessary(Parser);
-		SkipDeprecatedMacroIfNecessary(Parser);
-		SkipAlignasIfNecessary(Parser);
+						if (Parser.MatchSymbol('('))
+						{
+							int nestingCount = 1;
+							while (nestingCount > 0)
+							{
+								if (!Parser.GetToken(Token))
+								{
+									Parser.Throwf(TEXT("Expected end of attribute arguments"));
+								}
+								else if (Token.IsSymbol(')'))
+								{
+									nestingCount--;
+								}
+								else if (Token.IsSymbol('('))
+								{
+									nestingCount++;
+								}
+							}
+						}
+
+						IdentifierBuilder.LeftChopInline(identifierBuilderLength);
+					}
+					else if (Token.IsSymbol(','))
+					{
+					}
+					else if (Token.IsSymbol(']'))
+					{
+						// This isn't technically correct since it could be "] ]"
+						Parser.RequireSymbol(']', TEXT("end of attribute"), ESymbolParseOption::Normal);
+						break;
+					}
+					else
+					{
+						Parser.Throwf(TEXT("Expected end of attribute"));
+					}
+				}
+			}
+			else if (bAllowAlignAs && Token.IsIdentifier(TEXT("alignas"), ESearchCase::CaseSensitive))
+			{
+				Parser.RequireSymbol(TEXT('('), TEXT("'alignas'"));
+				Parser.RequireAnyConstInt(TEXT("'alignas'"));
+				Parser.RequireSymbol(TEXT(')'), TEXT("'alignas'"));
+			}
+			else
+			{
+				Parser.UngetToken(Token);
+				break;
+			}
+		}
+		return;
 	}
 
 	inline EPointerMemberBehavior ToPointerMemberBehavior(const FString& InString)
@@ -1089,7 +1181,7 @@ FUnrealEnumDefinitionInfo& FHeaderParser::CompileEnum()
 	{
 		CppForm = UEnum::ECppForm::Namespaced;
 
-		SkipDeprecatedMacroIfNecessary(*this);
+		ParseAttributes(*this, false);
 
 		bReadEnumName = GetIdentifier(EnumToken);
 	}
@@ -1112,7 +1204,7 @@ FUnrealEnumDefinitionInfo& FHeaderParser::CompileEnum()
 			CppForm = UEnum::ECppForm::Regular;
 		}
 
-		SkipAlignasAndDeprecatedMacroIfNecessary(*this);
+		ParseAttributes(*this, true);
 
 		bReadEnumName = GetIdentifier(EnumToken);
 	}
@@ -1203,7 +1295,7 @@ FUnrealEnumDefinitionInfo& FHeaderParser::CompileEnum()
 			// Now handle the inner true enum portion
 			RequireIdentifier(TEXT("enum"), ESearchCase::CaseSensitive, TEXT("'Enum'"));
 
-			SkipAlignasAndDeprecatedMacroIfNecessary(*this);
+			ParseAttributes(*this, true);
 
 			FToken InnerEnumToken;
 			if (!GetIdentifier(InnerEnumToken))
@@ -1247,7 +1339,7 @@ FUnrealEnumDefinitionInfo& FHeaderParser::CompileEnum()
 		FTokenValue TagIdentifier = TagToken.GetTokenValue();
 		AddFormattedPrevCommentAsTooltipMetaData(TagMetaData);
 
-		SkipDeprecatedMacroIfNecessary(*this);
+		ParseAttributes(*this, false);
 
 		// Try to read an optional explicit enum value specification
 		if (MatchSymbol(TEXT('=')))
@@ -1838,7 +1930,7 @@ FUnrealScriptStructDefinitionInfo& FHeaderParser::CompileStructDeclaration()
 	// The required API module for this struct, if any
 	FString RequiredAPIMacroIfPresent;
 
-	SkipAlignasAndDeprecatedMacroIfNecessary(*this);
+	ParseAttributes(*this, true);
 
 	// Read the struct name
 	ParseNameWithPotentialAPIMacroPrefix(/*out*/ StructNameInScript, /*out*/ RequiredAPIMacroIfPresent, TEXT("struct"));
@@ -6313,7 +6405,7 @@ FUnrealClassDefinitionInfo& FHeaderParser::CompileClassDeclaration()
 	// New style files have the class name / extends afterwards
 	RequireIdentifier(TEXT("class"), ESearchCase::CaseSensitive, TEXT("Class declaration"));
 
-	SkipAlignasAndDeprecatedMacroIfNecessary(*this);
+	ParseAttributes(*this, true);
 
 	FString DeclaredClassName;
 	FString RequiredAPIMacroIfPresent;
@@ -9121,7 +9213,7 @@ TSharedRef<FUnrealTypeDefinitionInfo> FHeaderPreParser::ParseClassDeclaration(co
 	// Require 'class'
 	RequireIdentifier(TEXT("class"), ESearchCase::CaseSensitive, ErrorMsg);
 
-	SkipAlignasAndDeprecatedMacroIfNecessary(*this);
+	ParseAttributes(*this, true);
 
 	// Read the class name
 	FString RequiredAPIMacroIfPresent;
@@ -9193,7 +9285,7 @@ TSharedRef<FUnrealTypeDefinitionInfo> FHeaderPreParser::ParseEnumDeclaration(con
 	{
 		CppForm = UEnum::ECppForm::Namespaced;
 
-		SkipDeprecatedMacroIfNecessary(*this);
+		ParseAttributes(*this, false);
 
 		bReadEnumName = GetIdentifier(EnumToken);
 	}
@@ -9216,7 +9308,7 @@ TSharedRef<FUnrealTypeDefinitionInfo> FHeaderPreParser::ParseEnumDeclaration(con
 			CppForm = UEnum::ECppForm::Regular;
 		}
 
-		SkipAlignasAndDeprecatedMacroIfNecessary(*this);
+		ParseAttributes(*this, true);
 
 		bReadEnumName = GetIdentifier(EnumToken);
 	}
@@ -9259,7 +9351,7 @@ TSharedRef<FUnrealTypeDefinitionInfo> FHeaderPreParser::ParseStructDeclaration(c
 	// Consume the struct keyword
 	RequireIdentifier(TEXT("struct"), ESearchCase::CaseSensitive, TEXT("Struct declaration specifier"));
 
-	SkipAlignasAndDeprecatedMacroIfNecessary(*this);
+	ParseAttributes(*this, true);
 
 	// The struct name as parsed in script and stripped of it's prefix
 	FString StructNameInScript;
